Copying //UE4/Dev-Framework to //UE4/Dev-Main (Source: //UE4/Dev-Framework @ 3136612)

diff --git a/Engine/Source/Runtime/AIModule/Private/Navigation/PathFollowingComponent.cpp b/Engine/Source/Runtime/AIModule/Private/Navigation/PathFollowingComponent.cpp
--- a/Engine/Source/Runtime/AIModule/Private/Navigation/PathFollowingComponent.cpp
+++ b/Engine/Source/Runtime/AIModule/Private/Navigation/PathFollowingComponent.cpp
@@ -1293,9 +1293,15 @@ void UPathFollowingComponent::ForceBlockDetectionUpdate()
 
 void UPathFollowingComponent::UpdateDecelerationData()
 {
-	const float CurrentMaxSpeed = MovementComp ? MovementComp->GetMaxSpeed() : 0.0f;
+	// no point in updating if we don't have a MovementComp
+	if (MovementComp == nullptr)
+	{
+		return;
+	}
+
+	const float CurrentMaxSpeed = MovementComp->GetMaxSpeed();
 	bool bUpdatePathSegment = (DecelerationSegmentIndex == INDEX_NONE);
-	if (MovementComp && CurrentMaxSpeed != CachedBrakingMaxSpeed)
+	if (CurrentMaxSpeed != CachedBrakingMaxSpeed)
 	{
 		const float PrevBrakingDistance = CachedBrakingDistance;
 		CachedBrakingDistance = MovementComp->GetPathFollowingBrakingDistance(CurrentMaxSpeed);
diff --git a/Engine/Source/Runtime/Analytics/AnalyticsET/Private/HttpServiceTracker.cpp b/Engine/Source/Runtime/Analytics/AnalyticsET/Private/HttpServiceTracker.cpp
--- a/Engine/Source/Runtime/Analytics/AnalyticsET/Private/HttpServiceTracker.cpp
+++ b/Engine/Source/Runtime/Analytics/AnalyticsET/Private/HttpServiceTracker.cpp
@@ -65,8 +65,8 @@ FHttpServiceTracker::FHttpServiceTracker(const FHttpServiceTrackerConfig& Config
 	, NextFlushTime(FPlatformTime::Seconds() + FlushIntervalSec)
 {
 	AnalyticsProvider = FAnalyticsET::Get().CreateAnalyticsProvider(FAnalyticsET::Config(Config.APIKey, Config.APIServer, Config.ApiVersion, false, TEXT("unknown"), TEXT("qosmetrics")));
-	// we'll just use the MachineID for the User. It actually won't matter much for this service.
-	AnalyticsProvider->SetUserID(FPlatformMisc::GetMachineId().ToString(EGuidFormats::Digits).ToLower());
+	// Use the standard UserID
+	AnalyticsProvider->SetUserID(FString::Printf(TEXT("%s|%s|%s"), *FPlatformMisc::GetMachineId().ToString(EGuidFormats::Digits).ToLower(), *FPlatformMisc::GetEpicAccountId(), *FPlatformMisc::GetOperatingSystemId()));
 	// Note we also don't start/stop the session. The AnalyticsET provider allows this, and this enables our collector
 	// to receive ONLY monitoring events.
 }
diff --git a/Engine/Source/Runtime/AnimGraphRuntime/Private/AnimNodes/AnimNode_CopyPoseFromMesh.cpp b/Engine/Source/Runtime/AnimGraphRuntime/Private/AnimNodes/AnimNode_CopyPoseFromMesh.cpp
--- a/Engine/Source/Runtime/AnimGraphRuntime/Private/AnimNodes/AnimNode_CopyPoseFromMesh.cpp
+++ b/Engine/Source/Runtime/AnimGraphRuntime/Private/AnimNodes/AnimNode_CopyPoseFromMesh.cpp
@@ -40,7 +40,9 @@ void FAnimNode_CopyPoseFromMesh::Evaluate(FPoseContext& Output)
 	FCompactPose& OutPose = Output.Pose;
 	OutPose.ResetToRefPose();
 
-	if (CurrentlyUsedSourceMeshComponent.IsValid() && CurrentlyUsedSourceMeshComponent.Get()->SkeletalMesh)
+	USkeletalMeshComponent* CurrentMeshComponent = CurrentlyUsedSourceMeshComponent.IsValid()? CurrentlyUsedSourceMeshComponent.Get() : nullptr;
+
+	if (CurrentMeshComponent && CurrentMeshComponent->SkeletalMesh && CurrentMeshComponent->IsRegistered())
 	{
 		const FBoneContainer& RequiredBones = OutPose.GetBoneContainer();
 		for(FCompactPoseBoneIndex PoseBoneIndex : OutPose.ForEachBoneIndex())
@@ -51,18 +53,18 @@ void FAnimNode_CopyPoseFromMesh::Evaluate(FPoseContext& Output)
 			if(Value && *Value!=INDEX_NONE)
 			{
 				const int32 SourceBoneIndex = *Value;
-				const int32 ParentIndex = CurrentlyUsedSourceMeshComponent.Get()->SkeletalMesh->RefSkeleton.GetParentIndex(SourceBoneIndex);
+				const int32 ParentIndex = CurrentMeshComponent->SkeletalMesh->RefSkeleton.GetParentIndex(SourceBoneIndex);
 				const FCompactPoseBoneIndex MyParentIndex = RequiredBones.GetParentBoneIndex(PoseBoneIndex);
 				// only apply if I also have parent, otherwise, it should apply the space bases
 				if (ParentIndex != INDEX_NONE && MyParentIndex != INDEX_NONE)
 				{
-					const FTransform& ParentTransform = CurrentlyUsedSourceMeshComponent.Get()->GetComponentSpaceTransforms()[ParentIndex];
-					const FTransform& ChildTransform = CurrentlyUsedSourceMeshComponent.Get()->GetComponentSpaceTransforms()[SourceBoneIndex];
+					const FTransform& ParentTransform = CurrentMeshComponent->GetComponentSpaceTransforms()[ParentIndex];
+					const FTransform& ChildTransform = CurrentMeshComponent->GetComponentSpaceTransforms()[SourceBoneIndex];
 					OutPose[PoseBoneIndex] = ChildTransform.GetRelativeTransform(ParentTransform);
 				}
 				else
 				{
-					OutPose[PoseBoneIndex] = CurrentlyUsedSourceMeshComponent.Get()->GetComponentSpaceTransforms()[SourceBoneIndex];
+					OutPose[PoseBoneIndex] = CurrentMeshComponent->GetComponentSpaceTransforms()[SourceBoneIndex];
 				}
 			}
 		}
diff --git a/Engine/Source/Runtime/AnimGraphRuntime/Private/AnimNodes/AnimNode_LayeredBoneBlend.cpp b/Engine/Source/Runtime/AnimGraphRuntime/Private/AnimNodes/AnimNode_LayeredBoneBlend.cpp
--- a/Engine/Source/Runtime/AnimGraphRuntime/Private/AnimNodes/AnimNode_LayeredBoneBlend.cpp
+++ b/Engine/Source/Runtime/AnimGraphRuntime/Private/AnimNodes/AnimNode_LayeredBoneBlend.cpp
@@ -38,6 +38,9 @@ void FAnimNode_LayeredBoneBlend::ReinitializeBoneBlendWeights(const FBoneContain
 
 	CurrentBoneBlendWeights.Empty(DesiredBoneBlendWeights.Num());
 	CurrentBoneBlendWeights.AddZeroed(DesiredBoneBlendWeights.Num());
+
+	//Reinitialize bone blend weights now that we have cleared them
+	FAnimationRuntime::UpdateDesiredBoneWeight(DesiredBoneBlendWeights, CurrentBoneBlendWeights, BlendWeights);
 }
 
 void FAnimNode_LayeredBoneBlend::CacheBones(const FAnimationCacheBonesContext& Context) 
diff --git a/Engine/Source/Runtime/AnimGraphRuntime/Private/AnimNodes/AnimNode_PoseByName.cpp b/Engine/Source/Runtime/AnimGraphRuntime/Private/AnimNodes/AnimNode_PoseByName.cpp
--- a/Engine/Source/Runtime/AnimGraphRuntime/Private/AnimNodes/AnimNode_PoseByName.cpp
+++ b/Engine/Source/Runtime/AnimGraphRuntime/Private/AnimNodes/AnimNode_PoseByName.cpp
@@ -30,7 +30,7 @@ void FAnimNode_PoseByName::UpdateAssetPlayer(const FAnimationUpdateContext& Cont
 	}
 	else
 	{
-		PoseUID = FSmartNameMapping::MaxUID;
+		PoseUID = SmartName::MaxUID;
 	}
 }
 
@@ -41,7 +41,7 @@ void FAnimNode_PoseByName::Evaluate(FPoseContext& Output)
 		// Zero out all curve value
 		FMemory::Memzero(PoseExtractContext.PoseCurves.GetData(), PoseExtractContext.PoseCurves.GetAllocatedSize());
 
-		if (PoseUID != FSmartNameMapping::MaxUID)
+		if (PoseUID != SmartName::MaxUID)
 		{
 			int32 Index = PoseUIDList.Find(PoseUID);
 			if (Index!=INDEX_NONE)
diff --git a/Engine/Source/Runtime/AnimGraphRuntime/Private/AnimNodes/AnimNode_PoseDriver.cpp b/Engine/Source/Runtime/AnimGraphRuntime/Private/AnimNodes/AnimNode_PoseDriver.cpp
new file mode 100644
--- /dev/null
+++ b/Engine/Source/Runtime/AnimGraphRuntime/Private/AnimNodes/AnimNode_PoseDriver.cpp
@@ -0,0 +1,262 @@
+// Copyright 1998-2016 Epic Games, Inc. All Rights Reserved.
+
+#include "AnimGraphRuntimePrivatePCH.h"
+#include "AnimNodes/AnimNode_PoseDriver.h"
+#include "Animation/PoseAsset.h"
+
+#include "AnimInstanceProxy.h"
+
+
+FAnimNode_PoseDriver::FAnimNode_PoseDriver()
+{
+	RadialScaling = 0.25f;
+	bIncludeRefPoseAsNeutralPose = true;
+	Type = EPoseDriverType::SwingOnly;
+}
+
+void FAnimNode_PoseDriver::Initialize(const FAnimationInitializeContext& Context)
+{
+	FAnimNode_PoseHandler::Initialize(Context);
+
+	SourcePose.Initialize(Context);
+}
+
+void FAnimNode_PoseDriver::CacheBones(const FAnimationCacheBonesContext& Context)
+{
+	FAnimNode_PoseHandler::CacheBones(Context);
+	// Init pose input
+	SourcePose.CacheBones(Context);
+	// Init bone ref
+	SourceBone.Initialize(Context.AnimInstanceProxy->GetRequiredBones());
+}
+
+void FAnimNode_PoseDriver::UpdateAssetPlayer(const FAnimationUpdateContext& Context)
+{
+	FAnimNode_PoseHandler::UpdateAssetPlayer(Context);
+	SourcePose.Update(Context);
+}
+
+void FAnimNode_PoseDriver::GatherDebugData(FNodeDebugData& DebugData)
+{
+	FAnimNode_PoseHandler::GatherDebugData(DebugData);
+	SourcePose.GatherDebugData(DebugData.BranchFlow(1.f));
+}
+
+void FAnimNode_PoseDriver::OnPoseAssetChange()
+{
+	FAnimNode_PoseHandler::OnPoseAssetChange();
+	bCachedPoseInfoUpToDate = false;
+}
+
+FVector FAnimNode_PoseDriver::GetTwistAxisVector()
+{
+	switch (TwistAxis)
+	{
+	case BA_X:
+	default:
+		return FVector(1.f, 0.f, 0.f);
+	case BA_Y:
+		return FVector(0.f, 1.f, 0.f);
+	case BA_Z:
+		return FVector(0.f, 0.f, 1.f);
+	}
+}
+
+float FAnimNode_PoseDriver::FindDistanceBetweenPoses(const FTransform& A, const FTransform& B)
+{
+	if (Type == EPoseDriverType::SwingAndTwist)
+	{
+		const FQuat AQuat = A.GetRotation();
+		const FQuat BQuat = B.GetRotation();
+		return AQuat.AngularDistance(BQuat);
+	}
+	else if(Type == EPoseDriverType::SwingOnly)
+	{
+		FVector TwistVector = GetTwistAxisVector();
+		FVector VecA = A.TransformVectorNoScale(TwistVector);
+		FVector VecB = B.TransformVectorNoScale(TwistVector);
+
+		const float Dot = FVector::DotProduct(VecA, VecB);
+		return FMath::Acos(Dot);
+	}
+	else if (Type == EPoseDriverType::Translation)
+	{
+		float DistSquared = (A.GetTranslation() - B.GetTranslation()).SizeSquared();
+		if (DistSquared > SMALL_NUMBER)
+		{
+			return FMath::Sqrt(DistSquared);
+		}
+		else
+		{
+			return 0.f;
+		}
+	}
+	else
+	{
+		ensureMsgf(false, TEXT("Unknown EPoseDriverType"));
+		return 0.f;
+	}
+}
+
+void FAnimNode_PoseDriver::UpdateCachedPoseInfo(const FTransform& RefTM)
+{
+	if (CurrentPoseAsset.IsValid())
+	{
+		int32 NumPoses = CurrentPoseAsset.Get()->GetNumPoses();
+		int32 NumPoseInfos = NumPoses;
+		if (bIncludeRefPoseAsNeutralPose)
+		{
+			NumPoseInfos++;
+		}
+
+		PoseInfos.Empty();
+		PoseInfos.AddZeroed(NumPoseInfos);
+
+		// Get track index in PoseAsset for bone of interest
+		const int32 TrackIndex = CurrentPoseAsset.Get()->GetTrackIndexByName(SourceBone.BoneName);
+		if (TrackIndex != INDEX_NONE)
+		{
+			TArray<FSmartName> PoseNames = CurrentPoseAsset.Get()->GetPoseNames();
+
+			// Cache quat for source bone for each pose
+			for (int32 PoseIdx = 0; PoseIdx < NumPoses; PoseIdx++)
+			{
+				FPoseDriverPoseInfo& PoseInfo = PoseInfos[PoseIdx];
+
+				FTransform BoneTransform;
+				bool bFound = CurrentPoseAsset.Get()->GetLocalPoseForTrack(PoseIdx, TrackIndex, BoneTransform);
+				check(bFound);
+
+				PoseInfo.PoseTM = BoneTransform;
+				PoseInfo.PoseName = PoseNames[PoseIdx].DisplayName;
+
+			}
+
+			// If we want to include ref pose, add that to end
+			if (bIncludeRefPoseAsNeutralPose)
+			{
+				FPoseDriverPoseInfo& PoseInfo = PoseInfos.Last();
+
+				static FName RefPoseName = FName(TEXT("RefPose"));
+				PoseInfo.PoseTM = RefTM;
+				PoseInfo.PoseName = RefPoseName;
+			}
+
+			// Calculate largest distance between poses
+			if (NumPoseInfos == 1)
+			{
+				// If only one pose, set to 180 degrees
+				PoseInfos[0].NearestPoseDist = PI;
+			}
+			else
+			{
+				// Iterate over poses
+				for (int32 PoseIdx = 0; PoseIdx < NumPoseInfos; PoseIdx++)
+				{
+					FPoseDriverPoseInfo& PoseInfo = PoseInfos[PoseIdx];
+					PoseInfo.NearestPoseDist = BIG_NUMBER; // init to large value
+
+					for (int32 OtherPoseIdx = 0; OtherPoseIdx < NumPoseInfos; OtherPoseIdx++)
+					{
+						if (OtherPoseIdx != PoseIdx) // If not ourself..
+						{
+							// Get distance between poses
+							FPoseDriverPoseInfo& OtherPoseInfo = PoseInfos[OtherPoseIdx];
+							float Dist = FindDistanceBetweenPoses(PoseInfo.PoseTM, OtherPoseInfo.PoseTM);
+							PoseInfo.NearestPoseDist = FMath::Min(Dist, PoseInfo.NearestPoseDist);
+						}
+					}
+
+					// Avoid zero dist if poses are all on top of each other
+					PoseInfo.NearestPoseDist = FMath::Max(PoseInfo.NearestPoseDist, KINDA_SMALL_NUMBER);
+				}
+			}
+		}
+	}
+
+	bCachedPoseInfoUpToDate = true;
+}
+
+void FAnimNode_PoseDriver::Evaluate(FPoseContext& Output)
+{
+	FPoseContext SourceData(Output);
+	SourcePose.Evaluate(SourceData);
+
+	// Get the index of the source bone
+	const FBoneContainer& BoneContainer = SourceData.Pose.GetBoneContainer();
+	const FCompactPoseBoneIndex SourceCompactIndex = SourceBone.GetCompactPoseIndex(BoneContainer);
+
+	// Do nothing if no PoseAsset, or bone is not found/LOD-ed out
+	if (!CurrentPoseAsset.IsValid() || SourceCompactIndex == INDEX_NONE || !Output.AnimInstanceProxy->IsSkeletonCompatible(CurrentPoseAsset->GetSkeleton()))
+	{
+		Output = SourceData;
+		return;
+	}
+
+	// Get transform of source bone
+	SourceBoneTM = SourceData.Pose[SourceCompactIndex];
+
+	// Ensure cached info is up to date
+	if (!bCachedPoseInfoUpToDate)
+	{
+		const FTransform& SourceRefPoseBoneTM = BoneContainer.GetRefPoseTransform(SourceCompactIndex);
+
+		UpdateCachedPoseInfo(SourceRefPoseBoneTM);
+	}
+
+	float TotalWeight = 0.f; // Keep track of total weight generated, to renormalize at the end
+
+	// Iterate over each pose, adding its contribution
+	for (int32 PoseInfoIdx = 0; PoseInfoIdx < PoseInfos.Num(); PoseInfoIdx++)
+	{
+		FPoseDriverPoseInfo& PoseInfo = PoseInfos[PoseInfoIdx];
+
+		// Find distance
+		PoseInfo.PoseDistance = FindDistanceBetweenPoses(SourceBoneTM, PoseInfo.PoseTM);
+
+		// Evaluate radial basis function to find weight
+		float GaussVarianceSqr = FMath::Square(RadialScaling * PoseInfo.NearestPoseDist);
+		PoseInfo.PoseWeight = FMath::Exp(-(PoseInfo.PoseDistance * PoseInfo.PoseDistance) / GaussVarianceSqr);
+
+		// Add weight to total
+		TotalWeight += PoseInfo.PoseWeight;
+	}
+
+	// Only normalize and apply if we got some kind of weight
+	if (TotalWeight > KINDA_SMALL_NUMBER)
+	{
+		// Normalize pose weights by rescaling by sum of weights
+		const float WeightScale = 1.f / TotalWeight;
+		for (FPoseDriverPoseInfo& PoseInfo : PoseInfos)
+		{
+			PoseInfo.PoseWeight *= WeightScale;
+		}
+
+		FPoseContext CurrentPose(Output);
+		check(PoseExtractContext.PoseCurves.Num() == PoseUIDList.Num());
+		for (int32 PoseIdx = 0; PoseIdx < PoseUIDList.Num(); ++PoseIdx)
+		{
+			PoseExtractContext.PoseCurves[PoseIdx] = PoseInfos[PoseIdx].PoseWeight;
+		}
+
+		if (CurrentPoseAsset.Get()->GetAnimationPose(CurrentPose.Pose, CurrentPose.Curve, PoseExtractContext))
+		{
+			// blend by weight
+			if (CurrentPoseAsset->IsValidAdditive())
+			{
+				// Don't want to modify SourceBone, set additive offset to identity
+				CurrentPose.Pose[SourceCompactIndex] = FTransform::Identity;
+
+				Output = SourceData;
+				FAnimationRuntime::AccumulateAdditivePose(Output.Pose, CurrentPose.Pose, Output.Curve, CurrentPose.Curve, 1.f, EAdditiveAnimationType::AAT_LocalSpaceBase);
+			}
+			else
+			{
+				// Don't want to modify SourceBone, set weight to zero
+				BoneBlendWeights[SourceCompactIndex.GetInt()] = 0.f;
+
+				FAnimationRuntime::BlendTwoPosesTogetherPerBone(SourceData.Pose, CurrentPose.Pose, SourceData.Curve, CurrentPose.Curve, BoneBlendWeights, Output.Pose, Output.Curve);
+			}
+		}
+	}
+}
diff --git a/Engine/Source/Runtime/AnimGraphRuntime/Private/AnimNodes/AnimNode_RandomPlayer.cpp b/Engine/Source/Runtime/AnimGraphRuntime/Private/AnimNodes/AnimNode_RandomPlayer.cpp
--- a/Engine/Source/Runtime/AnimGraphRuntime/Private/AnimNodes/AnimNode_RandomPlayer.cpp
+++ b/Engine/Source/Runtime/AnimGraphRuntime/Private/AnimNodes/AnimNode_RandomPlayer.cpp
@@ -184,11 +184,12 @@ void FAnimNode_RandomPlayer::Evaluate(FPoseContext& Output)
 					FBlendedCurve Curves[2];
 					float Weights[2];
 
-					Poses[0].SetBoneContainer(&AnimProxy->GetRequiredBones());
-					Poses[1].SetBoneContainer(&AnimProxy->GetRequiredBones());
+					const FBoneContainer& RequiredBone = AnimProxy->GetRequiredBones();
+					Poses[0].SetBoneContainer(&RequiredBone);
+					Poses[1].SetBoneContainer(&RequiredBone);
 
-					Curves[0].InitFrom(AnimProxy->GetSkelMeshComponent()->GetCachedAnimCurveMappingNameUids());
-					Curves[1].InitFrom(AnimProxy->GetSkelMeshComponent()->GetCachedAnimCurveMappingNameUids());
+					Curves[0].InitFrom(RequiredBone);
+					Curves[1].InitFrom(RequiredBone);
 
 					Weights[0] = CurrentData.BlendWeight;
 					Weights[1] = NextData.BlendWeight;
diff --git a/Engine/Source/Runtime/AnimGraphRuntime/Private/BoneControllers/AnimNode_BoneDrivenController.cpp b/Engine/Source/Runtime/AnimGraphRuntime/Private/BoneControllers/AnimNode_BoneDrivenController.cpp
--- a/Engine/Source/Runtime/AnimGraphRuntime/Private/BoneControllers/AnimNode_BoneDrivenController.cpp
+++ b/Engine/Source/Runtime/AnimGraphRuntime/Private/BoneControllers/AnimNode_BoneDrivenController.cpp
@@ -185,10 +185,10 @@ void FAnimNode_BoneDrivenController::EvaluateComponentSpaceInternal(FComponentSp
 	{
 		//	Morph target and Material parameter curves
 		USkeleton* Skeleton = Context.AnimInstanceProxy->GetSkeleton();
-		FSmartNameMapping::UID NameUID = Skeleton->GetUIDByName(USkeleton::AnimCurveMappingName, ParameterName);
-		if (NameUID != FSmartNameMapping::MaxUID)
+		SmartName::UID_Type NameUID = Skeleton->GetUIDByName(USkeleton::AnimCurveMappingName, ParameterName);
+		if (NameUID != SmartName::MaxUID)
 		{
-			Context.Curve.Set(NameUID, FinalDriverValue, (DestinationMode == EDrivenDestinationMode::MorphTarget) ? EAnimCurveFlags::ACF_DriveMorphTarget : EAnimCurveFlags::ACF_DriveMaterial);
+			Context.Curve.Set(NameUID, FinalDriverValue);
 		}
 	}
 }
diff --git a/Engine/Source/Runtime/AnimGraphRuntime/Private/BoneControllers/AnimNode_TwistCorrectiveNode.cpp b/Engine/Source/Runtime/AnimGraphRuntime/Private/BoneControllers/AnimNode_TwistCorrectiveNode.cpp
--- a/Engine/Source/Runtime/AnimGraphRuntime/Private/BoneControllers/AnimNode_TwistCorrectiveNode.cpp
+++ b/Engine/Source/Runtime/AnimGraphRuntime/Private/BoneControllers/AnimNode_TwistCorrectiveNode.cpp
@@ -13,7 +13,6 @@ FAnimNode_TwistCorrectiveNode::FAnimNode_TwistCorrectiveNode()
 	, RemappedMin(0.0f)
 	, RemappedMax(1.0f)
 {
-	Curve.Type = FAnimCurveType(true, true, false);
 }
 
 void FAnimNode_TwistCorrectiveNode::GatherDebugData(FNodeDebugData& DebugData)
@@ -46,7 +45,7 @@ void FAnimNode_TwistCorrectiveNode::EvaluateComponentSpaceInternal(FComponentSpa
 	const float FinalMappedValue = (FMath::Clamp(CurAngle, ReferenceAngle, RangeMaxInRadian) - ReferenceAngle) * (RemappedMax - RemappedMin) / (RangeMaxInRadian - ReferenceAngle);
 
 	//	set Curve Value
-	Context.Curve.Set(Curve.UID, FinalMappedValue*Alpha, Curve.Type.GetAnimCurveFlags());
+	Context.Curve.Set(Curve.UID, FinalMappedValue*Alpha);
 }
 
 bool FAnimNode_TwistCorrectiveNode::IsValidToEvaluate(const USkeleton* Skeleton, const FBoneContainer& RequiredBones)
diff --git a/Engine/Source/Runtime/AnimGraphRuntime/Public/AnimNodes/AnimNode_PoseByName.h b/Engine/Source/Runtime/AnimGraphRuntime/Public/AnimNodes/AnimNode_PoseByName.h
--- a/Engine/Source/Runtime/AnimGraphRuntime/Public/AnimNodes/AnimNode_PoseByName.h
+++ b/Engine/Source/Runtime/AnimGraphRuntime/Public/AnimNodes/AnimNode_PoseByName.h
@@ -23,7 +23,7 @@ public:
 public:	
 	FAnimNode_PoseByName()
 		: PoseWeight(1.f)
-		, PoseUID(FSmartNameMapping::MaxUID)
+		, PoseUID(SmartName::MaxUID)
 	{
 	}
 
@@ -36,6 +36,6 @@ public:
 	// End of FAnimNode_Base interface
 private:
 	// PoseUID
-	FSmartNameMapping::UID PoseUID;
+	SmartName::UID_Type PoseUID;
 };
 
diff --git a/Engine/Source/Runtime/AnimGraphRuntime/Public/AnimNodes/AnimNode_PoseDriver.h b/Engine/Source/Runtime/AnimGraphRuntime/Public/AnimNodes/AnimNode_PoseDriver.h
new file mode 100644
--- /dev/null
+++ b/Engine/Source/Runtime/AnimGraphRuntime/Public/AnimNodes/AnimNode_PoseDriver.h
@@ -0,0 +1,110 @@
+// Copyright 1998-2016 Epic Games, Inc. All Rights Reserved.
+
+#pragma once
+
+#include "AnimNode_PoseHandler.h"
+#include "AnimNode_PoseDriver.generated.h"
+
+class UPoseAsset;
+
+
+/** One target pose for the bone, with parameters to drive bone approaches this orientation. */
+struct FPoseDriverPoseInfo
+{
+	/** Name of pose, from PoseAsset */
+	FName PoseName;
+
+	/** Transform of source bone in this pose, from PoseAsset */
+	FTransform PoseTM;
+
+	/** Distance to nearest pose */
+	float NearestPoseDist;
+
+	/** Last weight calculated for this pose. */
+	float PoseWeight;
+
+	/** Last distance calculated from this pose */
+	float PoseDistance;
+
+	FPoseDriverPoseInfo()
+	: PoseTM(FTransform::Identity)
+	{}
+};
+
+/** Orientation aspect used to drive interpolation */
+UENUM()
+enum class EPoseDriverType : uint8
+{
+	/** Consider full rotation for interpolation */
+	SwingAndTwist,
+
+	/** Consider only swing for interpolation */
+	SwingOnly,
+
+	/** Consider translation relative to parent */
+	Translation
+};
+
+/** RBF based orientation driver */
+USTRUCT()
+struct ANIMGRAPHRUNTIME_API FAnimNode_PoseDriver : public FAnimNode_PoseHandler
+{
+	GENERATED_BODY()
+
+	UPROPERTY(EditAnywhere, EditFixedSize, BlueprintReadWrite, Category = Links)
+	FPoseLink SourcePose;
+
+	/** Bone to use for driving parameters based on its orientation */
+	UPROPERTY(EditAnywhere, Category = PoseDriver)
+	FBoneReference SourceBone;
+
+	/** Scaling of radial basis, applied to max distance between poses */
+	UPROPERTY(EditAnywhere, Category = PoseDriver, meta = (ClampMin = "0.01"))
+	float RadialScaling;
+
+	/** Should we consider the mesh ref pose of SourceBone as a 'neutral' pose (zero curves)   */
+	UPROPERTY(EditAnywhere, Category = PoseDriver)
+	bool bIncludeRefPoseAsNeutralPose;
+
+	/** Type of orientation for driving parameter  */
+	UPROPERTY(EditAnywhere, Category = PoseDriver)
+	EPoseDriverType Type;
+
+	/** Axis to use when Type is SwingOnly */
+	UPROPERTY(EditAnywhere, Category = PoseDriver)
+	TEnumAsByte<EBoneAxis> TwistAxis;
+
+	/** Input source bone TM, used for debug drawing */
+	FTransform SourceBoneTM;
+
+	/** Cached set of info for each pose */
+	TArray<FPoseDriverPoseInfo> PoseInfos;
+
+	// FAnimNode_Base interface
+	virtual void Initialize(const FAnimationInitializeContext& Context) override;
+	virtual void CacheBones(const FAnimationCacheBonesContext& Context) override;
+	virtual void UpdateAssetPlayer(const FAnimationUpdateContext& Context) override;
+	virtual void Evaluate(FPoseContext& Output) override;
+	virtual void GatherDebugData(FNodeDebugData& DebugData) override;
+	// End of FAnimNode_Base interface
+
+	// FAnimNode_PoseHandler interface
+	virtual void OnPoseAssetChange() override;
+	// End FAnimNode_PoseHandler interface
+
+
+	FAnimNode_PoseDriver();
+
+	/** Find max distance between any of the poses */
+	void UpdateCachedPoseInfo(const FTransform& RefTM);
+
+	/** Util to return unit vector for current twist axis */
+	FVector GetTwistAxisVector();
+	/** Util to find 'distance' between 2 transforms, using Type and TwistAxis settings */
+	float FindDistanceBetweenPoses(const FTransform& A, const FTransform& B);
+
+
+protected:
+	/** Is cached info up to date */
+	bool bCachedPoseInfoUpToDate;
+};
\ No newline at end of file
diff --git a/Engine/Source/Runtime/AnimGraphRuntime/Public/AnimNodes/AnimNode_PoseHandler.h b/Engine/Source/Runtime/AnimGraphRuntime/Public/AnimNodes/AnimNode_PoseHandler.h
--- a/Engine/Source/Runtime/AnimGraphRuntime/Public/AnimNodes/AnimNode_PoseHandler.h
+++ b/Engine/Source/Runtime/AnimGraphRuntime/Public/AnimNodes/AnimNode_PoseHandler.h
@@ -43,12 +43,15 @@ public:
 	// End of FAnimNode_AssetPlayerBase Interface
 
 protected:
+	/** Called after CurrentPoseAsset is changed.  */
+	virtual void OnPoseAssetChange() {}
+
 	TWeakObjectPtr<UPoseAsset> CurrentPoseAsset;
 	FAnimExtractContext PoseExtractContext;
 	// weight to blend pose per joint - has to be cached whenever cache bones for LOD
 	// note this is for mesh bone
 	TArray<float> BoneBlendWeights;
-	TArray<FSmartNameMapping::UID> PoseUIDList;
+	TArray<SmartName::UID_Type> PoseUIDList;
 
 private:
 	void UpdatePoseAssetProperty(struct FAnimInstanceProxy* InstanceProxy);
diff --git a/Engine/Source/Runtime/AnimGraphRuntime/Public/BoneControllers/AnimNode_LegIK.h b/Engine/Source/Runtime/AnimGraphRuntime/Public/BoneControllers/AnimNode_LegIK.h
new file mode 100644
--- /dev/null
+++ b/Engine/Source/Runtime/AnimGraphRuntime/Public/BoneControllers/AnimNode_LegIK.h
@@ -0,0 +1,178 @@
+// Copyright 1998-2016 Epic Games, Inc. All Rights Reserved.
+
+#pragma once
+
+#include "AnimNode_SkeletalControlBase.h"
+#include "AnimNode_LegIK.generated.h"
+
+/** Per foot definitions */
+USTRUCT()
+struct FAnimLegIKDefinition
+{
+	GENERATED_USTRUCT_BODY()
+
+	UPROPERTY(EditAnywhere, Category = "Settings")
+	FBoneReference IKFootBone;
+
+	UPROPERTY(EditAnywhere, Category = "Settings")
+	FBoneReference FKFootBone;
+
+	UPROPERTY(EditAnywhere, Category = "Settings")
+	int32 NumBonesInLimb;
+
+	/** Forward Axis for Foot bone. */
+	UPROPERTY(EditAnywhere, Category = "Settings")
+	TEnumAsByte<EAxis::Type> FootBoneForwardAxis;
+
+	/** If enabled, we prevent the leg from bending backwards and enforce a min compression angle */
+	UPROPERTY(EditAnywhere, Category = "Settings")
+	bool bEnableRotationLimit;
+
+	/** Only used if bEnableRotationLimit is enabled. Prevents the leg from folding onto itself,
+	* and forces at least this angle between Parent and Child bone. */
+	UPROPERTY(EditAnywhere, Category = "Settings")
+	float MinRotationAngle;
+
+	/** Enable Knee Twist correction, by comparing Foot FK with Foot IK orientation. */
+	UPROPERTY(EditAnywhere, Category = "Settings")
+	bool bEnableKneeTwistCorrection;
+
+	FAnimLegIKDefinition()
+		: NumBonesInLimb(2)
+		, FootBoneForwardAxis(EAxis::Y)
+		, bEnableRotationLimit(false)
+		, MinRotationAngle(15.f)
+		, bEnableKneeTwistCorrection(true)
+	{}
+};
+
+/** Runtime foot data after validation, we guarantee these bones to exist */
+USTRUCT()
+struct FAnimLegIKData
+{
+	GENERATED_USTRUCT_BODY()
+
+public:
+	FCompactPoseBoneIndex IKFootBoneIndex;
+	FTransform IKFootTransform;
+
+	FAnimLegIKDefinition* LegDefPtr;
+
+	int32 NumBones;
+	TArray<FCompactPoseBoneIndex> FKLegBoneIndices;
+	TArray<FTransform> FKLegBoneTransforms;
+
+public:
+	void InitializeTransforms(USkeletalMeshComponent* SkelComp, FCSPose<FCompactPose>& MeshBases);
+
+	FAnimLegIKData()
+		: IKFootBoneIndex(INDEX_NONE)
+		, LegDefPtr(nullptr)
+		, NumBones(INDEX_NONE)
+	{}
+};
+
+USTRUCT()
+struct FIKChainLink
+{
+	GENERATED_USTRUCT_BODY()
+
+		FVector Location;
+	float Length;
+	FVector LinkAxisZ;
+
+	FIKChainLink()
+		: Location(FVector::ZeroVector)
+		, Length(0.f)
+		, LinkAxisZ(FVector::ZeroVector)
+	{}
+
+	FIKChainLink(FVector InLocation, float InLength)
+		: Location(InLocation)
+		, Length(InLength)
+		, LinkAxisZ(FVector::ZeroVector)
+	{}
+};
+
+USTRUCT()
+struct FIKChain
+{
+	GENERATED_USTRUCT_BODY()
+
+public:
+	TArray<FIKChainLink> Links;
+	float MinRotationAngleRadians;
+
+private:
+	bool bInitialized;
+	float MaximumReach;
+	int32 NumLinks;
+	bool bEnableRotationLimit;
+	USkeletalMeshComponent* SkelMeshComp;
+
+public:
+	FIKChain()
+		: bInitialized(false)
+		, MaximumReach(0.f)
+	{}
+
+	void InitializeFromLegData(const FAnimLegIKData& InLegData, USkeletalMeshComponent* InSkelMeshComp);
+	void ReachTarget(const FVector& InTargetLocation, float InReachPrecision, int32 InMaxIterations);
+
+	float GetMaximumReach() const
+	{
+		return MaximumReach;
+	}
+
+private:
+	void OrientAllLinksToDirection(const FVector& InDirection);
+	void SolveFABRIK(const FVector& InTargetLocation, float InReachPrecision, int32 InMaxIterations);
+
+	static void FABRIK_ForwardReach(const FVector& InTargetLocation, FIKChain& IKChain);
+	static void FABRIK_BackwardReach(const FVector& InRootTargetLocation, FIKChain& IKChain);
+	static void FABRIK_ApplyLinkConstraints_Forward(FIKChain& IKChain, int32 LinkIndex);
+	static void FABRIK_ApplyLinkConstraints_Backward(FIKChain& IKChain, int32 LinkIndex);
+
+	static void DrawDebugIKChain(const FIKChain& IKChain, const FColor& InColor);
+};
+
+USTRUCT()
+struct ANIMGRAPHRUNTIME_API FAnimNode_LegIK : public FAnimNode_SkeletalControlBase
+{
+	GENERATED_USTRUCT_BODY()
+
+	FAnimNode_LegIK();
+
+	/** Tolerance for reaching IK Target, in unreal units. */
+	UPROPERTY(EditAnywhere, Category = "Settings")
+	float ReachPrecision;
+
+	/** Max Number of Iterations. */
+	UPROPERTY(EditAnywhere, Category = "Settings")
+	int32 MaxIterations;
+
+	UPROPERTY(EditAnywhere, Category = "Settings")
+	TArray<FAnimLegIKDefinition> LegsDefinition;
+
+	UPROPERTY(Transient)
+	TArray<FAnimLegIKData> LegsData;
+
+public:
+	// FAnimNode_Base interface
+	virtual void GatherDebugData(FNodeDebugData& DebugData) override;
+	// End of FAnimNode_Base interface
+
+	// FAnimNode_SkeletalControlBase interface
+	virtual void EvaluateBoneTransforms(USkeletalMeshComponent* SkelComp, FCSPose<FCompactPose>& MeshBases, TArray<FBoneTransform>& OutBoneTransforms) override;
+	virtual bool IsValidToEvaluate(const USkeleton* Skeleton, const FBoneContainer& RequiredBones) override;
+	// End of FAnimNode_SkeletalControlBase interface
+
+	void OrientLegTowardsIK(FAnimLegIKData& InLegData, USkeletalMeshComponent* SkelComp);
+	void DoLegReachIK(FAnimLegIKData& InLegData, USkeletalMeshComponent* SkelComp);
+	void AdjustKneeTwist(FAnimLegIKData& InLegData, USkeletalMeshComponent* SkelComp);
+
+private:
+	// FAnimNode_SkeletalControlBase interface
+	virtual void InitializeBoneReferences(const FBoneContainer& RequiredBones) override;
+	// End of FAnimNode_SkeletalControlBase interface
+};
diff --git a/Engine/Source/Runtime/AssetRegistry/Public/AssetData.h b/Engine/Source/Runtime/AssetRegistry/Public/AssetData.h
--- a/Engine/Source/Runtime/AssetRegistry/Public/AssetData.h
+++ b/Engine/Source/Runtime/AssetRegistry/Public/AssetData.h
@@ -105,6 +105,11 @@ public:
 		return ObjectPath == Other.ObjectPath;
 	}
 
+	bool operator!=(const FAssetData& Other) const
+	{
+		return ObjectPath != Other.ObjectPath;
+	}
+
 	bool operator>(const FAssetData& Other) const
 	{
 		return ObjectPath > Other.ObjectPath;
diff --git a/Engine/Source/Runtime/AudioMixer/AudioMixer.Build.cs b/Engine/Source/Runtime/AudioMixer/AudioMixer.Build.cs
new file mode 100644
--- /dev/null
+++ b/Engine/Source/Runtime/AudioMixer/AudioMixer.Build.cs
@@ -0,0 +1,87 @@
+// Copyright 1998-2016 Epic Games, Inc. All Rights Reserved.
+
+namespace UnrealBuildTool.Rules
+{
+	public class AudioMixer : ModuleRules
+	{
+		public AudioMixer(TargetInfo Target)
+		{
+			//OptimizeCode = CodeOptimization.Always;
+
+			PrivateIncludePathModuleNames.Add("TargetPlatform");
+
+			PrivateIncludePaths.AddRange(
+				new string[] 
+				{
+					"Runtime/AudioMixer/Private",
+				}
+			);
+
+			PublicDependencyModuleNames.AddRange(
+				new string[]
+				{
+					"Core",
+					"CoreUObject",
+					"Engine",
+				}
+			);
+
+			PrivateDependencyModuleNames.AddRange(
+				new string[] 
+				{
+					"CoreUObject", 
+				}
+			);
+
+			if ((Target.Platform == UnrealTargetPlatform.Win64) ||
+				(Target.Platform == UnrealTargetPlatform.Win32))
+			{
+				AddEngineThirdPartyPrivateStaticDependencies(Target,
+					"UEOgg",
+					"Vorbis",
+					"VorbisFile",
+					"libOpus"
+					);
+			}
+
+			if (Target.Platform == UnrealTargetPlatform.HTML5 && Target.Architecture == "-win32")
+			{
+				AddEngineThirdPartyPrivateStaticDependencies(Target,
+					"UEOgg",
+					"Vorbis",
+					"VorbisFile"
+					);
+			}
+
+			if (Target.Platform == UnrealTargetPlatform.Mac)
+			{
+				AddEngineThirdPartyPrivateStaticDependencies(Target,
+					"UEOgg",
+					"Vorbis",
+					"libOpus"
+					);
+				PublicFrameworks.AddRange(new string[] { "AVFoundation", "CoreVideo", "CoreMedia" });
+			}
+
+			if (Target.Platform == UnrealTargetPlatform.Android)
+			{
+				AddEngineThirdPartyPrivateStaticDependencies(Target,
+					"UEOgg",
+					"Vorbis",
+					"VorbisFile"
+					);
+			}
+
+			if (Target.Platform == UnrealTargetPlatform.Linux)
+			{
+				AddEngineThirdPartyPrivateStaticDependencies(Target,
+					"UEOgg",
+					"Vorbis",
+					"VorbisFile",
+					"libOpus"
+					);
+			}
+
+		}
+	}
+}
diff --git a/Engine/Source/Runtime/AudioMixer/Classes/Effects/SoundEffectLowPassFilter.h b/Engine/Source/Runtime/AudioMixer/Classes/Effects/SoundEffectLowPassFilter.h
new file mode 100644
--- /dev/null
+++ b/Engine/Source/Runtime/AudioMixer/Classes/Effects/SoundEffectLowPassFilter.h
@@ -0,0 +1,61 @@
+// Copyright 1998-2016 Epic Games, Inc. All Rights Reserved.
+
+#pragma once
+
+#include "Sound/SoundEffectSource.h"
+#include "SoundEffectLowPassFilter.generated.h"
+
+// ========================================================================
+// FSoundEffectLowPassFilterSettings
+// ========================================================================
+
+USTRUCT()
+struct AUDIOMIXER_API FSoundEffectLowPassFilterSettings
+{
+	GENERATED_USTRUCT_BODY()
+
+	UPROPERTY(EditAnywhere, Category = SoundEffect)
+	float CutoffFrequency;
+
+	UPROPERTY(EditAnywhere, Category = SoundEffect)
+	float Q;
+};
+
+// ========================================================================
+// USoundEffectLowPassFilter
+// ========================================================================
+
+UCLASS()
+class AUDIOMIXER_API USoundEffectLowPassFilter : public USoundEffectSource
+{
+	GENERATED_BODY()
+
+	// Called on an audio effect at initialization on main thread before audio processing begins.
+	void Init(const FSoundEffectSourceInitData& InSampleRate) override {}
+	
+	// Process the input block of audio. Called on audio thread.
+	void OnProcessAudio(const FSoundEffectSourceInputData& InData, FSoundEffectSourceOutputData& OutData) override {}
+
+private:
+	//Audio::FLowPass LowPassFilter;
+};
+
+// ========================================================================
+// USoundEffectLowPassFilterPreset
+// This code exposes your preset settings and effect class to the editor.
+// Do not modify this code!
+// ========================================================================
+
+UCLASS()
+class AUDIOMIXER_API USoundEffectLowPassFilterPreset : public USoundEffectSourcePreset
+{
+	GENERATED_BODY()
+
+	void* GetSettings() override { return (void*)&Settings; }
+	uint32 GetSettingsSize() const override { return sizeof(Settings); }
+	UScriptStruct* GetSettingsStruct() const override { return FSoundEffectLowPassFilterSettings::StaticStruct(); }
+	UClass* GetEffectClass() const override { return USoundEffectLowPassFilter::StaticClass(); }
+
+	UPROPERTY(EditAnywhere, Category = SoundEffect)
+	FSoundEffectLowPassFilterSettings Settings;
+};
diff --git a/Engine/Source/Runtime/AudioMixer/Classes/SourceEffects/AudioMixerSourceEffectFilter.h b/Engine/Source/Runtime/AudioMixer/Classes/SourceEffects/AudioMixerSourceEffectFilter.h
new file mode 100644
--- /dev/null
+++ b/Engine/Source/Runtime/AudioMixer/Classes/SourceEffects/AudioMixerSourceEffectFilter.h
@@ -0,0 +1,86 @@
+// Copyright 1998-2016 Epic Games, Inc. All Rights Reserved.
+
+#pragma once
+
+#include "Sound/SoundEffectSource.h"
+// Includes for audio dsp helper classes
+#include "AudioMixerSourceEffectFilter.generated.h"
+
+// ========================================================================
+// FSourceEffectFilterSettings
+// UStruct used to define user-exposed params for use with your effect.
+// ========================================================================
+
+UENUM()
+namespace ESourceEffectFilter
+{
+	enum Type
+	{
+		LowPass,
+		HighPass,
+		BandPass,
+		LowShelf,
+		HighShelf
+	};
+}
+
+USTRUCT()
+struct FSourceEffectFilterSettings
+{
+	GENERATED_USTRUCT_BODY()
+
+	// The type of filter to use for this source effect
+	UPROPERTY(EditAnywhere, Category = SoundEffect)
+	TEnumAsByte<enum ESourceEffectFilter::Type> FilterType;
+
+	// The gain of the filter (only used in LowShelf and HighShelf types)
+	UPROPERTY(EditAnywhere, Category = SoundEffect)
+	float GainDb;
+
+	// The corner or cutoff frequency of the filter in Hz
+	UPROPERTY(EditAnywhere, Category = SoundEffect)
+	float CutoffFrequency;
+
+	// The Q of the filter. 
+	UPROPERTY(EditAnywhere, Category = SoundEffect)
+	float Q;
+};
+
+// ========================================================================
+// USoundEffectSourceFilter
+// Class which processes audio streams and uses parameters defined in the preset class.
+// ========================================================================
+
+UCLASS()
+class AUDIOMIXER_API USourceEffectFilter : public USoundEffectSource
+{
+	GENERATED_BODY()
+	
+	// Called on an audio effect at initialization on main thread before audio processing begins.
+	void Init(const FSoundEffectSourceInitData& InSampleRate) override;
+
+	// Process the input block of audio. Called on audio thread.
+	void OnProcessAudio(const FSoundEffectSourceInputData& InData, FSoundEffectSourceOutputData& OutData) override;
+};
+
+// ========================================================================
+// USoundEffectSourceFilterPreset
+// This code exposes your preset settings and effect class to the editor.
+// Do not modify this code!
+// ========================================================================
+
+UCLASS()
+class AUDIOMIXER_API USoundEffectSourceFilterPreset : public USoundEffectSourcePreset
+{
+	GENERATED_BODY()
+	
+	void* GetSettings() override { return (void*)&Settings; }
+	uint32 GetSettingsSize() const override { return sizeof(Settings); }
+	UScriptStruct* GetSettingsStruct() const override { return FSourceEffectFilterSettings::StaticStruct(); }
+	UClass* GetEffectClass() const override { return USourceEffectFilter::StaticClass(); }
+	FText GetAssetActionName() const override { return NSLOCTEXT("AssetTypeActions", "AssetTypeActions_SoundEffectSourceFilterPreset", "SoundEffectSourceFilterPreset"); }
+
+	// Sound effect source filter settings.
+	UPROPERTY(EditAnywhere, Category = SoundEffect)
+	FSourceEffectFilterSettings Settings;
+};
diff --git a/Engine/Source/Runtime/AudioMixer/Classes/SubmixEffects/AudioMixerSubmixEffectEQ.h b/Engine/Source/Runtime/AudioMixer/Classes/SubmixEffects/AudioMixerSubmixEffectEQ.h
new file mode 100644
--- /dev/null
+++ b/Engine/Source/Runtime/AudioMixer/Classes/SubmixEffects/AudioMixerSubmixEffectEQ.h
@@ -0,0 +1,58 @@
+// Copyright 1998-2016 Epic Games, Inc. All Rights Reserved.
+
+#pragma once
+
+#include "Sound/SoundEffectSource.h"
+// Includes for audio dsp helper classes
+#include "AudioMixerSubmixEffectEQ.generated.h"
+
+// ========================================================================
+// FAudiomixerSourceEffectFilterSettings
+// ========================================================================
+
+USTRUCT()
+struct FSubmixEffectEQSettings
+{
+	GENERATED_USTRUCT_BODY()
+
+	// TODO
+};
+
+// ========================================================================
+// USoundEffectSubmixEQ
+// Class which processes audio streams and uses parameters defined in the preset class.
+// ========================================================================
+
+UCLASS()
+class AUDIOMIXER_API USoundEffectSubmixEQ : public USoundEffectSource
+{
+	GENERATED_BODY()
+	
+	// Called on an audio effect at initialization on main thread before audio processing begins.
+	void Init(const FSoundEffectSourceInitData& InSampleRate) override;
+
+	// Process the input block of audio. Called on audio thread.
+	void OnProcessAudio(const FSoundEffectSourceInputData& InData, FSoundEffectSourceOutputData& OutData) override;
+};
+
+// ========================================================================
+// UMyTestSourceEffectPreset
+// This code exposes your preset settings and effect class to the editor.
+// Do not modify this code!
+// ========================================================================
+
+UCLASS()
+class AUDIOMIXER_API USoundEffectSubmixEQPreset : public USoundEffectSourcePreset
+{
+	GENERATED_BODY()
+	
+	void* GetSettings() override { return (void*)&Settings; }
+	uint32 GetSettingsSize() const override { return sizeof(Settings); }
+	UScriptStruct* GetSettingsStruct() const override { return FSubmixEffectEQSettings::StaticStruct(); }
+	UClass* GetEffectClass() const override { return USoundEffectSubmixEQ::StaticClass(); }
+	FText GetAssetActionName() const override { return NSLOCTEXT("AssetTypeActions", "AssetTypeActions_SoundEffectSubmixEQPreset", "SoundEffectSubmixEQPreset"); }
+
+	// FSubmixEffectEQSettings Preset Settings
+	UPROPERTY(EditAnywhere, Category = SoundSourceEffect)
+	FSubmixEffectEQSettings Settings;
+};
diff --git a/Engine/Source/Runtime/AudioMixer/Private/AudioMixer.cpp b/Engine/Source/Runtime/AudioMixer/Private/AudioMixer.cpp
new file mode 100644
--- /dev/null
+++ b/Engine/Source/Runtime/AudioMixer/Private/AudioMixer.cpp
@@ -0,0 +1,111 @@
+// Copyright 1998-2016 Epic Games, Inc. All Rights Reserved.
+
+#include "AudioMixerPCH.h"
+#include "AudioMixerDevice.h"
+#include "AudioMixer.h"
+
+namespace Audio
+{
+	/**
+	 * FAudioMixerCallbackWorker
+	 */
+
+	FAudioMixerCallbackWorker::FAudioMixerCallbackWorker(IAudioMixerPlatformInterface* InMixerPlatform, TArray<float>* InOutBuffer)
+		: MixerPlatform(InMixerPlatform)
+		, OutBuffer(InOutBuffer)
+	{
+		check(MixerPlatform);
+		check(OutBuffer);
+	}
+
+	void FAudioMixerCallbackWorker::DoWork()
+	{
+		// Zero the buffer before processing a new one
+		check(OutBuffer);
+		FPlatformMemory::Memzero((void*)OutBuffer->GetData(), OutBuffer->Num() * sizeof(float));
+
+		// Call into multi-platform mixer for new buffer
+		MixerPlatform->GenerateBuffer(*OutBuffer);
+	}
+
+	/**
+	 * IAudioMixerPlatformInterface
+	 */
+
+	IAudioMixerPlatformInterface::IAudioMixerPlatformInterface()
+		: CallbackTask(nullptr)
+		, CurrentBufferIndex(0)
+		, LastError(TEXT("None"))
+	{
+	}
+
+	void IAudioMixerPlatformInterface::GenerateBuffer(TArray<float>& Buffer)
+	{
+		// Call into platform independent code to process next stream
+		AudioStreamInfo.AudioMixer->OnProcessAudioStream(Buffer);
+	}
+
+	void IAudioMixerPlatformInterface::ReadNextBuffer()
+	{
+		if (AudioStreamInfo.StreamState != EAudioOutputStreamState::Stopping)
+		{
+			// Make sure we have a pending task in the read index
+			check(CallbackTask);
+
+			// Make sure it's done
+			CallbackTask->EnsureCompletion();
+
+			// Get the processed buffer
+			const TArray<float>* Buffer = CallbackTask->GetTask().GetBuffer();
+
+			// Submit the buffer to the platform output device
+			check(Buffer);
+			SubmitBuffer(*Buffer);
+
+			// Clean up the task
+			delete CallbackTask;
+
+			// Create a new task and start it
+			CallbackTask = new FAsyncAudioMixerCallbackTask(this, &OutputBuffers[CurrentBufferIndex]);
+			CallbackTask->StartBackgroundTask();
+
+			// Increment the buffer index
+			CurrentBufferIndex = (CurrentBufferIndex + 1) % NumMixerBuffers;
+		}
+	}
+
+	void IAudioMixerPlatformInterface::BeginGeneratingAudio()
+	{
+		FAudioPlatformDeviceInfo & DeviceInfo = AudioStreamInfo.DeviceInfo;
+
+		// Setup the output buffers
+		for (int32 Index = 0; Index < NumMixerBuffers; ++Index)
+		{
+			OutputBuffers[Index].SetNumZeroed(DeviceInfo.NumSamples);
+		}
+
+		// Submit the first empty buffer. This will begin audio callback notifications on some platforms.
+		SubmitBuffer(OutputBuffers[CurrentBufferIndex++]);
+
+		// Launch the task to begin generating audio
+		CallbackTask = new FAsyncAudioMixerCallbackTask(this, &OutputBuffers[CurrentBufferIndex]);
+		CallbackTask->StartBackgroundTask();
+	}
+
+
+	uint32 IAudioMixerPlatformInterface::GetNumBytesForFormat(const EAudioMixerStreamDataFormat::Type DataFormat)
+	{
+		switch (DataFormat)
+		{
+			case EAudioMixerStreamDataFormat::Float:	return 4;
+			case EAudioMixerStreamDataFormat::Double:	return 8;
+			case EAudioMixerStreamDataFormat::Int16:	return 2;
+			case EAudioMixerStreamDataFormat::Int24:	return 3;
+			case EAudioMixerStreamDataFormat::Int32:	return 4;
+
+			default:
+			checkf(false, TEXT("Unknown or unsupported data format."))
+				return 0;
+		}
+	}
+}
\ No newline at end of file
diff --git a/Engine/Source/Runtime/AudioMixer/Private/AudioMixerBuffer.cpp b/Engine/Source/Runtime/AudioMixer/Private/AudioMixerBuffer.cpp
new file mode 100644
--- /dev/null
+++ b/Engine/Source/Runtime/AudioMixer/Private/AudioMixerBuffer.cpp
@@ -0,0 +1,416 @@
+// Copyright 1998-2016 Epic Games, Inc. All Rights Reserved.
+
+#include "AudioMixerPCH.h"
+#include "AudioMixerBuffer.h"
+#include "AudioMixerDevice.h"
+#include "Engine.h"
+#include "SoundDefinitions.h"
+#include "TargetPlatform.h"
+
+namespace Audio
+{
+
+	FMixerBuffer::FMixerBuffer(FAudioDevice* InAudioDevice, USoundWave* InWave, EBufferType::Type InBufferType)
+		: FSoundBuffer(InAudioDevice)
+		, RealtimeAsyncHeaderParseTask(nullptr)
+		, DecompressionState(nullptr)
+		, BufferType(InBufferType)
+		, SampleRate(InWave->SampleRate)
+		, BitsPerSample(16) // TODO: support more bits, currently hard-coded to 16
+		, Data(nullptr)
+		, DataSize(0)
+		, bIsRealTimeSourceReady(false)
+		, bIsDynamicResource(false)
+	{
+		if (InWave->DecompressionType == EDecompressionType::DTYPE_Native || InWave->DecompressionType == EDecompressionType::DTYPE_Preview)
+		{
+			Data = InWave->RawPCMData;
+			DataSize = InWave->RawPCMDataSize;
+// 			InWave->RawPCMData = nullptr;
+// 			InWave->RawPCMDataSize = 0;
+		}
+
+		// Set the base-class NumChannels to wave's NumChannels
+		NumChannels = InWave->NumChannels;
+	}
+
+	FMixerBuffer::~FMixerBuffer()
+	{
+		if (bAllocationInPermanentPool)
+		{
+			UE_LOG(LogAudioMixer, Fatal, TEXT("Can't free resource '%s' as it was allocated in permanent pool."), *ResourceName);
+		}
+
+		if (DecompressionState)
+		{
+			delete DecompressionState;
+			DecompressionState = nullptr;
+		}
+
+		switch (BufferType)
+		{
+			case EBufferType::PCM:
+			{
+				if (Data)
+				{
+					FMemory::Free((void*)Data);
+				}
+			}
+			break;
+
+			case EBufferType::PCMPreview:
+			{
+				if (bIsDynamicResource && Data)
+				{
+					FMemory::Free((void*)Data);
+				}
+			}
+			break;
+
+			case EBufferType::PCMRealTime:
+			case EBufferType::Streaming:
+			// Buffers are freed as part of the ~FSoundSource
+			break;
+
+			case EBufferType::Invalid:
+			// nothing
+			break;
+		}
+	}
+
+	int32 FMixerBuffer::GetSize()
+	{
+		switch (BufferType)
+		{
+			case EBufferType::PCM:
+			case EBufferType::PCMPreview:
+				return DataSize;
+
+			case EBufferType::PCMRealTime:
+				return (DecompressionState ? DecompressionState->GetSourceBufferSize() : 0) + (MONO_PCM_BUFFER_SIZE * NumChannels);
+			
+			case EBufferType::Streaming:
+				return MONO_PCM_BUFFER_SIZE * NumChannels;
+
+			case EBufferType::Invalid:
+			break;
+		}
+
+		return 0;
+	}
+
+	int32 FMixerBuffer::GetCurrentChunkIndex() const
+	{
+		if (DecompressionState)
+		{
+			return DecompressionState->GetCurrentChunkIndex();
+		}
+
+		return 0;
+	}
+
+	int32 FMixerBuffer::GetCurrentChunkOffset() const
+	{
+		if (DecompressionState)
+		{
+			return DecompressionState->GetCurrentChunkOffset();
+		}
+		return 0;
+	}
+
+	bool FMixerBuffer::IsRealTimeSourceReady()
+	{
+		// If we have a realtime async header parse task, then we check if its done
+		if (RealtimeAsyncHeaderParseTask)
+		{
+			bool bIsDone = RealtimeAsyncHeaderParseTask->IsDone();
+			if (bIsDone)
+			{
+				delete RealtimeAsyncHeaderParseTask;
+				RealtimeAsyncHeaderParseTask = nullptr;
+			}
+			return bIsDone;
+		}
+
+		// Otherwise, we weren't a real time decoding sound buffer (or we've already asked and it was ready)
+		return true;
+	}
+
+	bool FMixerBuffer::ReadCompressedInfo(USoundWave* SoundWave)
+	{
+		if (!DecompressionState)
+		{
+			UE_LOG(LogAudioMixer, Warning, TEXT("Attempting to read compressed info without a compression state instance for resource '%s'"), *ResourceName);
+			return false;
+		}
+		return DecompressionState->ReadCompressedInfo(SoundWave->ResourceData, SoundWave->ResourceSize, nullptr);
+	}
+
+	bool FMixerBuffer::ReadCompressedData(uint8* Destination, bool bLooping)
+	{
+		if (!DecompressionState)
+		{
+			UE_LOG(LogAudioMixer, Warning, TEXT("Attempting to read compressed data without a compression state instance for resource '%s'"), *ResourceName);
+			return false;
+		}
+
+		const uint32 kPCMBufferSize = MONO_PCM_BUFFER_SIZE  * NumChannels;
+		if (BufferType == EBufferType::Streaming)
+		{
+			return DecompressionState->StreamCompressedData(Destination, bLooping, kPCMBufferSize);
+		}
+		else
+		{
+			return DecompressionState->ReadCompressedData(Destination, bLooping, kPCMBufferSize);
+		}
+	}
+
+	void FMixerBuffer::Seek(const float SeekTime)
+	{
+		if (ensure(DecompressionState))
+		{
+			DecompressionState->SeekToTime(SeekTime);
+		}
+	}
+
+	FMixerBuffer* FMixerBuffer::Init(FAudioDevice* InAudioDevice, USoundWave* InWave, bool bForceRealtime)
+	{
+		// Can't create a buffer without any source data
+		if (InWave == nullptr || InWave->NumChannels == 0)
+		{
+			return nullptr;
+		}
+
+		FAudioDeviceManager* AudioDeviceManager = FAudioDevice::GetAudioDeviceManager();
+
+		FMixerDevice* Mixer = (FMixerDevice*)InAudioDevice;
+		FMixerBuffer* Buffer = nullptr;
+
+		EDecompressionType DecompressionType = InWave->DecompressionType;
+
+		if (bForceRealtime  && DecompressionType != DTYPE_Setup && DecompressionType != DTYPE_Streaming)
+		{
+			DecompressionType = DTYPE_RealTime;
+		}
+
+		switch (DecompressionType)
+		{
+			case DTYPE_Setup:
+			{
+				// We've circumvented the level-load precache mechanism, precache synchronously // TODO: support async loading here?
+				const bool bSynchronous = true;
+				InAudioDevice->Precache(InWave, bSynchronous, false);
+				check(InWave->DecompressionType != DTYPE_Setup);
+				Buffer = Init(InAudioDevice, InWave, bForceRealtime);
+			}
+			break;
+
+			case DTYPE_Preview:
+			{
+				// Find any existing buffers
+				if (InWave->ResourceID)
+				{
+					Buffer = (FMixerBuffer*)AudioDeviceManager->GetSoundBufferForResourceID(InWave->ResourceID);
+				}
+
+				// Override with any new PCM data even if the buffer already exists
+				if (InWave->RawPCMData)
+				{
+					// If we already have a buffer for this wave resource, free it
+					if (Buffer)
+					{
+						AudioDeviceManager->FreeBufferResource(Buffer);
+					}
+
+					// Create a new preview buffer
+					Buffer = FMixerBuffer::CreatePreviewBuffer(Mixer, InWave);
+
+					// Track the new created buffer
+					AudioDeviceManager->TrackResource(InWave, Buffer);
+				}
+			}
+			break;
+
+			case DTYPE_Procedural:
+			{
+				// Always create a new buffer for procedural buffers
+				Buffer = FMixerBuffer::CreateProceduralBuffer(Mixer, InWave);
+			}
+			break;
+
+			case DTYPE_RealTime:
+			{
+				// Always create a new buffer for real-time buffers
+				Buffer = FMixerBuffer::CreateRealTimeBuffer(Mixer, InWave);
+			}
+			break;
+
+			case DTYPE_Native:
+			case DTYPE_Xenon:
+			{
+				if (InWave->ResourceID)
+				{
+					Buffer = (FMixerBuffer*)AudioDeviceManager->GetSoundBufferForResourceID(InWave->ResourceID);
+				}
+
+				if (Buffer == nullptr)
+				{
+					Buffer = FMixerBuffer::CreateNativeBuffer(Mixer, InWave);
+
+					// Track the resource with the audio device manager
+					AudioDeviceManager->TrackResource(InWave, Buffer);
+					InWave->RemoveAudioResource();
+				}
+			}
+			break;
+
+			case DTYPE_Streaming:
+			{
+				Buffer = FMixerBuffer::CreateStreamingBuffer(Mixer, InWave);
+			}
+			break;
+
+			case DTYPE_Invalid:
+			default:
+			{
+				// Invalid will be set if the wave cannot be played.
+			}
+			break;
+		}
+
+		return Buffer;
+	}
+
+	FMixerBuffer* FMixerBuffer::CreatePreviewBuffer(FMixerDevice* InMixer, USoundWave* InWave)
+	{
+		// Create a new buffer
+		FMixerBuffer* Buffer = new FMixerBuffer(InMixer, InWave, EBufferType::PCMPreview);
+
+		Buffer->bIsDynamicResource = InWave->bDynamicResource;
+		return Buffer;
+	}
+
+	FMixerBuffer* FMixerBuffer::CreateProceduralBuffer(FMixerDevice* InMixer, USoundWave* InWave)
+	{
+		FMixerBuffer* Buffer = new FMixerBuffer(InMixer, InWave, EBufferType::PCMRealTime);
+
+		// No tracking of this resource needed
+		Buffer->ResourceID = 0;
+		InWave->ResourceID = 0;
+
+		return Buffer;
+	}
+
+	FMixerBuffer* FMixerBuffer::CreateNativeBuffer(FMixerDevice* InMixer, USoundWave* InWave)
+	{
+		// If wave a decompressor running, make sure it finishes and delete it
+		if (InWave->AudioDecompressor != nullptr)
+		{
+			InWave->AudioDecompressor->EnsureCompletion();
+			delete InWave->AudioDecompressor;
+			InWave->AudioDecompressor = nullptr;
+		}
+
+		FMixerBuffer* Buffer = new FMixerBuffer(InMixer, InWave, EBufferType::PCM);
+		return Buffer;
+	}
+
+	FMixerBuffer* FMixerBuffer::CreateStreamingBuffer(FMixerDevice* InMixer, USoundWave* InWave)
+	{
+		FMixerBuffer* Buffer = new FMixerBuffer(InMixer, InWave, EBufferType::Streaming);
+		
+		FSoundQualityInfo QualityInfo = { 0 };
+
+		Buffer->DecompressionState = InMixer->CreateCompressedAudioInfo(InWave);
+
+		// Get the header information of our compressed format
+		if (Buffer->DecompressionState->StreamCompressedInfo(InWave, &QualityInfo))
+		{
+			// Refresh the wave data
+			InWave->SampleRate = QualityInfo.SampleRate;
+			InWave->NumChannels = QualityInfo.NumChannels;
+			InWave->RawPCMDataSize = QualityInfo.SampleDataSize;
+			InWave->Duration = QualityInfo.Duration;
+		}
+		else
+		{
+			InWave->DecompressionType = DTYPE_Invalid;
+			InWave->NumChannels = 0;
+
+			InWave->RemoveAudioResource();
+		}
+
+		return Buffer;
+	}
+
+	FMixerBuffer* FMixerBuffer::CreateRealTimeBuffer(FMixerDevice* InMixer, USoundWave* InWave)
+	{
+		check(InMixer);
+		check(InWave);
+
+		// If we have a decompressor running, make sure it finishes and delete it
+		if (InWave->AudioDecompressor != nullptr)
+		{
+			InWave->AudioDecompressor->EnsureCompletion();
+
+			delete InWave->AudioDecompressor;
+			InWave->AudioDecompressor = nullptr;
+		}
+
+		// Create a new buffer for real-time sounds
+		FMixerBuffer* Buffer = new FMixerBuffer(InMixer, InWave, EBufferType::PCMRealTime);
+
+		if (InWave->ResourceData == nullptr)
+		{
+			InWave->InitAudioResource(InMixer->GetRuntimeFormat(InWave));
+		}
+
+		Buffer->DecompressionState = InMixer->CreateCompressedAudioInfo(InWave);
+
+		if (Buffer->DecompressionState)
+		{
+			check(Buffer->RealtimeAsyncHeaderParseTask == nullptr);
+			Buffer->RealtimeAsyncHeaderParseTask = new FAsyncRealtimeAudioTask(Buffer, InWave);
+			Buffer->RealtimeAsyncHeaderParseTask->StartBackgroundTask();
+
+			Buffer->NumChannels = InWave->NumChannels;
+		}
+		else
+		{
+			InWave->DecompressionType = DTYPE_Invalid;
+			InWave->NumChannels = 0;
+
+			InWave->RemoveAudioResource();
+		}
+
+		return Buffer;
+	}
+
+	EBufferType::Type FMixerBuffer::GetType() const
+	{
+		return BufferType;
+	}
+
+	bool FMixerBuffer::IsRealTimeBuffer() const
+	{
+		return BufferType == EBufferType::PCMRealTime || BufferType == EBufferType::Streaming;
+	}
+
+	void FMixerBuffer::GetPCMData(uint8** OutData, uint32* OutDataSize)
+	{
+		*OutData = Data;
+		*OutDataSize = DataSize;
+	}
+
+	void FMixerBuffer::EnsureHeaderParseTaskFinished()
+	{
+		if (RealtimeAsyncHeaderParseTask)
+		{
+			RealtimeAsyncHeaderParseTask->EnsureCompletion();
+			delete RealtimeAsyncHeaderParseTask;
+			RealtimeAsyncHeaderParseTask = nullptr;
+		}
+	}
+
+
+}
\ No newline at end of file
diff --git a/Engine/Source/Runtime/AudioMixer/Private/AudioMixerBuffer.h b/Engine/Source/Runtime/AudioMixer/Private/AudioMixerBuffer.h
new file mode 100644
--- /dev/null
+++ b/Engine/Source/Runtime/AudioMixer/Private/AudioMixerBuffer.h
@@ -0,0 +1,100 @@
+// Copyright 1998-2016 Epic Games, Inc. All Rights Reserved.
+
+#pragma once
+
+/* Public dependencies
+*****************************************************************************/
+
+#include "Core.h"
+#include "Engine.h"
+
+#include "SoundDefinitions.h"
+#include "AudioDecompress.h"
+#include "AudioMixer.h"
+#include "AudioMixerDevice.h"
+
+namespace Audio
+{
+	namespace EBufferType
+	{
+		enum Type
+		{
+			PCM,
+			PCMPreview,
+			PCMRealTime,
+			Streaming,
+			Invalid,
+		};
+	}
+
+	class FMixerDevice;
+	class FMixerBuffer;
+
+	typedef FAsyncRealtimeAudioTaskProxy<FMixerBuffer> FAsyncRealtimeAudioTask;
+
+	class FMixerBuffer : public FSoundBuffer
+	{
+	public:
+		FMixerBuffer(FAudioDevice* InAudioDevice, USoundWave* InWave, EBufferType::Type InBufferType);
+		~FMixerBuffer();
+
+		//~ Begin FSoundBuffer Interface
+		int32 GetSize() override;
+		int32 GetCurrentChunkIndex() const override;
+		int32 GetCurrentChunkOffset() const override;
+		bool IsRealTimeSourceReady() override;
+		bool ReadCompressedInfo(USoundWave* SoundWave) override;
+		bool ReadCompressedData(uint8* Destination, bool bLooping) override;
+		void Seek(const float SeekTime) override;
+
+		//~ End FSoundBuffer Interface
+
+		static FMixerBuffer* Init(FAudioDevice* AudioDevice, USoundWave* InWave, bool bForceRealtime);
+		static FMixerBuffer* CreatePreviewBuffer(FMixerDevice* InMixer, USoundWave* InWave);
+		static FMixerBuffer* CreateProceduralBuffer(FMixerDevice* InMixer, USoundWave* InWave);
+		static FMixerBuffer* CreateNativeBuffer(FMixerDevice* InMixer, USoundWave* InWave);
+		static FMixerBuffer* CreateStreamingBuffer(FMixerDevice* InMixer, USoundWave* InWave);
+		static FMixerBuffer* CreateRealTimeBuffer(FMixerDevice* InMixer, USoundWave* InWave);
+
+		/** Returns the buffer's format */
+		EBufferType::Type GetType() const;
+		bool IsRealTimeBuffer() const;
+
+		/** Returns the contained raw PCM data and data size */
+		void GetPCMData(uint8** OutData, uint32* OutDataSize);
+
+		void EnsureHeaderParseTaskFinished();
+
+		float GetSampleRate() const { return SampleRate; }
+
+	private:
+
+		/** Async task for parsing real-time decompressed compressed info headers. */
+		FAsyncRealtimeAudioTask* RealtimeAsyncHeaderParseTask;
+
+		/** Wrapper to handle the decompression of audio codecs. */
+		ICompressedAudioInfo* DecompressionState;
+
+		/** Format of the sound referenced by this buffer */
+		EBufferType::Type BufferType;
+
+		/** Sample rate of the audio buffer. */
+		int32 SampleRate;
+
+		/** Number of bits per sample. */
+		int16 BitsPerSample;
+
+		/** Ptr to raw PCM data. */
+		uint8* Data;
+
+		/** The raw PCM data size. */
+		uint32 DataSize;
+
+		/** Bool indicating the that the real-time source is ready for real-time decoding. */
+		FThreadSafeBool bIsRealTimeSourceReady;
+
+		/** Set to true when the PCM data should be freed when the buffer is destroyed */
+		bool bIsDynamicResource;
+	};
+}
+
diff --git a/Engine/Source/Runtime/AudioMixer/Private/AudioMixerChannelMaps.cpp b/Engine/Source/Runtime/AudioMixer/Private/AudioMixerChannelMaps.cpp
new file mode 100644
--- /dev/null
+++ b/Engine/Source/Runtime/AudioMixer/Private/AudioMixerChannelMaps.cpp
@@ -0,0 +1,310 @@
+// Copyright 1998-2016 Epic Games, Inc. All Rights Reserved.
+
+#include "AudioMixerPCH.h"
+#include "AudioMixerDevice.h"
+
+namespace Audio
+{
+	/** Default 2D channel maps from input channels to output channels */
+
+	static float MonoToStereoMatrix[1 * 2] =
+	{
+		// FrontLeft
+		0.707f,		// FrontLeft
+		0.707f,		// FrontRight
+	};
+
+	static float MonoToQuadMatrix[1 * 4] =
+	{
+		// FrontLeft
+		0.707f,		// FrontLeft
+		0.707f,		// FrontRight
+		0.0f,		// BackLeft
+		0.0f,		// BackRight
+	};
+
+	static float MonoTo51Matrix[1 * 6] =
+	{
+		// FrontLeft
+		0.707f,		// FrontLeft
+		0.707f,		// FrontRight
+		0.0f,		// FrontCenter
+		0.0f,		// LowFrequency
+		0.0f,		// BackLeft
+		0.0f,		// BackRight
+	};
+
+	static float MonoTo71Matrix[1 * 8] =
+	{
+		// FrontLeft
+		0.707f,		// FrontLeft
+		0.707f,		// FrontRight
+		0.0f,		// FrontCenter
+		0.0f,		// LowFrequency
+		0.0f,		// BackLeft
+		0.0f,		// BackRight
+		0.0f,		// SideLeft
+		0.0f,		// SideRight
+	};
+
+	static float StereoToStereoMatrix[2 * 2] =
+	{
+		// FrontLeft		FrontRight
+		1.0f,				0.0f,		// FrontLeft
+		0.0f,				1.0f,		// FrontRight
+	};
+
+	static float StereoToQuadMatrix[2 * 4] =
+	{
+		// FrontLeft		FrontRight
+		0.707f,				0.0f,		// FrontLeft
+		0.0f,				0.707f,		// FrontRight
+		0.0f,				0.707f,		// BackLeft
+		0.707f,				0.0f,		// BackRight
+	};
+
+	static float StereoTo51Matrix[2 * 6] =
+	{
+		// FrontLeft		FrontRight
+		0.707f,				0.0f,		// FrontLeft
+		0.0f,				0.707f,		// FrontRight
+		0.0f,				0.0f,		// FrontCenter
+		0.0f,				0.0f,		// LowFrequency
+		0.0f,				0.707f,		// BackLeft
+		0.707f,				0.0f,		// BackRight
+	};
+
+	static float StereoTo71Matrix[2 * 8] =
+	{
+		// FrontLeft		FrontRight
+		0.707f,				0.0f,		// FrontLeft
+		0.0f,				0.707f,		// FrontRight
+		0.0f,				0.0f,		// FrontCenter
+		0.0f,				0.0f,		// LowFrequency
+		0.0f,				0.707f,		// BackLeft
+		0.707f,				0.0f,		// BackRight
+		0.0f,				0.0f,		// SideLeft
+		0.0f,				0.0f,		// SideRight
+	};
+
+	static float QuadToStereoMatrix[4 * 2] =
+	{
+		// FrontLeft		FrontRight		BackLeft		BackRight
+		0.5f,				0.0f,			0.5f,			0.0f,			// FrontLeft
+		0.0f,				0.5f,			0.0f,			0.5f,			// FrontRight
+	};
+
+	static float QuadToQuadMatrix[4 * 4] =
+	{
+		// FrontLeft		FrontRight		BackLeft		BackRight
+		1.0f,				0.0f,			0.0f,			0.0f,			// FrontLeft
+		0.0f,				1.0f,			0.0f,			0.0f,			// FrontRight
+		0.0f,				0.0f,			1.0f,			0.0f,			// BackLeft
+		0.0f,				0.0f,			0.0f,			1.0f,			// BackRight
+	};
+
+	static float QuadTo51Matrix[4 * 6] =
+	{
+		// FrontLeft		FrontRight		BackLeft		BackRight
+		1.0f,				0.0f,			0.0f,			0.0f,			// FrontLeft
+		0.0f,				1.0f,			0.0f,			0.0f,			// FrontRight
+		0.0f,				0.0f,			0.0f,			0.0f,			// FrontCenter
+		0.0f,				0.0f,			0.0f,			0.0f,			// LowFrequency
+		0.0f,				0.0f,			1.0f,			0.0f,			// BackLeft
+		0.0f,				0.0f,			0.0f,			1.0f,			// BackRight
+	};
+
+	static float QuadTo71Matrix[4 * 8] =
+	{
+		// FrontLeft		FrontRight		BackLeft		BackRight
+		1.0f,				0.0f,			0.0f,			0.0f,			// FrontLeft
+		0.0f,				1.0f,			0.0f,			0.0f,			// FrontRight
+		0.0f,				0.0f,			0.0f,			0.0f,			// FrontCenter
+		0.0f,				0.0f,			0.0f,			0.0f,			// LowFrequency
+		0.0f,				0.0f,			1.0f,			0.0f,			// BackLeft
+		0.0f,				0.0f,			0.0f,			1.0f,			// BackRight
+		0.0f,				0.0f,			0.0f,			0.0f,			// SideLeft
+		0.0f,				0.0f,			0.0f,			0.0f,			// SideRight
+	};
+
+	static float HexToStereoMatrix[6 * 2] =
+	{
+		// FrontLeft		FrontRight		FrontCenter		LowFrequency		BackLeft		BackRight
+		0.577f,				0.0f,			0.577f,			0.0f,				0.577f,			0.0f,			// FrontLeft
+		0.0f,				0.577f,			0.577f,			0.0f,				0.577f,			0.577f,			// FrontRight
+	};
+
+	static float HexToQuadMatrix[6 * 4] =
+	{
+		// FrontLeft		FrontRight		FrontCenter		LowFrequency		BackLeft		BackRight
+		0.577f,				0.0f,			0.577f,			0.0f,				0.0f,			0.0f,			// FrontLeft
+		0.0f,				0.577f,			0.577f,			0.0f,				0.0f,			0.0f,			// FrontRight
+		0.0f,				0.0f,			0.0f,			0.0f,				0.577f,			0.0f,			// BackLeft
+		0.0f,				0.0f,			0.0f,			0.0f,				0.0f,			0.577f,			// BackRight
+	};
+
+	static float HexTo51Matrix[6 * 6] =
+	{
+		// FrontLeft		FrontRight		FrontCenter		LowFrequency		BackLeft		BackRight
+		1.0f,				0.0f,			0.0f,			0.0f,				0.0f,			0.0f,			// FrontLeft
+		0.0f,				1.0f,			0.0f,			0.0f,				0.0f,			0.0f,			// FrontRight
+		0.0f,				0.0f,			1.0f,			0.0f,				0.0f,			0.0f,			// FrontCenter
+		0.0f,				0.0f,			0.0f,			1.0f,				0.0f,			0.0f,			// LowFrequency
+		0.0f,				0.0f,			0.0f,			0.0f,				1.0f,			0.0f,			// BackLeft
+		0.0f,				0.0f,			0.0f,			0.0f,				0.0f,			1.0f,			// BackRight
+	};
+
+	static float HexTo71Matrix[6 * 8] =
+	{
+		// FrontLeft		FrontRight		FrontCenter		LowFrequency		BackLeft		BackRight
+		1.0f,				0.0f,			0.0f,			0.0f,				0.0f,			0.0f,			// FrontLeft
+		0.0f,				1.0f,			0.0f,			0.0f,				0.0f,			0.0f,			// FrontRight
+		0.0f,				0.0f,			1.0f,			0.0f,				0.0f,			0.0f,			// FrontCenter
+		0.0f,				0.0f,			0.0f,			1.0f,				0.0f,			0.0f,			// LowFrequency
+		0.0f,				0.0f,			0.0f,			0.0f,				1.0f,			0.0f,			// BackLeft
+		0.0f,				0.0f,			0.0f,			0.0f,				0.0f,			1.0f,			// BackRight
+		0.0f,				0.0f,			0.0f,			0.0f,				0.0f,			0.0f,			// SideLeft
+		0.0f,				0.0f,			0.0f,			0.0f,				0.0f,			0.0f,			// SideRight
+	};
+
+	static inline int32 GetChannelMapID(const int32 InputChannels, const int32 OutputChannels)
+	{
+		return InputChannels + 100 * OutputChannels;
+	}
+
+	void FMixerDevice::Get2DChannelMap(const int32 NumSourceChannels, TArray<float>& OutChannelMap)
+	{
+		int32 ChannelMapIndex = GetChannelMapID(NumSourceChannels, PlatformInfo.NumChannels);
+		FChannelMapEntry* Entry = ChannelMaps.Find(ChannelMapIndex);
+		if (!Entry)
+		{
+			OutChannelMap.AddZeroed(PlatformInfo.NumChannels);
+			UE_LOG(LogAudioMixer, Warning, TEXT("Unsupported source channel count or output channels %d"), NumSourceChannels);
+			return;
+		}
+
+		OutChannelMap = Entry->ChannelWeights;
+	}
+
+	void FMixerDevice::Initialize2DChannelMaps()
+	{
+		ChannelMaps.Reset();
+
+		ChannelMaps.Add(GetChannelMapID(1, 2), FChannelMapEntry(EChannelMap::MonoToStereo,	MonoToStereoMatrix, ARRAY_COUNT(MonoToStereoMatrix)));
+		ChannelMaps.Add(GetChannelMapID(1, 4), FChannelMapEntry(EChannelMap::MonoToQuad,	MonoToQuadMatrix,	ARRAY_COUNT(MonoToQuadMatrix)));
+		ChannelMaps.Add(GetChannelMapID(1, 6), FChannelMapEntry(EChannelMap::MonoTo51,		MonoTo51Matrix,		ARRAY_COUNT(MonoTo51Matrix)));
+		ChannelMaps.Add(GetChannelMapID(1, 8), FChannelMapEntry(EChannelMap::MonoTo71,		MonoTo71Matrix,		ARRAY_COUNT(MonoTo71Matrix)));
+
+		ChannelMaps.Add(GetChannelMapID(2, 2), FChannelMapEntry(EChannelMap::StereoToStereo, StereoToStereoMatrix,	ARRAY_COUNT(StereoToStereoMatrix)));
+		ChannelMaps.Add(GetChannelMapID(2, 4), FChannelMapEntry(EChannelMap::StereoToQuad,	StereoToQuadMatrix,		ARRAY_COUNT(StereoToQuadMatrix)));
+		ChannelMaps.Add(GetChannelMapID(2, 6), FChannelMapEntry(EChannelMap::StereoTo51,	StereoTo51Matrix,		ARRAY_COUNT(StereoTo51Matrix)));
+		ChannelMaps.Add(GetChannelMapID(2, 8), FChannelMapEntry(EChannelMap::StereoTo71,	StereoTo71Matrix,		ARRAY_COUNT(StereoTo71Matrix)));
+
+		ChannelMaps.Add(GetChannelMapID(4, 2), FChannelMapEntry(EChannelMap::QuadToStereo,	QuadToStereoMatrix,	ARRAY_COUNT(QuadToStereoMatrix)));
+		ChannelMaps.Add(GetChannelMapID(4, 4), FChannelMapEntry(EChannelMap::QuadToQuad,	QuadToQuadMatrix,	ARRAY_COUNT(QuadToQuadMatrix)));
+		ChannelMaps.Add(GetChannelMapID(4, 6), FChannelMapEntry(EChannelMap::QuadTo51,		QuadTo51Matrix,		ARRAY_COUNT(QuadTo51Matrix)));
+		ChannelMaps.Add(GetChannelMapID(4, 8), FChannelMapEntry(EChannelMap::QuadTo71,		QuadTo71Matrix,		ARRAY_COUNT(QuadTo71Matrix)));
+
+		ChannelMaps.Add(GetChannelMapID(6, 2), FChannelMapEntry(EChannelMap::HexToStereo,	HexToStereoMatrix,	ARRAY_COUNT(HexToStereoMatrix)));
+		ChannelMaps.Add(GetChannelMapID(6, 4), FChannelMapEntry(EChannelMap::HexToQuad,		HexToQuadMatrix,	ARRAY_COUNT(HexToQuadMatrix)));
+		ChannelMaps.Add(GetChannelMapID(6, 6), FChannelMapEntry(EChannelMap::HexTo51,		HexTo51Matrix,		ARRAY_COUNT(HexTo51Matrix)));
+		ChannelMaps.Add(GetChannelMapID(6, 8), FChannelMapEntry(EChannelMap::HexTo71,		HexTo71Matrix,		ARRAY_COUNT(HexTo71Matrix)));
+	}
+
+	void FMixerDevice::InitializeChannelAzimuthMap(const int32 NumChannels)
+	{
+		// Now setup the hard-coded values
+		if (NumChannels == 2)
+		{
+			DefaultChannelAzimuthPosition[EAudioMixerChannel::FrontLeft] = { EAudioMixerChannel::FrontLeft, 270 };
+			DefaultChannelAzimuthPosition[EAudioMixerChannel::FrontRight] = { EAudioMixerChannel::FrontRight, 90 };
+		}
+		else
+		{
+			DefaultChannelAzimuthPosition[EAudioMixerChannel::FrontLeft] = { EAudioMixerChannel::FrontLeft, 330 };
+			DefaultChannelAzimuthPosition[EAudioMixerChannel::FrontRight] = { EAudioMixerChannel::FrontRight, 30 };
+		}
+
+		// Ignore front cent and low-frequency for azimuth computations. Both of these channels are manual opt-ins from data and not directly used in 3d audio calculations
+		DefaultChannelAzimuthPosition[EAudioMixerChannel::FrontCenter] = { EAudioMixerChannel::FrontCenter, -1 };
+		DefaultChannelAzimuthPosition[EAudioMixerChannel::LowFrequency] = { EAudioMixerChannel::LowFrequency, -1 };
+
+		DefaultChannelAzimuthPosition[EAudioMixerChannel::BackLeft] = { EAudioMixerChannel::BackLeft, 210 };
+		DefaultChannelAzimuthPosition[EAudioMixerChannel::BackRight] = { EAudioMixerChannel::BackRight, 150 };
+		DefaultChannelAzimuthPosition[EAudioMixerChannel::FrontLeftOfCenter] = { EAudioMixerChannel::FrontLeftOfCenter, 15 };
+		DefaultChannelAzimuthPosition[EAudioMixerChannel::FrontRightOfCenter] = { EAudioMixerChannel::FrontRightOfCenter, 345 };
+		DefaultChannelAzimuthPosition[EAudioMixerChannel::BackCenter] = { EAudioMixerChannel::BackCenter, 180 };
+		DefaultChannelAzimuthPosition[EAudioMixerChannel::SideLeft] = { EAudioMixerChannel::SideLeft, 270 };
+		DefaultChannelAzimuthPosition[EAudioMixerChannel::SideRight] = { EAudioMixerChannel::SideRight, 90 };
+
+		// Check any engine ini overrides for these default positions
+		if (NumChannels != 2)
+		{
+			int32 AzimuthPositionOverride = 0;
+			for (int32 ChannelOverrideIndex = 0; ChannelOverrideIndex < EAudioMixerChannel::MaxSupportedChannel; ++ChannelOverrideIndex)
+			{
+				EAudioMixerChannel::Type MixerChannelType = EAudioMixerChannel::Type(ChannelOverrideIndex);
+				const TCHAR* ChannelName = EAudioMixerChannel::ToString(MixerChannelType);
+				if (GConfig->GetInt(TEXT("AudioChannelAzimuthMap"), ChannelName, AzimuthPositionOverride, GEngineIni))
+				{
+					if (AzimuthPositionOverride >= 0 && AzimuthPositionOverride < 360)
+					{
+						// Make sure no channels have this azimuth angle first, otherwise we'll get some bad math later
+						bool bIsUnique = true;
+						for (int32 ExistingChannelIndex = 0; ExistingChannelIndex < EAudioMixerChannel::MaxSupportedChannel; ++ExistingChannelIndex)
+						{
+							if (DefaultChannelAzimuthPosition[ExistingChannelIndex].Azimuth == AzimuthPositionOverride)
+							{
+								bIsUnique = false;
+
+								// If the override is setting the same value as our default, don't print a warning
+								if (ExistingChannelIndex != ChannelOverrideIndex)
+								{
+									const TCHAR* ExistingChannelName = EAudioMixerChannel::ToString(EAudioMixerChannel::Type(ExistingChannelIndex));
+									UE_LOG(LogAudioMixer, Warning, TEXT("Azimuth value '%d' for audio mixer channel '%s' is already used by '%s'. Azimuth values must be unique."),
+										AzimuthPositionOverride, ChannelName, ExistingChannelName);
+								}
+								break;
+							}
+						}
+
+						if (bIsUnique)
+						{
+							DefaultChannelAzimuthPosition[MixerChannelType].Azimuth = AzimuthPositionOverride;
+						}
+					}
+					else
+					{
+						UE_LOG(LogAudioMixer, Warning, TEXT("Azimuth value, %d, for audio mixer channel %s out of range. Must be [0, 360)."), AzimuthPositionOverride, ChannelName);
+					}
+				}
+			}
+		}
+
+		// Build a map of azimuth positions of only the current audio device's output channels
+		CurrentChannelAzimuthPositions.Reset();
+		for (EAudioMixerChannel::Type Channel : PlatformInfo.OutputChannelArray)
+		{
+			// Only track non-LFE and non-Center channel azimuths for use with 3d channel mappings
+			if (Channel != EAudioMixerChannel::LowFrequency && Channel != EAudioMixerChannel::FrontCenter)
+			{
+				++NumSpatialChannels;
+				CurrentChannelAzimuthPositions.Add(DefaultChannelAzimuthPosition[Channel]);
+			}
+		}
+
+		check(NumSpatialChannels > 0);
+		OmniPanFactor = 1.0f / FMath::Sqrt(NumSpatialChannels);
+
+		// Sort the current mapping by azimuth
+		struct FCompareByAzimuth
+		{
+			FORCEINLINE bool operator()(const FChannelPositionInfo& A, const FChannelPositionInfo& B) const
+			{
+				return A.Azimuth < B.Azimuth;
+			}
+		};
+
+		CurrentChannelAzimuthPositions.Sort(FCompareByAzimuth());
+	}
+}
\ No newline at end of file
diff --git a/Engine/Source/Runtime/AudioMixer/Private/AudioMixerDevice.cpp b/Engine/Source/Runtime/AudioMixer/Private/AudioMixerDevice.cpp
new file mode 100644
--- /dev/null
+++ b/Engine/Source/Runtime/AudioMixer/Private/AudioMixerDevice.cpp
@@ -0,0 +1,528 @@
+// Copyright 1998-2016 Epic Games, Inc. All Rights Reserved.
+
+#include "AudioMixerPCH.h"
+#include "AudioMixerDevice.h"
+#include "AudioMixerSource.h"
+#include "DSP/Noise.h"
+#include "DSP/SinOsc.h"
+
+namespace Audio
+{
+	FMixerDevice::FMixerDevice(IAudioMixerPlatformInterface* InAudioMixerPlatform)
+		: AudioMixerPlatform(InAudioMixerPlatform)
+		, NumSpatialChannels(0)
+		, OmniPanFactor(0.0f)
+		, AudioClockDelta(0.0)
+		, AudioClock(0.0)
+		, MasterSubmix(nullptr)
+		, SourceManager(this)
+		, GameOrAudioThreadId(INDEX_NONE)
+		, AudioPlatformThreadId(INDEX_NONE)
+		, bDebugOutputEnabled(false)
+	{
+		// This audio device is the audio mixer
+		bAudioMixerModuleLoaded = true;
+	}
+
+	FMixerDevice::~FMixerDevice()
+	{
+		AUDIO_MIXER_CHECK_GAME_THREAD(this);
+
+		if (AudioMixerPlatform != nullptr)
+		{
+			delete AudioMixerPlatform;
+		}
+	}
+
+	void FMixerDevice::CheckGameThread()
+	{
+		if (GameOrAudioThreadId == INDEX_NONE)
+		{
+			GameOrAudioThreadId = FPlatformTLS::GetCurrentThreadId();
+		}
+		int32 CurrentThreadId = FPlatformTLS::GetCurrentThreadId();
+		AUDIO_MIXER_CHECK(CurrentThreadId == GameOrAudioThreadId);
+	}
+
+	void FMixerDevice::ResetAudioPlatformThreadId()
+	{
+#if AUDIO_MIXER_ENABLE_DEBUG_MODE
+		AudioPlatformThreadId = INDEX_NONE;
+		CheckAudioPlatformThread();
+#endif // AUDIO_MIXER_ENABLE_DEBUG_MODE
+	}
+
+	void FMixerDevice::CheckAudioPlatformThread()
+	{
+		if (AudioPlatformThreadId == INDEX_NONE)
+		{
+			AudioPlatformThreadId = FPlatformTLS::GetCurrentThreadId();
+		}
+		int32 CurrentThreadId = FPlatformTLS::GetCurrentThreadId();
+		AUDIO_MIXER_CHECK(CurrentThreadId == AudioPlatformThreadId);
+	}
+
+	void FMixerDevice::GetAudioDeviceList(TArray<FString>& OutAudioDeviceNames) const
+	{
+		if (AudioMixerPlatform && AudioMixerPlatform->IsInitialized())
+		{
+			uint32 NumOutputDevices;
+			if (AudioMixerPlatform->GetNumOutputDevices(NumOutputDevices))
+			{
+				for (uint32 i = 0; i < NumOutputDevices; ++i)
+				{
+					FAudioPlatformDeviceInfo DeviceInfo;
+					if (AudioMixerPlatform->GetOutputDeviceInfo(i, DeviceInfo))
+					{
+						OutAudioDeviceNames.Add(DeviceInfo.Name);
+					}
+				}
+			}
+		}
+	}
+
+	bool FMixerDevice::InitializeHardware()
+	{
+		AUDIO_MIXER_CHECK_GAME_THREAD(this);
+
+		if (AudioMixerPlatform && AudioMixerPlatform->InitializeHardware())
+		{
+			AUDIO_MIXER_CHECK(SampleRate != 0.0f);
+			AUDIO_MIXER_CHECK(BufferLength != 0);
+
+			OpenStreamParams.NumFrames = BufferLength;
+			OpenStreamParams.OutputDeviceIndex = 0; // Default device (for now)
+			OpenStreamParams.SampleRate = SampleRate;
+			OpenStreamParams.AudioMixer = this;
+
+			if (AudioMixerPlatform->OpenAudioStream(OpenStreamParams))
+			{
+				// Get the platform device info we're using
+				PlatformInfo = AudioMixerPlatform->GetPlatformDeviceInfo();
+
+				// Initialize some data that depends on speaker configuration, etc.
+				InitializeChannelAzimuthMap(PlatformInfo.NumChannels);
+				Initialize2DChannelMaps();
+
+				SourceManager.Init(MaxChannels);
+
+				AudioClock = 0.0;
+				AudioClockDelta = (double)OpenStreamParams.NumFrames / OpenStreamParams.SampleRate;
+
+				// Get the master submix UObject
+
+				check(MasterSubmix == nullptr);
+				MasterSubmix = new FMixerSubmix(nullptr, this);
+
+				// Start streaming audio
+				return AudioMixerPlatform->StartAudioStream();
+			}
+		}
+		return false;
+	}
+
+	void FMixerDevice::TeardownHardware()
+	{
+		if (AudioMixerPlatform)
+		{
+			AudioMixerPlatform->TeardownHardware();
+
+			delete MasterSubmix;
+			MasterSubmix = nullptr;
+		}
+	}
+
+	void FMixerDevice::UpdateHardware()
+	{
+		SourceManager.Update();
+	}
+
+	double FMixerDevice::GetAudioTime() const
+	{
+		return AudioClock;
+	}
+
+	FAudioEffectsManager* FMixerDevice::CreateEffectsManager()
+	{
+		return nullptr;
+	}
+
+	FSoundSource* FMixerDevice::CreateSoundSource()
+	{
+		return new FMixerSource(this);
+	}
+
+	FName FMixerDevice::GetRuntimeFormat(USoundWave* InSoundWave)
+	{
+		check(AudioMixerPlatform);
+		return AudioMixerPlatform->GetRuntimeFormat(InSoundWave);
+	}
+
+	bool FMixerDevice::HasCompressedAudioInfoClass(USoundWave* InSoundWave)
+	{
+		check(InSoundWave);
+		check(AudioMixerPlatform);
+		return AudioMixerPlatform->HasCompressedAudioInfoClass(InSoundWave);
+	}
+
+	bool FMixerDevice::SupportsRealtimeDecompression() const
+	{
+		// TODO: Test and implement realt-time decompression on all other platforms
+#if PLATFORM_WINDOWS
+		return true;
+#else
+		return false;
+#endif
+		
+	}
+
+	class ICompressedAudioInfo* FMixerDevice::CreateCompressedAudioInfo(USoundWave* InSoundWave)
+	{
+		check(InSoundWave);
+		check(AudioMixerPlatform);
+		return AudioMixerPlatform->CreateCompressedAudioInfo(InSoundWave);
+	}
+
+	bool FMixerDevice::ValidateAPICall(const TCHAR* Function, uint32 ErrorCode)
+	{
+		return false;
+	}
+
+	bool FMixerDevice::Exec(UWorld* InWorld, const TCHAR* Cmd, FOutputDevice& Ar)
+	{
+		return false;
+	}
+
+	void FMixerDevice::CountBytes(FArchive& InArchive)
+	{
+		FAudioDevice::CountBytes(InArchive);
+	}
+
+	bool FMixerDevice::IsExernalBackgroundSoundActive()
+	{
+		return false;
+	}
+
+	void FMixerDevice::Precache(class USoundWave* InSoundWave, bool bSynchronous, bool bTrackMemory)
+	{
+		FAudioDevice::Precache(InSoundWave, bSynchronous, bTrackMemory);
+	}
+
+	void FMixerDevice::ResumeContext()
+	{
+	}
+
+	void FMixerDevice::SetMaxChannels(int InMaxChannels)
+	{
+	}
+
+	void FMixerDevice::StopAllSounds(bool InStop)
+	{
+	}
+
+	void FMixerDevice::SuspendContext()
+	{
+	}
+
+	void FMixerDevice::EnableDebugAudioOutput()
+	{
+		bDebugOutputEnabled = true;
+	}
+
+	bool FMixerDevice::OnProcessAudioStream(TArray<float>& Output)
+	{
+		// This function could be called in a task manager, which means the thread ID may change between calls.
+		ResetAudioPlatformThreadId();
+
+		// Compute the next block of audio in the source manager
+		SourceManager.ComputeNextBlockOfSamples();
+
+		if (MasterSubmix)
+		{
+			// Process the audio output from the master submix
+			MasterSubmix->ProcessAudio(Output);
+		}
+
+		// Do any debug output performing
+ 		if (bDebugOutputEnabled)
+		{
+			SineOscTest(Output);
+		}
+
+		// Update the audio clock
+		AudioClock += AudioClockDelta;
+
+		return true;
+	}
+
+	void FMixerDevice::InitSoundSubmixes()
+	{
+		// Reset existing submixes if they exist
+		Submixes.Reset();
+
+		// Reset the maps of sound instances properties
+		for (TObjectIterator<USoundSubmix> It; It; ++It)
+		{
+			USoundSubmix* SoundSubmix = *It;
+			FMixerSubmix* MixerSubmix = new FMixerSubmix(SoundSubmix, this);
+			Submixes.Add(SoundSubmix, MixerSubmix);
+		}
+
+	}
+
+	void FMixerDevice::RegisterSoundSubmix(USoundSubmix* SoundSubmix)
+	{
+		if (SoundSubmix)
+		{
+			if (!Submixes.Contains(SoundSubmix))
+			{
+				FMixerSubmix* MixerSubmix = new FMixerSubmix(SoundSubmix, this);
+				Submixes.Add(SoundSubmix, MixerSubmix);
+			}
+		}
+	}
+
+	void FMixerDevice::UnregisterSoundSubmix(USoundSubmix* SoundSubmix)
+	{
+		if (SoundSubmix)
+		{
+			Submixes.Remove(SoundSubmix);
+		}
+	}
+
+	FMixerSubmix* FMixerDevice::GetSubmixInstance(USoundSubmix* SoundSubmix)
+	{
+		FMixerSubmix** Result = Submixes.Find(SoundSubmix);
+		if (Result)
+		{
+			return *Result;
+		}
+		return nullptr;
+	}
+
+	FMixerSourceVoice* FMixerDevice::GetMixerSourceVoice(const FWaveInstance* InWaveInstance, ISourceBufferQueueListener* InBufferQueueListener, bool bUseHRTFSpatialization)
+	{
+		// Create a new mixer source voice using our source manager
+		FMixerSourceVoice* NewMixerSourceVoice = new FMixerSourceVoice(this, &SourceManager);
+
+		return NewMixerSourceVoice;
+	}
+
+	int32 FMixerDevice::GetNumSources() const
+	{
+		return Sources.Num();
+	}
+
+	int32 FMixerDevice::GetNumActiveSources() const
+	{
+		return SourceManager.GetNumActiveSources();
+	}
+
+	void FMixerDevice::Get3DChannelMap(const FWaveInstance* InWaveInstance, float EmitterAzimith, float NormalizedOmniRadius, TArray<float>& OutChannelMap)
+	{
+		// If we're center-channel only, then no need for spatial calculations, but need to build a channel map
+		if (InWaveInstance->bCenterChannelOnly)
+		{
+			// If we only have stereo channels
+			if (NumSpatialChannels == 2)
+			{
+				// Equal volume in left + right channel with equal power panning
+				static const float Pan = 1.0f / FMath::Sqrt(2.0f);
+				OutChannelMap.Add(Pan);
+				OutChannelMap.Add(Pan);
+			}
+			else
+			{
+				for (EAudioMixerChannel::Type Channel : PlatformInfo.OutputChannelArray)
+				{
+					float Pan = (Channel == EAudioMixerChannel::FrontCenter) ? 1.0f : 0.0f;
+					OutChannelMap.Add(Pan);
+				}
+			}
+
+			return;
+		}
+
+		float Azimuth = EmitterAzimith;
+
+		const FChannelPositionInfo* PrevChannelInfo = nullptr;
+		const FChannelPositionInfo* NextChannelInfo = nullptr;
+
+		for (int32 i = 0; i < CurrentChannelAzimuthPositions.Num(); ++i)
+		{
+			const FChannelPositionInfo& ChannelPositionInfo = CurrentChannelAzimuthPositions[i];
+
+			if (Azimuth <= ChannelPositionInfo.Azimuth)
+			{
+				NextChannelInfo = &CurrentChannelAzimuthPositions[i];
+
+				int32 PrevIndex = i - 1;
+				if (PrevIndex < 0)
+				{
+					PrevIndex = CurrentChannelAzimuthPositions.Num() - 1;
+				}
+
+				PrevChannelInfo = &CurrentChannelAzimuthPositions[PrevIndex];
+				break;
+			}
+		}
+
+		// If we didn't find anything, that means our azimuth position is at the top of the mapping
+		if (PrevChannelInfo == nullptr)
+		{
+			PrevChannelInfo = &CurrentChannelAzimuthPositions[CurrentChannelAzimuthPositions.Num() - 1];
+			NextChannelInfo = &CurrentChannelAzimuthPositions[0];
+			AUDIO_MIXER_CHECK(PrevChannelInfo != NextChannelInfo);
+		}
+
+		float NextChannelAzimuth = NextChannelInfo->Azimuth;
+		float PrevChannelAzimuth = PrevChannelInfo->Azimuth;
+
+		if (NextChannelAzimuth < PrevChannelAzimuth)
+		{
+			NextChannelAzimuth += 360.0f;
+		}
+
+		if (Azimuth < PrevChannelAzimuth)
+		{
+			Azimuth += 360.0f;
+		}
+
+		AUDIO_MIXER_CHECK(NextChannelAzimuth > PrevChannelAzimuth);
+		AUDIO_MIXER_CHECK(Azimuth > PrevChannelAzimuth);
+		float Fraction = (Azimuth - PrevChannelAzimuth) / (NextChannelAzimuth - PrevChannelAzimuth);
+		AUDIO_MIXER_CHECK(Fraction >= 0.0f && Fraction <= 1.0f);
+
+		// Compute the panning values using equal-power panning law
+		float PrevChannelPan = FMath::Cos(Fraction * 0.5f * PI);
+		float NextChannelPan = FMath::Sin(Fraction * 0.5f * PI);
+
+		float NormalizedOmniRadSquared = NormalizedOmniRadius * NormalizedOmniRadius;
+		float OmniAmount = 0.0f;
+
+		if (NormalizedOmniRadSquared > 0.0f)
+		{
+			OmniAmount = 1.0f - FMath::Max(1.0f / NormalizedOmniRadSquared, 1.0f);
+		}
+
+		// OmniPan is the amount of pan to use if fully omni-directional
+		AUDIO_MIXER_CHECK(NumSpatialChannels > 0);
+
+		// Build the output channel map based on the current platform device output channel array 
+
+		float DefaultEffectivePan = !OmniAmount ? 0.0f : FMath::Lerp(0.0f, OmniPanFactor, OmniAmount);
+
+		for (EAudioMixerChannel::Type Channel : PlatformInfo.OutputChannelArray)
+		{
+			float EffectivePan = DefaultEffectivePan;
+
+			// Check for manual channel mapping parameters (LFE and Front Center)
+			if (Channel == EAudioMixerChannel::LowFrequency)
+			{
+				EffectivePan = InWaveInstance->LFEBleed;
+			}
+			else if (Channel == EAudioMixerChannel::FrontCenter)
+			{
+				EffectivePan = InWaveInstance->VoiceCenterChannelVolume;
+			}
+			else if (Channel == PrevChannelInfo->Channel)
+			{
+				EffectivePan = !OmniAmount ? PrevChannelPan : FMath::Lerp(PrevChannelPan, OmniPanFactor, OmniAmount);
+			}
+			else if (Channel == NextChannelInfo->Channel)
+			{
+				EffectivePan = !OmniAmount ? NextChannelPan : FMath::Lerp(NextChannelPan, OmniPanFactor, OmniAmount);
+			}
+
+			AUDIO_MIXER_CHECK(EffectivePan >= 0.0f && EffectivePan <= 1.0f);
+			OutChannelMap.Add(EffectivePan);
+
+		}
+	}
+
+	void FMixerDevice::SetChannelAzimuth(EAudioMixerChannel::Type ChannelType, int32 Azimuth)
+	{
+		if (ChannelType >= EAudioMixerChannel::TopCenter)
+		{
+			UE_LOG(LogAudioMixer, Warning, TEXT("Unsupported mixer channel type: %s"), EAudioMixerChannel::ToString(ChannelType));
+			return;
+		}
+
+		if (Azimuth < 0 || Azimuth >= 360)
+		{
+			UE_LOG(LogAudioMixer, Warning, TEXT("Supplied azimuth is out of range: %d [0, 360)"), Azimuth);
+			return;
+		}
+
+		DefaultChannelAzimuthPosition[ChannelType].Azimuth = Azimuth;
+	}
+
+
+
+	int32 FMixerDevice::GetAzimuthForChannelType(EAudioMixerChannel::Type ChannelType)
+	{
+		if (ChannelType >= EAudioMixerChannel::TopCenter)
+		{
+			UE_LOG(LogAudioMixer, Warning, TEXT("Unsupported mixer channel type: %s"), EAudioMixerChannel::ToString(ChannelType));
+			return 0;
+		}
+
+		return DefaultChannelAzimuthPosition[ChannelType].Azimuth;
+	}
+
+	int32 FMixerDevice::GetDeviceSampleRate() const
+	{
+		return PlatformInfo.SampleRate;
+	}
+
+	int32 FMixerDevice::GetDeviceOutputChannels() const
+	{
+		return PlatformInfo.NumChannels;
+	}
+
+	FMixerSourceManager* FMixerDevice::GetSourceManager()
+	{
+		return &SourceManager;
+	}
+
+	void FMixerDevice::WhiteNoiseTest(TArray<float>& Output)
+	{
+		int32 NumFrames = PlatformInfo.NumFrames;
+		int32 NumChannels = PlatformInfo.NumChannels;
+
+		static FWhiteNoise WhiteNoise(0.2f);
+
+		for (int32 FrameIndex = 0; FrameIndex < NumFrames; ++FrameIndex)
+		{
+			for (int32 ChannelIndex = 0; ChannelIndex < NumChannels; ++ChannelIndex)
+			{
+				int32 Index = FrameIndex * NumChannels + ChannelIndex;
+				Output[Index] += WhiteNoise();
+			}
+		}
+	}
+
+	void FMixerDevice::SineOscTest(TArray<float>& Output)
+	{
+		int32 NumFrames = PlatformInfo.NumFrames;
+		int32 NumChannels = PlatformInfo.NumChannels;
+
+		check(NumChannels > 0);
+
+		static FSineOsc SineOscLeft(PlatformInfo.SampleRate, 440.0f, 0.2f);
+		static FSineOsc SineOscRight(PlatformInfo.SampleRate, 220.0f, 0.2f);
+
+		for (int32 FrameIndex = 0; FrameIndex < NumFrames; ++FrameIndex)
+		{
+			int32 Index = FrameIndex * NumChannels;
+
+			Output[Index] += SineOscLeft();
+
+			if (NumChannels > 1)
+			{
+				Output[Index + 1] += SineOscRight();
+			}
+		}
+	}
+
+}
\ No newline at end of file
diff --git a/Engine/Source/Runtime/AudioMixer/Private/AudioMixerModule.cpp b/Engine/Source/Runtime/AudioMixer/Private/AudioMixerModule.cpp
new file mode 100644
--- /dev/null
+++ b/Engine/Source/Runtime/AudioMixer/Private/AudioMixerModule.cpp
@@ -0,0 +1,19 @@
+// Copyright 1998-2016 Epic Games, Inc. All Rights Reserved.
+
+#include "AudioMixerPCH.h"
+#include "AudioMixerModule.h"
+#include "ModuleManager.h"
+
+DEFINE_LOG_CATEGORY(LogAudioMixer);
+DEFINE_LOG_CATEGORY(LogAudioMixerDebug);
+
+class FAudioMixerModule : public IModuleInterface
+{
+public:
+
+	virtual void StartupModule() override
+	{
+	}
+};
+
+IMPLEMENT_MODULE(FAudioMixerModule, AudioMixer);
\ No newline at end of file
diff --git a/Engine/Source/Runtime/AudioMixer/Private/AudioMixerPCH.h b/Engine/Source/Runtime/AudioMixer/Private/AudioMixerPCH.h
new file mode 100644
--- /dev/null
+++ b/Engine/Source/Runtime/AudioMixer/Private/AudioMixerPCH.h
@@ -0,0 +1,6 @@
+// Copyright 1998-2016 Epic Games, Inc. All Rights Reserved.
+
+#pragma once 
+
+#include "AudioMixer.h"
+
