Copying //UE4/Dev-Framework to //UE4/Dev-Main (Source: //UE4/Dev-Framework @ 3136612)

diff --git a/Engine/Source/Runtime/AudioMixer/Private/AudioMixerSource.cpp b/Engine/Source/Runtime/AudioMixer/Private/AudioMixerSource.cpp
new file mode 100644
--- /dev/null
+++ b/Engine/Source/Runtime/AudioMixer/Private/AudioMixerSource.cpp
@@ -0,0 +1,762 @@
+// Copyright 1998-2016 Epic Games, Inc. All Rights Reserved.
+
+#include "AudioMixerPCH.h"
+#include "AudioMixerSource.h"
+#include "AudioMixerDevice.h"
+#include "IAudioExtensionPlugin.h"
+
+namespace Audio
+{
+
+	FMixerSource::FMixerSource(FAudioDevice* InAudioDevice)
+		: FSoundSource(InAudioDevice)
+		, MixerDevice((FMixerDevice*)InAudioDevice)
+		, MixerBuffer(nullptr)
+		, MixerSourceVoice(nullptr)
+		, AsyncRealtimeAudioTask(nullptr)
+		, CurrentBuffer(0)
+		, PreviousAzimuth(-1.0f)
+		, bPlayedCachedBuffer(false)
+		, bPlaying(false)
+		, bLoopCallback(false)
+		, bIsFinished(false)
+		, bBuffersToFlush(false)
+		, bResourcesNeedFreeing(false)
+		, bEditorWarnedChangedSpatialization(false)
+		, bUsingHRTFSpatialization(false)
+		, bDebugMode(false)
+	{
+		// Create the source voice buffers
+		for (int32 i = 0; i < Audio::MAX_BUFFERS_QUEUED; ++i)
+		{
+			SourceVoiceBuffers.Add(FMixerSourceBufferPtr(new FMixerSourceVoiceBuffer()));
+		}
+	}
+
+	FMixerSource::~FMixerSource()
+	{
+		FreeResources();
+	}
+
+	bool FMixerSource::Init(FWaveInstance* InWaveInstance)
+	{
+		AUDIO_MIXER_CHECK(MixerBuffer);
+		AUDIO_MIXER_CHECK(MixerBuffer->IsRealTimeSourceReady());
+
+		// Get the number of frames before creating the buffer
+		int32 NumFrames = INDEX_NONE;
+
+		if (InWaveInstance->WaveData->DecompressionType != EDecompressionType::DTYPE_Procedural)
+		{
+			const int32 NumBytes = InWaveInstance->WaveData->RawPCMDataSize;
+			NumFrames = NumBytes / (InWaveInstance->WaveData->NumChannels * sizeof(int16));
+		}
+
+		FSoundBuffer* SoundBuffer = static_cast<FSoundBuffer*>(MixerBuffer);
+		if (SoundBuffer->NumChannels > 0)
+		{
+			SCOPE_CYCLE_COUNTER(STAT_AudioSourceInitTime);
+
+			AUDIO_MIXER_CHECK(MixerDevice);
+			MixerSourceVoice = MixerDevice->GetMixerSourceVoice(InWaveInstance, this, UseHRTSpatialization());
+			if (!MixerSourceVoice)
+			{
+				return false;
+			}
+
+			// Default the submix to use to use the master submix
+			FMixerSubmix* SourceSubmix = MixerDevice->GetMasterSubmix();
+
+			// If we're told to use a different submix, then grab that submix and set it for the sound
+			if (InWaveInstance->SoundSubmix)
+			{
+				FMixerSubmix** Submix = MixerDevice->GetSubmix(InWaveInstance->SoundSubmix);
+				check(Submix);
+				SourceSubmix = *Submix;
+			}
+
+			// Initialize the source voice with the necessary format information
+			FMixerSourceVoiceInitParams InitParams;
+			InitParams.BufferQueueListener = this;
+			InitParams.NumInputChannels = InWaveInstance->WaveData->NumChannels;
+			InitParams.NumInputFrames = NumFrames;
+			InitParams.OwningSubmix = SourceSubmix;
+			InitParams.SourceVoice = MixerSourceVoice;
+			InitParams.bUseHRTFSpatialization = UseHRTSpatialization();
+
+			// Check to see if this sound has been flagged to be in debug mode
+#if AUDIO_MIXER_ENABLE_DEBUG_MODE
+			InitParams.DebugName = InWaveInstance->GetName();
+
+			bool bIsDebug = false;
+			FString WaveInstanceName = WaveInstance->GetName();
+			FString TestName = GEngine->GetAudioDeviceManager()->GetAudioMixerDebugSoundName();
+			if (WaveInstanceName.Contains(TestName))
+			{
+				bDebugMode = true;
+				InitParams.bIsDebugMode = bDebugMode;
+			}
+#endif
+
+			if (MixerSourceVoice->Init(InitParams))
+			{
+				AUDIO_MIXER_CHECK(WaveInstance);
+				if (WaveInstance->StartTime)
+				{
+					MixerBuffer->Seek(WaveInstance->StartTime);
+				}
+
+				const EBufferType::Type BufferType = MixerBuffer->GetType();
+				switch (BufferType)
+				{
+					case EBufferType::PCM:
+					case EBufferType::PCMPreview:
+						SubmitPCMBuffers();
+						break;
+
+					case EBufferType::PCMRealTime:
+					case EBufferType::Streaming:
+						SubmitPCMRTBuffers();
+						break;
+
+					case EBufferType::Invalid:
+						break;
+				}
+
+				bInitialized = true;
+
+				ChannelMap.Reset();
+
+				Update();
+
+				return true;
+			}
+		}
+		return false;
+	}
+
+	void FMixerSource::Update()
+	{
+		SCOPE_CYCLE_COUNTER(STAT_AudioUpdateSources);
+
+		if (!WaveInstance || !MixerSourceVoice || Paused || !bInitialized)
+		{
+			return;
+		}
+
+		UpdatePitch();
+
+		UpdateVolume();
+
+		UpdateSpatialization();
+
+		UpdateEffects();
+
+		UpdateChannelMaps();
+
+		FSoundSource::DrawDebugInfo();
+	}
+
+	bool FMixerSource::PrepareForInitialization(FWaveInstance* InWaveInstance)
+	{
+		// We are currently not supporting playing audio on a controller
+		if (InWaveInstance->OutputTarget == EAudioOutputTarget::Controller)
+		{
+			return false;
+		}
+
+		// We are not initialized yet. We won't be until the sound file finishes loading and parsing the header.
+		bInitialized = false;
+
+		//  Reset so next instance will warn if algorithm changes in-flight
+		bEditorWarnedChangedSpatialization = false;
+
+		check(InWaveInstance);
+		check(MixerBuffer == nullptr);
+		check(AudioDevice);
+
+		MixerBuffer = FMixerBuffer::Init(AudioDevice, InWaveInstance->WaveData, InWaveInstance->StartTime > 0.0f);
+		if (MixerBuffer)
+		{
+			Buffer = MixerBuffer;
+			WaveInstance = InWaveInstance;
+
+			LPFFrequency = MAX_FILTER_FREQUENCY;
+			LastLPFFrequency = FLT_MAX;
+			bIsFinished = false;
+
+			// We succeeded in preparing the buffer for initialization, but we are not technically initialized yet.
+			// If the buffer is asynchronously preparing a file-handle, we may not yet initialize the source.
+			return true;
+		}
+
+		// Something went wrong with initializing the generator
+		return false;
+	}
+
+	bool FMixerSource::IsPreparedToInit()
+	{
+		return MixerBuffer && MixerBuffer->IsRealTimeSourceReady();
+	}
+
+	void FMixerSource::Play()
+	{
+		if (!WaveInstance)
+		{
+			return;
+		}
+
+		// It's possible if Pause and Play are called while a sound is async initializing. In this case
+		// we'll just not actually play the source here. Instead we'll call play when the sound finishes loading.
+		if (MixerSourceVoice && bInitialized)
+		{
+			MixerSourceVoice->Play();
+		}
+
+		Paused = false;
+		Playing = true;
+		bBuffersToFlush = false;
+		bLoopCallback = false;
+	}
+
+	void FMixerSource::Stop()
+	{
+		bInitialized = false;
+		if (WaveInstance)
+		{
+			if (MixerSourceVoice && Playing)
+			{
+				MixerSourceVoice->Stop();
+			}
+
+			Paused = false;
+			Playing = false;
+
+			FreeResources();
+		}
+
+		FSoundSource::Stop();
+	}
+
+	void FMixerSource::Pause()
+	{
+		if (!WaveInstance)
+		{
+			return;
+		}
+
+		if (MixerSourceVoice)
+		{
+			MixerSourceVoice->Pause();
+		}
+
+		Paused = true;
+	}
+
+	bool FMixerSource::IsFinished()
+	{
+		// A paused source is not finished.
+		if (Paused || !bInitialized)
+		{
+			return false;
+		}
+
+		if (WaveInstance && MixerSourceVoice)
+		{
+			if (bIsFinished)
+			{
+				WaveInstance->NotifyFinished();
+				return true;
+			}
+
+			if (bLoopCallback && WaveInstance->LoopingMode == LOOP_WithNotification)
+			{
+				WaveInstance->NotifyFinished();
+				bLoopCallback = false;
+			}
+
+			return false;
+		}
+		return true;
+	}
+
+	FString FMixerSource::Describe(bool bUseLongName)
+	{
+		return FString(TEXT("Stub"));
+	}
+
+	void FMixerSource::SubmitPCMBuffers()
+	{
+		CurrentBuffer = 0;
+
+		uint8* Data = nullptr;
+		uint32 DataSize = 0;
+		MixerBuffer->GetPCMData(&Data, &DataSize);
+
+		// Reset the data, copy it over
+		SourceVoiceBuffers[0]->AudioData.Reset();
+		SourceVoiceBuffers[0]->AudioData.AddDefaulted(DataSize);
+		FMemory::Memcpy(SourceVoiceBuffers[0]->AudioData.GetData(), Data, DataSize);
+
+		AUDIO_MIXER_CHECK(SourceVoiceBuffers[0]->AudioData.Num() == DataSize);
+
+		// Set the size of the data. In general AudioBytes may not be the same as AudioData.Num() 
+		SourceVoiceBuffers[0]->AudioBytes = DataSize;
+		SourceVoiceBuffers[0]->bRealTimeBuffer = false;
+		SourceVoiceBuffers[0]->LoopCount = (WaveInstance->LoopingMode != LOOP_Never) ? Audio::LOOP_FOREVER : 0;
+
+		MixerSourceVoice->SubmitBuffer(SourceVoiceBuffers[0]);
+	}
+
+	void FMixerSource::SubmitPCMRTBuffers()
+	{
+		CurrentBuffer = 0;
+
+		const uint32 BufferSize = MONO_PCM_BUFFER_SIZE * MixerBuffer->NumChannels;
+
+		// Set up buffer areas to decompress audio into
+		for (int32 i = 0; i < Audio::MAX_BUFFERS_QUEUED; ++i)
+		{
+			SourceVoiceBuffers[i]->AudioData.Reset();
+			SourceVoiceBuffers[i]->AudioData.AddZeroed(BufferSize);
+			SourceVoiceBuffers[i]->AudioBytes = BufferSize;
+			SourceVoiceBuffers[i]->bRealTimeBuffer = true;
+			SourceVoiceBuffers[i]->LoopCount = 0;
+		}
+
+		bPlayedCachedBuffer = false;
+		bool bIsSeeking = (WaveInstance->StartTime > 0.0f);
+		if (!bIsSeeking && WaveInstance->WaveData && WaveInstance->WaveData->CachedRealtimeFirstBuffer)
+		{
+			bPlayedCachedBuffer = true;
+			FMemory::Memcpy(SourceVoiceBuffers[0]->AudioData.GetData(), WaveInstance->WaveData->CachedRealtimeFirstBuffer, BufferSize);
+			FMemory::Memcpy(SourceVoiceBuffers[1]->AudioData.GetData(), WaveInstance->WaveData->CachedRealtimeFirstBuffer + BufferSize, BufferSize);
+		}
+		else
+		{
+			// If we didn't cache or if we're seeking, then immediately read the next 2 buffers synchronously
+			ReadMorePCMRTData(0, EBufferReadMode::Synchronous);
+			ReadMorePCMRTData(1, EBufferReadMode::Synchronous);
+		}
+
+		MixerSourceVoice->SubmitBuffer(SourceVoiceBuffers[0]);
+		MixerSourceVoice->SubmitBuffer(SourceVoiceBuffers[1]);
+
+		// Set the current buffer index to 2 after
+		CurrentBuffer = 2;
+
+		bResourcesNeedFreeing = true;
+	}
+
+	bool FMixerSource::ReadMorePCMRTData(const int32 BufferIndex, EBufferReadMode BufferReadMode, bool* OutLooped)
+	{
+		USoundWave* WaveData = WaveInstance->WaveData;
+
+		if (WaveData && WaveData->bProcedural)
+		{
+			AUDIO_MIXER_CHECK_AUDIO_PLAT_THREAD(MixerDevice);
+			const int32 MaxSamples = (MONO_PCM_BUFFER_SIZE * Buffer->NumChannels) / sizeof(int16);
+
+			if (BufferReadMode == EBufferReadMode::Synchronous || WaveData->bCanProcessAsync == false)
+			{
+				const int32 BytesWritten = WaveData->GeneratePCMData(SourceVoiceBuffers[BufferIndex]->AudioData.GetData(), MaxSamples);
+				SourceVoiceBuffers[BufferIndex]->AudioBytes = BytesWritten;
+			}
+			else
+			{
+				check(!AsyncRealtimeAudioTask);
+				AsyncRealtimeAudioTask = new FAsyncRealtimeAudioTask(WaveData, SourceVoiceBuffers[BufferIndex]->AudioData.GetData(), MaxSamples);
+				AsyncRealtimeAudioTask->StartBackgroundTask();
+			}
+
+			// Not looping
+			return false;
+		}
+		else if (BufferReadMode == EBufferReadMode::Synchronous)
+		{
+			return MixerBuffer->ReadCompressedData(SourceVoiceBuffers[BufferIndex]->AudioData.GetData(), WaveInstance->LoopingMode != LOOP_Never);
+		}
+		else
+		{
+			AUDIO_MIXER_CHECK_AUDIO_PLAT_THREAD(MixerDevice);
+			check(!AsyncRealtimeAudioTask);
+			AsyncRealtimeAudioTask = new FAsyncRealtimeAudioTask(MixerBuffer, SourceVoiceBuffers[BufferIndex]->AudioData.GetData(), WaveInstance->LoopingMode != LOOP_Never, BufferReadMode == EBufferReadMode::AsynchronousSkipFirstFrame);
+			AsyncRealtimeAudioTask->StartBackgroundTask();
+
+			// Not looping
+			return false;
+		}
+	}
+
+	void FMixerSource::SubmitRealTimeSourceData(bool bLooped)
+	{
+		AUDIO_MIXER_CHECK_AUDIO_PLAT_THREAD(MixerDevice);
+
+		// Have we reached the end of the sound
+		if (bLooped)
+		{
+			switch (WaveInstance->LoopingMode)
+			{
+				case LOOP_Never:
+					// Play out any queued buffers - once there are no buffers left, the state check at the beginning of IsFinished will fire
+					bBuffersToFlush = true;
+					break;
+
+				case LOOP_WithNotification:
+					// If we have just looped, and we are looping, send notification
+					// This will trigger a WaveInstance->NotifyFinished() in the FXAudio2SoundSournce::IsFinished() function on main thread.
+					bLoopCallback = true;
+					break;
+
+				case LOOP_Forever:
+					// Let the sound loop indefinitely
+					break;
+			}
+		}
+
+		if (SourceVoiceBuffers[CurrentBuffer]->AudioData.Num() > 0)
+		{
+			MixerSourceVoice->SubmitBufferAudioThread(SourceVoiceBuffers[CurrentBuffer]);
+		}
+	}
+
+	void FMixerSource::ProcessRealTimeSource(bool bBlockForData)
+	{
+		const bool bGetMoreData = bBlockForData || (AsyncRealtimeAudioTask == nullptr);
+		if (AsyncRealtimeAudioTask)
+		{
+			const bool bTaskDone = AsyncRealtimeAudioTask->IsDone();
+			if (bTaskDone || bBlockForData)
+			{
+				bool bLooped = false;
+				if (!bTaskDone)
+				{
+					AsyncRealtimeAudioTask->EnsureCompletion();
+				}
+
+				switch (AsyncRealtimeAudioTask->GetTask().GetTaskType())
+				{
+					case ERealtimeAudioTaskType::Decompress:
+						bLooped = AsyncRealtimeAudioTask->GetTask().GetBufferLooped();
+						break;
+
+					case ERealtimeAudioTaskType::Procedural:
+						SourceVoiceBuffers[CurrentBuffer]->AudioBytes = AsyncRealtimeAudioTask->GetTask().GetBytesWritten();
+						break;
+				}
+
+				delete AsyncRealtimeAudioTask;
+				AsyncRealtimeAudioTask = nullptr;
+				
+				SubmitRealTimeSourceData(bLooped);
+			}
+		}
+
+		if (bGetMoreData)
+		{
+			// Update the buffer index
+			if (++CurrentBuffer > 2)
+			{
+				CurrentBuffer = 0;
+			}
+
+			EBufferReadMode DataReadMode;
+			if (bPlayedCachedBuffer)
+			{
+				bPlayedCachedBuffer = false;
+				DataReadMode = EBufferReadMode::AsynchronousSkipFirstFrame;
+			}
+			else
+			{
+				DataReadMode = (MixerBuffer->GetType() == EBufferType::Streaming ? EBufferReadMode::Synchronous : EBufferReadMode::Asynchronous);
+			}
+			const bool bLooped = ReadMorePCMRTData(CurrentBuffer, DataReadMode);
+
+			// If this was a synchronous read, then immediately write it
+			if (AsyncRealtimeAudioTask == nullptr)
+			{
+				SubmitRealTimeSourceData(bLooped);
+			}
+		}
+	}
+
+	void FMixerSource::OnSourceBufferEnd()
+	{
+		AUDIO_MIXER_CHECK_AUDIO_PLAT_THREAD(MixerDevice);
+
+		if (Playing && MixerSourceVoice)
+		{
+			int32 BuffersQueued = MixerSourceVoice->GetNumBuffersQueued();
+
+			const bool bIsRealTimeSource = MixerBuffer->IsRealTimeBuffer();
+
+			if (BuffersQueued == 0 && (bBuffersToFlush || !bIsRealTimeSource))
+			{
+				bIsFinished = true;
+			}
+			else if (bIsRealTimeSource && !bBuffersToFlush && BuffersQueued <= (Audio::MAX_BUFFERS_QUEUED - 1))
+			{
+				ProcessRealTimeSource(BuffersQueued < (Audio::MAX_BUFFERS_QUEUED - 1));
+			}
+		}
+	}
+
+	void FMixerSource::FreeResources()
+	{
+		if (AsyncRealtimeAudioTask)
+		{
+			AsyncRealtimeAudioTask->EnsureCompletion();
+			delete AsyncRealtimeAudioTask;
+			AsyncRealtimeAudioTask = nullptr;
+		}
+
+		if (MixerBuffer)
+		{
+			MixerBuffer->EnsureHeaderParseTaskFinished();
+		}
+
+		if (MixerSourceVoice)
+		{
+			MixerSourceVoice->Release();
+			MixerSourceVoice = nullptr;
+		}
+
+		if (bResourcesNeedFreeing)
+		{
+			if (Buffer)
+			{
+				check(Buffer->ResourceID == 0);
+				delete Buffer;
+			}
+
+			CurrentBuffer = 0;
+		}
+
+		MixerBuffer = nullptr;
+		Buffer = nullptr;
+		bBuffersToFlush = false;
+		bLoopCallback = false;
+		bResourcesNeedFreeing = false;
+	}
+
+	void FMixerSource::UpdatePitch()
+	{
+		float Pitch = WaveInstance->Pitch;
+
+		if (!WaveInstance->bIsUISound)
+		{	
+			Pitch *= AudioDevice->GetGlobalPitchScale().GetValue();
+		}
+
+		// Scale in the sound sample rate divided by device sample rate so pitch is 
+		// accurate independent of sound source sample rate or device sample rate
+		AUDIO_MIXER_CHECK(MixerBuffer);
+		Pitch *= (MixerBuffer->GetSampleRate() / AudioDevice->GetSampleRate());
+
+		Pitch = FMath::Clamp<float>(Pitch, AUDIO_MIXER_MIN_PITCH, AUDIO_MIXER_MAX_PITCH);
+
+		MixerSourceVoice->SetPitch(Pitch);
+	}
+
+	void FMixerSource::UpdateVolume()
+	{
+		float CurrentVolume = 0.0f;
+
+		if (!AudioDevice->IsAudioDeviceMuted())
+		{
+			CurrentVolume = WaveInstance->GetActualVolume();
+		}
+
+		CurrentVolume = FMath::Clamp<float>(CurrentVolume * AudioDevice->GetPlatformAudioHeadroom(), 0.0f, MAX_VOLUME);
+
+		MixerSourceVoice->SetVolume(CurrentVolume);
+	}
+
+	void FMixerSource::UpdateSpatialization()
+	{
+		SpatializationParams = GetSpatializationParams();
+		if (WaveInstance->bUseSpatialization)
+		{
+			MixerSourceVoice->SetSpatializationParams(SpatializationParams);
+		}
+	}
+
+	void FMixerSource::UpdateEffects()
+	{
+		// Update the default LPF filter frequency 
+		SetFilterFrequency();
+
+		if (LastLPFFrequency != LPFFrequency)
+		{
+			MixerSourceVoice->SetLPFFrequency(LPFFrequency);
+			LastLPFFrequency = LPFFrequency;
+		}
+
+		// TODO: Update the wet level of the sound based on game data
+		MixerSourceVoice->SetWetLevel(0.5f);
+	}
+
+	void FMixerSource::UpdateChannelMaps()
+	{
+		SetStereoBleed();
+
+		SetLFEBleed();
+
+		bool bChanged = false;
+
+		check(Buffer);
+		if (Buffer->NumChannels == 1)
+		{
+			bChanged = ComputeMonoChannelMap();
+		}
+		else if (Buffer->NumChannels == 2)
+		{
+			bChanged = ComputeStereoChannelMap();
+		}
+		else if (Buffer->NumChannels == 4)
+		{
+			bChanged = ComputeQuadChannelMap();
+		}
+		else if (Buffer->NumChannels == 6)
+		{
+			bChanged = ComputeHexChannelMap();
+		}
+		else
+		{
+			UE_LOG(LogAudioMixer, Warning, TEXT("Unsupported input audio channels '%d'"), Buffer->NumChannels);
+			return;
+		}
+
+		if (bChanged)
+		{
+			MixerSourceVoice->SetChannelMap(ChannelMap);
+		}
+	}
+
+	bool FMixerSource::ComputeMonoChannelMap()
+	{
+		if (UseHRTSpatialization())
+		{
+			if (WaveInstance->SpatializationAlgorithm != SPATIALIZATION_HRTF && !bEditorWarnedChangedSpatialization)
+			{
+				bEditorWarnedChangedSpatialization = true;
+				UE_LOG(LogAudioMixer, Warning, TEXT("Changing the spatialization algorithm on a playing sound is not supported (WaveInstance: %s)"), *WaveInstance->WaveData->GetFullName());
+			}
+
+			// Treat the source as if it is a 2D stereo source
+			return ComputeStereoChannelMap();
+		}
+		else if (WaveInstance->bUseSpatialization && !FMath::IsNearlyEqual(WaveInstance->AbsoluteAzimuth, PreviousAzimuth, 0.01f))
+		{
+			// Don't need to compute the source channel map if the absolute azimuth hasn't changed much
+			PreviousAzimuth = WaveInstance->AbsoluteAzimuth;
+			ChannelMap.Reset();
+			MixerDevice->Get3DChannelMap(WaveInstance, WaveInstance->AbsoluteAzimuth, SpatializationParams.NormalizedOmniRadius, ChannelMap);
+			return true;
+		}
+		else if (!ChannelMap.Num())
+		{
+			// Only need to compute the 2D channel map once
+			MixerDevice->Get2DChannelMap(1, ChannelMap);
+			return true;
+		}
+
+		// Return false means the channel map hasn't changed
+		return false;
+	}
+
+	bool FMixerSource::ComputeStereoChannelMap()
+	{
+		if (!UseHRTSpatialization() && WaveInstance->bUseSpatialization)
+		{
+			// Make sure our stereo emitter positions are updated relative to the sound emitter position
+			UpdateStereoEmitterPositions();
+
+			float AzimuthOffset = 0.0f;
+			if (WaveInstance->AttenuationDistance > 0.0f)
+			{
+				AzimuthOffset = FMath::Atan(0.5f * WaveInstance->StereoSpread / WaveInstance->AttenuationDistance);
+				AzimuthOffset = FMath::RadiansToDegrees(AzimuthOffset);
+			}
+
+			float LeftAzimuth = WaveInstance->AbsoluteAzimuth - AzimuthOffset;
+			if (LeftAzimuth < 0.0f)
+			{
+				LeftAzimuth += 360.0f;
+			}
+
+			float RightAzimuth = WaveInstance->AbsoluteAzimuth + AzimuthOffset;
+			if (RightAzimuth > 360.0f)
+			{
+				RightAzimuth -= 360.0f;
+			}
+
+			// Reset the channel map, the stereo spatialization channel mapping calls below will append their mappings
+			StereoChannelMap.Reset();
+
+			MixerDevice->Get3DChannelMap(WaveInstance, LeftAzimuth, SpatializationParams.NormalizedOmniRadius, StereoChannelMap);
+			MixerDevice->Get3DChannelMap(WaveInstance, RightAzimuth, SpatializationParams.NormalizedOmniRadius, StereoChannelMap);
+
+			ChannelMap.Reset();
+
+			int32 NumDeviceChannels = MixerDevice->GetNumDeviceChannels();
+			check(StereoChannelMap.Num() == 2 * NumDeviceChannels);
+
+			int32 LeftIndex = 0;
+			int32 RightIndex = NumDeviceChannels;
+
+			for (int32 i = 0; i < StereoChannelMap.Num(); ++i)
+			{
+				// Interleave the left and right channel maps to a single stereo source -> output channel map
+				if (i % 2 == 0)
+				{
+					ChannelMap.Add(StereoChannelMap[LeftIndex++]);
+				}
+				else
+				{
+					ChannelMap.Add(StereoChannelMap[RightIndex++]);
+				}
+			}
+
+			return true;
+		}
+		else if (!ChannelMap.Num())
+		{
+			MixerDevice->Get2DChannelMap(2, ChannelMap);
+			return true;
+		}
+
+		return false;
+	}
+
+	bool FMixerSource::ComputeQuadChannelMap()
+	{
+		if (!ChannelMap.Num())
+		{
+			MixerDevice->Get2DChannelMap(4, ChannelMap);
+			return true;
+		}
+		return false;
+	}
+
+	bool FMixerSource::ComputeHexChannelMap()
+	{
+		if (!ChannelMap.Num())
+		{
+			MixerDevice->Get2DChannelMap(6, ChannelMap);
+			return true;
+		}
+		return false;
+	}
+
+	bool FMixerSource::UseHRTSpatialization() const
+	{
+		return (Buffer->NumChannels == 1 &&
+				AudioDevice->SpatializationPlugin != nullptr &&
+				AudioDevice->IsSpatializationPluginEnabled() &&
+				WaveInstance->SpatializationAlgorithm != SPATIALIZATION_Default);
+	}
+
+
+
+}
\ No newline at end of file
diff --git a/Engine/Source/Runtime/AudioMixer/Private/AudioMixerSource.h b/Engine/Source/Runtime/AudioMixer/Private/AudioMixerSource.h
new file mode 100644
--- /dev/null
+++ b/Engine/Source/Runtime/AudioMixer/Private/AudioMixerSource.h
@@ -0,0 +1,177 @@
+// Copyright 1998-2016 Epic Games, Inc. All Rights Reserved.
+
+#pragma once
+
+/* Public dependencies
+*****************************************************************************/
+
+#include "Core.h"
+#include "Engine.h"
+
+#include "SoundDefinitions.h"
+
+#include "AudioMixerDevice.h"
+#include "AudioMixerBuffer.h"
+#include "AudioMixerSourceVoice.h"
+#include "AudioMixerSourceManager.h"
+
+namespace Audio
+{
+	class FMixerDevice;
+	class FMixerSourceVoice;
+	class FMixerSource;
+	class FMixerBuffer;
+	class ISourceBufferQueueListener;
+
+	static const int32 MAX_BUFFERS_QUEUED = 3;
+	static const int32 LOOP_FOREVER = -1;
+
+
+	/** Struct defining a source voice buffer. */
+	struct FMixerSourceVoiceBuffer
+	{
+		/** Raw audio PCM data. */
+		TArray<uint8> AudioData;
+
+		/** The amount of real audio data in the byte array (may not be the size of the array in the case of procedural audio. */
+		uint32 AudioBytes;
+
+		/** How many times this buffer will loop. */
+		int32 LoopCount;
+
+		/** If this buffer is from real-time decoding and needs to make callbacks for more data. */
+		uint32 bRealTimeBuffer : 1;
+
+		FMixerSourceVoiceBuffer()
+		{
+			FMemory::Memzero(this, sizeof(this));
+		}
+	};
+
+	typedef FAsyncRealtimeAudioTaskProxy<FMixerBuffer> FAsyncRealtimeAudioTask;
+
+	/** 
+	 * FMixerSource
+	 * Class which implements a sound source object for the audio mixer module.
+	 */
+	class FMixerSource :	public FSoundSource, 
+							public ISourceBufferQueueListener
+	{
+	public:
+		/** Constructor. */
+		FMixerSource(FAudioDevice* InAudioDevice);
+
+		/** Destructor. */
+		~FMixerSource();
+
+		//~ Begin FSoundSource Interface
+		bool Init(FWaveInstance* InWaveInstance) override;
+		void Update() override;
+		bool PrepareForInitialization(FWaveInstance* InWaveInstance) override;
+		bool IsPreparedToInit() override;
+		void Play() override;
+		void Stop() override;
+		void Pause() override;
+		bool IsFinished() override;
+		FString Describe(bool bUseLongName) override;
+		//~ End FSoundSource Interface
+
+		//~Begin ISourceBufferQueueListener
+		void OnSourceBufferEnd() override;
+		//~End ISourceBufferQueueListener
+
+	private:
+
+		/** Enum describing the data-read mode of an audio buffer. */
+		enum class EBufferReadMode : uint8
+		{
+			/** Read the next buffer synchronously. */
+			Synchronous,
+
+			/** Read the next buffer asynchronously. */
+			Asynchronous,
+
+			/** Read the next buffer asynchronously but skip the first chunk of audio. */
+			AsynchronousSkipFirstFrame
+		};
+
+		/** Submit the current decoded PCM buffer to the contained source voice. */
+		void SubmitPCMBuffers();
+
+		/** Submit the current decoded PCMRT (PCM RealTime) buffer to the contained source voice. */
+		void SubmitPCMRTBuffers();
+
+		/** Reads more PCMRT data for the real-time decoded audio buffers. */
+		bool ReadMorePCMRTData(const int32 BufferIndex, EBufferReadMode BufferReadMode, bool* OutLooped = nullptr);
+
+		/** Called when a buffer finishes for a real-time source and more buffers need to be read and submitted. */
+		void ProcessRealTimeSource(bool bBlockForData);
+
+		/** Submits new real-time decoded buffers to a source voice. */
+		void SubmitRealTimeSourceData(bool bLooped);
+
+		/** Frees any resources for this sound source. */
+		void FreeResources();
+
+		/** Updates the pitch parameter set from the game thread. */
+		void UpdatePitch();
+		
+		/** Updates the volume parameter set from the game thread. */
+		void UpdateVolume();
+
+		/** Gets updated spatialization information for the voice. */
+		void UpdateSpatialization();
+
+		/** Updates and source effect on this voice. */
+		void UpdateEffects();
+
+		/** Updates the channel map of the sound if its a 3d sound.*/
+		void UpdateChannelMaps();
+
+		/** Computes the mono-channel map. */
+		bool ComputeMonoChannelMap();
+
+		/** Computes the stereo-channel map. */
+		bool ComputeStereoChannelMap();
+
+		/** Computes the quad channel map. */
+		bool ComputeQuadChannelMap();
+
+		/** Computes the hex channel map. */
+		bool ComputeHexChannelMap();
+
+		/** Whether or not we should create the source voice with the HRTF spatializer. */
+		bool UseHRTSpatialization() const;
+
+	private:
+
+		FMixerDevice* MixerDevice;
+		FMixerBuffer* MixerBuffer;
+		FMixerSourceVoice* MixerSourceVoice;
+		FAsyncRealtimeAudioTask* AsyncRealtimeAudioTask;
+
+		TArray<float> ChannelMap;
+		TArray<float> StereoChannelMap;
+
+		int32 CurrentBuffer;
+		float PreviousAzimuth;
+
+		// The decoded source buffers are using a shared pointer because the audio mixer thread 
+		// will need to have a ref while playing back the sound. There is a small window where a 
+		// flushed/destroyed sound can destroy the buffer before the sound finishes using the buffer.
+		TArray<FMixerSourceBufferPtr> SourceVoiceBuffers;
+
+		FSpatializationParams SpatializationParams;
+
+		FThreadSafeBool bPlayedCachedBuffer;
+		FThreadSafeBool bPlaying;
+		FThreadSafeBool bLoopCallback;
+		FThreadSafeBool bIsFinished;
+		FThreadSafeBool bBuffersToFlush;
+
+		uint32 bResourcesNeedFreeing : 1;
+		uint32 bEditorWarnedChangedSpatialization : 1;
+		uint32 bUsingHRTFSpatialization : 1;
+		uint32 bDebugMode : 1;
+	};
+}
diff --git a/Engine/Source/Runtime/AudioMixer/Private/AudioMixerSourceManager.cpp b/Engine/Source/Runtime/AudioMixer/Private/AudioMixerSourceManager.cpp
new file mode 100644
--- /dev/null
+++ b/Engine/Source/Runtime/AudioMixer/Private/AudioMixerSourceManager.cpp
@@ -0,0 +1,940 @@
+// Copyright 1998-2016 Epic Games, Inc. All Rights Reserved.
+
+#include "AudioMixerPCH.h"
+#include "AudioMixerSource.h"
+#include "AudioMixerDevice.h"
+#include "AudioMixerSourceManager.h"
+#include "IAudioExtensionPlugin.h"
+
+#define ONEOVERSHORTMAX (3.0517578125e-5f) // 1/32768
+
+#define VALIDATE_SOURCE_MIXER_STATE 1
+
+namespace Audio
+{
+	FSourceParam::FSourceParam(float InNumInterpFrames)
+		: StartValue(0.0f)
+		, EndValue(0.0f)
+		, CurrentValue(0.0f)
+		, NumInterpFrames(InNumInterpFrames)
+		, Frame(0.0f)
+		, bIsInit(true)
+	{
+	}
+
+	void FSourceParam::Reset()
+	{
+		bIsInit = true;
+	}
+
+	void FSourceParam::SetValue(float InValue)
+	{
+		if (bIsInit)
+		{
+			bIsInit = false;
+			CurrentValue = InValue;
+			StartValue = InValue;
+			EndValue = InValue;
+			Frame = NumInterpFrames;
+		}
+		else
+		{
+			StartValue = CurrentValue;
+			EndValue = InValue;
+			Frame = 0.0f;
+		}
+	}
+
+	float FSourceParam::operator()()
+	{
+		float Alpha = Frame / NumInterpFrames;
+		if (Alpha >= 1.0f)
+		{
+			CurrentValue = EndValue;
+			return EndValue;
+		}
+		else
+		{
+			CurrentValue = FMath::Lerp(StartValue, EndValue, Alpha);
+		}
+
+		Frame += 1.0f;
+		return CurrentValue;
+	}
+
+
+	/*************************************************************************
+	* FMixerSourceManager
+	**************************************************************************/
+
+	FMixerSourceManager::FMixerSourceManager(FMixerDevice* InMixerDevice)
+		: MixerDevice(InMixerDevice)
+		, NumActiveSources(0)
+		, NumTotalSources(0)
+		, bInitialized(false)
+	{
+	}
+
+	FMixerSourceManager::~FMixerSourceManager()
+	{
+
+	}
+
+	void FMixerSourceManager::Init(const int32 InNumSources)
+	{
+		AUDIO_MIXER_CHECK(MixerDevice);
+		AUDIO_MIXER_CHECK(MixerDevice->GetSampleRate() > 0);
+		AUDIO_MIXER_CHECK(InNumSources > 0);
+
+		if (bInitialized)
+		{
+			return;
+		}
+
+#if ENABLE_AUDIO_OUTPUT_DEBUGGING
+		for (int32 i = 0; i < InNumSources; ++i)
+		{
+			DebugOutputGenerators.Add(FSineOsc(44100, 50 + i * 5, 0.5f));
+		}
+#endif
+
+		MixerSources.Init(nullptr, InNumSources);
+		BufferQueue.AddDefaulted(InNumSources);
+		BufferQueueListener.Init(nullptr, InNumSources);
+		NumBuffersQueued.AddDefaulted(InNumSources);
+		CurrentPCMBuffer.Init(nullptr, InNumSources);
+		CurrentAudioChunkNumFrames.AddDefaulted(InNumSources);
+		SourceBuffer.AddDefaulted(InNumSources);
+		CurrentFrameValues.AddDefaulted(InNumSources);
+		NextFrameValues.AddDefaulted(InNumSources);
+		CurrentFrameAlpha.AddDefaulted(InNumSources);
+		CurrentFrameIndex.AddDefaulted(InNumSources);
+		NumFramesPlayed.AddDefaulted(InNumSources);
+
+		float InterpFrames = MixerDevice->GetSampleRate() * 0.033f;
+		PitchSourceParam.Init(FSourceParam(InterpFrames), InNumSources);
+		VolumeSourceParam.Init(FSourceParam(InterpFrames), InNumSources);
+		WetLevelSourceParam.Init(FSourceParam(InterpFrames), InNumSources);
+		LPFCutoffFrequencyParam.Init(FSourceParam(InterpFrames), InNumSources);
+		ChannelMapParam.Init(FSourceChannelMap(InterpFrames), InNumSources);
+		SpatParams.AddDefaulted(InNumSources);
+		LowPassFilters.AddDefaulted(InNumSources);
+
+		PostEffectBuffers.AddDefaulted(InNumSources);
+		DryBuffer.AddDefaulted(InNumSources);
+		WetBuffer.AddDefaulted(InNumSources);
+		bIsActive.AddDefaulted(InNumSources);
+		bIsPlaying.AddDefaulted(InNumSources);
+		bIsPaused.AddDefaulted(InNumSources);
+		bIsDone.AddDefaulted(InNumSources);
+		bIsBusy.AddDefaulted(InNumSources);
+		bUseHRTFSpatializer.AddDefaulted(InNumSources);
+		bHasStarted.AddDefaulted(InNumSources);
+		bIsDebugMode.AddDefaulted(InNumSources);
+		DebugName.AddDefaulted(InNumSources);
+
+		NumInputChannels.AddDefaulted(InNumSources);
+		NumPostEffectChannels.AddDefaulted(InNumSources);
+		NumInputFrames.AddDefaulted(InNumSources);
+
+		GameThreadInfo.bIsBusy.AddDefaulted(InNumSources);
+		GameThreadInfo.bIsDone.AddDefaulted(InNumSources);
+		GameThreadInfo.bIsDebugMode.AddDefaulted(InNumSources);
+		GameThreadInfo.FreeSourceIndices.Reset(InNumSources);
+		for (int32 i = InNumSources - 1; i >= 0; --i)
+		{
+			GameThreadInfo.FreeSourceIndices.Add(i);
+		}
+
+		NumTotalSources = InNumSources;
+		bInitialized = true;
+		bPumpQueue = false;
+	}
+
+	void FMixerSourceManager::Update()
+	{
+		AUDIO_MIXER_CHECK_GAME_THREAD(MixerDevice);
+
+#if VALIDATE_SOURCE_MIXER_STATE
+		for (int32 i = 0; i < NumTotalSources; ++i)
+		{
+			if (!GameThreadInfo.bIsBusy[i])
+			{
+				// Make sure that our bIsFree and FreeSourceIndices are correct
+				AUDIO_MIXER_CHECK(GameThreadInfo.FreeSourceIndices.Contains(i) == true);
+			}
+		}
+#endif
+
+		// Pump the source command queue from the game thread to make sure 
+		// playsound calls, param updates, etc, all happen simultaneously
+		bPumpQueue = true;
+	}
+
+	void FMixerSourceManager::ReleaseSource(const int32 SourceId)
+	{
+		AUDIO_MIXER_CHECK_AUDIO_PLAT_THREAD(MixerDevice);
+
+		AUDIO_MIXER_CHECK(SourceId < NumTotalSources);
+		AUDIO_MIXER_CHECK(bInitialized);
+		AUDIO_MIXER_CHECK(MixerSources[SourceId] != nullptr);
+
+		AUDIO_MIXER_DEBUG_LOG(SourceId, TEXT("Is releasing"));
+
+#if AUDIO_MIXER_ENABLE_DEBUG_MODE
+		if (bIsDebugMode[SourceId])
+		{
+			DebugSoloSources.Remove(SourceId);
+		}
+#endif
+
+		// Remove the mixer source from its owning submix
+		MixerSources[SourceId]->GetOwningSubmix()->RemoveSourceVoice(MixerSources[SourceId]);
+
+		// Delete the mixer source and null the slot
+		delete MixerSources[SourceId];
+		MixerSources[SourceId] = nullptr;
+
+		// Reset all state and data
+		PitchSourceParam[SourceId].Reset();
+		VolumeSourceParam[SourceId].Reset();
+		WetLevelSourceParam[SourceId].Reset();
+		LPFCutoffFrequencyParam[SourceId].Reset();
+
+		LowPassFilters[SourceId].Reset();
+		ChannelMapParam[SourceId].Reset();
+		BufferQueue[SourceId].Empty();
+		BufferQueueListener[SourceId] = nullptr;
+		CurrentPCMBuffer[SourceId] = nullptr;
+		CurrentAudioChunkNumFrames[SourceId] = 0;
+		SourceBuffer[SourceId].Reset();
+		CurrentFrameValues[SourceId].Reset();
+		NextFrameValues[SourceId].Reset();
+		CurrentFrameAlpha[SourceId] = 0.0f;
+		CurrentFrameIndex[SourceId] = 0;
+		NumFramesPlayed[SourceId] = 0;
+		PostEffectBuffers[SourceId].Reset();
+		DryBuffer[SourceId].Reset();
+		WetBuffer[SourceId].Reset();
+		bIsActive[SourceId] = false;
+		bIsPlaying[SourceId] = false;
+		bIsDone[SourceId] = true;
+		bIsPaused[SourceId] = false;
+		bIsBusy[SourceId] = false;
+		bUseHRTFSpatializer[SourceId] = false;
+		bIsDone[SourceId] = false;
+		bHasStarted[SourceId] = false;
+		bIsDebugMode[SourceId] = false;
+		DebugName[SourceId] = FString();
+		NumInputChannels[SourceId] = 0;
+		NumPostEffectChannels[SourceId] = 0;
+	}
+
+	bool FMixerSourceManager::GetFreeSourceId(int32& OutSourceId)
+	{
+		AUDIO_MIXER_CHECK_GAME_THREAD(MixerDevice);
+
+		if (GameThreadInfo.FreeSourceIndices.Num())
+		{
+			OutSourceId = GameThreadInfo.FreeSourceIndices.Pop();
+
+			AUDIO_MIXER_CHECK(OutSourceId < NumTotalSources);
+			AUDIO_MIXER_CHECK(!GameThreadInfo.bIsBusy[OutSourceId]);
+
+			AUDIO_MIXER_CHECK(!GameThreadInfo.bIsDebugMode[OutSourceId]);
+			AUDIO_MIXER_CHECK(NumActiveSources < NumTotalSources);
+			++NumActiveSources;
+
+			GameThreadInfo.bIsBusy[OutSourceId] = true;
+			return true;
+		}
+		AUDIO_MIXER_CHECK(false);
+		return false;
+	}
+
+	int32 FMixerSourceManager::GetNumActiveSources() const
+	{
+		return NumActiveSources;
+	}
+
+	void FMixerSourceManager::InitSource(const int32 SourceId, const FMixerSourceVoiceInitParams& InitParams)
+	{
+		AUDIO_MIXER_CHECK(SourceId < NumTotalSources);
+		AUDIO_MIXER_CHECK(GameThreadInfo.bIsBusy[SourceId]);
+		AUDIO_MIXER_CHECK(!GameThreadInfo.bIsDebugMode[SourceId]);
+		AUDIO_MIXER_CHECK(InitParams.BufferQueueListener != nullptr);
+		AUDIO_MIXER_CHECK_GAME_THREAD(MixerDevice);
+
+#if AUDIO_MIXER_ENABLE_DEBUG_MODE
+		GameThreadInfo.bIsDebugMode[SourceId] = InitParams.bIsDebugMode;
+#endif 
+
+		AudioMixerThreadCommand([this, SourceId, InitParams]()
+		{
+			AUDIO_MIXER_CHECK_AUDIO_PLAT_THREAD(MixerDevice);
+			AUDIO_MIXER_CHECK(InitParams.OwningSubmix != nullptr);
+			AUDIO_MIXER_CHECK(InitParams.SourceVoice != nullptr);
+
+			bIsPlaying[SourceId] = false;
+			bIsPaused[SourceId] = false;
+			bIsActive[SourceId] = true;
+			bIsBusy[SourceId] = true;
+			bUseHRTFSpatializer[SourceId] = InitParams.bUseHRTFSpatialization;
+
+			BufferQueueListener[SourceId] = InitParams.BufferQueueListener;
+			NumInputChannels[SourceId] = InitParams.NumInputChannels;
+			NumInputFrames[SourceId] = InitParams.NumInputFrames;
+
+			AUDIO_MIXER_CHECK(BufferQueue[SourceId].IsEmpty());
+
+			// Initialize the number of per-source LPF filters based on input channels
+			LowPassFilters[SourceId].AddDefaulted(InitParams.NumInputChannels);
+
+			CurrentFrameValues[SourceId].Init(0.0f, InitParams.NumInputChannels);
+			NextFrameValues[SourceId].Init(0.0f, InitParams.NumInputChannels);
+
+			AUDIO_MIXER_CHECK(MixerSources[SourceId] == nullptr);
+			MixerSources[SourceId] = InitParams.SourceVoice;
+
+			InitParams.OwningSubmix->AddSourceVoice(InitParams.SourceVoice);
+
+#if AUDIO_MIXER_ENABLE_DEBUG_MODE
+			AUDIO_MIXER_CHECK(!bIsDebugMode[SourceId]);
+			bIsDebugMode[SourceId] = InitParams.bIsDebugMode;
+
+			AUDIO_MIXER_CHECK(DebugName[SourceId].IsEmpty());
+			DebugName[SourceId] = InitParams.DebugName;
+#endif 
+
+			AUDIO_MIXER_DEBUG_LOG(SourceId, TEXT("Is initializing"));
+		});
+	}
+
+	void FMixerSourceManager::ReleaseSourceId(const int32 SourceId)
+	{
+		AUDIO_MIXER_CHECK(GameThreadInfo.bIsBusy[SourceId]);
+		AUDIO_MIXER_CHECK_GAME_THREAD(MixerDevice);
+
+		AUDIO_MIXER_CHECK(NumActiveSources > 0);
+		--NumActiveSources;
+
+		GameThreadInfo.bIsBusy[SourceId] = false;
+
+#if AUDIO_MIXER_ENABLE_DEBUG_MODE
+		GameThreadInfo.bIsDebugMode[SourceId] = false;
+#endif
+
+		GameThreadInfo.FreeSourceIndices.Push(SourceId);
+
+		AUDIO_MIXER_CHECK(GameThreadInfo.FreeSourceIndices.Contains(SourceId));
+
+		AudioMixerThreadCommand([this, SourceId]()
+		{
+			AUDIO_MIXER_CHECK_AUDIO_PLAT_THREAD(MixerDevice);
+
+			ReleaseSource(SourceId);
+		});
+	}
+
+	void FMixerSourceManager::Play(const int32 SourceId)
+	{
+		AUDIO_MIXER_CHECK(SourceId < NumTotalSources);
+		AUDIO_MIXER_CHECK(GameThreadInfo.bIsBusy[SourceId]);
+		AUDIO_MIXER_CHECK_GAME_THREAD(MixerDevice);
+
+		AudioMixerThreadCommand([this, SourceId]()
+		{
+			AUDIO_MIXER_CHECK_AUDIO_PLAT_THREAD(MixerDevice);
+
+			bIsPlaying[SourceId] = true;
+			bIsPaused[SourceId] = false;
+			bIsActive[SourceId] = true;
+
+			AUDIO_MIXER_DEBUG_LOG(SourceId, TEXT("Is playing"));
+		});
+	}
+
+	void FMixerSourceManager::Stop(const int32 SourceId)
+	{
+		AUDIO_MIXER_CHECK(SourceId < NumTotalSources);
+		AUDIO_MIXER_CHECK(GameThreadInfo.bIsBusy[SourceId]);
+		AUDIO_MIXER_CHECK_GAME_THREAD(MixerDevice);
+
+		AudioMixerThreadCommand([this, SourceId]()
+		{
+			AUDIO_MIXER_CHECK_AUDIO_PLAT_THREAD(MixerDevice);
+
+			bIsPlaying[SourceId] = false;
+			bIsPaused[SourceId] = false;
+			bIsActive[SourceId] = false;
+
+			AUDIO_MIXER_DEBUG_LOG(SourceId, TEXT("Is stopping"));
+		});
+	}
+
+	void FMixerSourceManager::Pause(const int32 SourceId)
+	{
+		AUDIO_MIXER_CHECK(SourceId < NumTotalSources);
+		AUDIO_MIXER_CHECK(GameThreadInfo.bIsBusy[SourceId]);
+		AUDIO_MIXER_CHECK_GAME_THREAD(MixerDevice);
+
+		AudioMixerThreadCommand([this, SourceId]()
+		{
+			AUDIO_MIXER_CHECK_AUDIO_PLAT_THREAD(MixerDevice);
+
+			bIsPaused[SourceId] = true;
+			bIsActive[SourceId] = false;
+		});
+	}
+
+	void FMixerSourceManager::SetPitch(const int32 SourceId, const float Pitch)
+	{
+		AUDIO_MIXER_CHECK(SourceId < NumTotalSources);
+		AUDIO_MIXER_CHECK(GameThreadInfo.bIsBusy[SourceId]);
+
+		AudioMixerThreadCommand([this, SourceId, Pitch]()
+		{
+			AUDIO_MIXER_CHECK_AUDIO_PLAT_THREAD(MixerDevice);
+
+			PitchSourceParam[SourceId].SetValue(Pitch);
+		});
+	}
+
+	void FMixerSourceManager::SetVolume(const int32 SourceId, const float Volume)
+	{
+		AUDIO_MIXER_CHECK(SourceId < NumTotalSources);
+		AUDIO_MIXER_CHECK(GameThreadInfo.bIsBusy[SourceId]);
+		AUDIO_MIXER_CHECK_GAME_THREAD(MixerDevice);
+
+		AudioMixerThreadCommand([this, SourceId, Volume]()
+		{
+			AUDIO_MIXER_CHECK_AUDIO_PLAT_THREAD(MixerDevice);
+
+			VolumeSourceParam[SourceId].SetValue(Volume);
+		});
+	}
+
+	void FMixerSourceManager::SetSpatializationParams(const int32 SourceId, const FSpatializationParams& InParams)
+	{
+		AUDIO_MIXER_CHECK(SourceId < NumTotalSources);
+		AUDIO_MIXER_CHECK(GameThreadInfo.bIsBusy[SourceId]);
+		AUDIO_MIXER_CHECK_GAME_THREAD(MixerDevice);
+
+		AudioMixerThreadCommand([this, SourceId, InParams]()
+		{
+			AUDIO_MIXER_CHECK_AUDIO_PLAT_THREAD(MixerDevice);
+
+			SpatParams[SourceId] = InParams;
+		});
+	}
+
+	void FMixerSourceManager::SetWetLevel(const int32 SourceId, const float WetLevel)
+	{
+		AUDIO_MIXER_CHECK(SourceId < NumTotalSources);
+		AUDIO_MIXER_CHECK(GameThreadInfo.bIsBusy[SourceId]);
+		AUDIO_MIXER_CHECK_GAME_THREAD(MixerDevice);
+
+		AudioMixerThreadCommand([this, SourceId, WetLevel]()
+		{
+			AUDIO_MIXER_CHECK_AUDIO_PLAT_THREAD(MixerDevice);
+
+			WetLevelSourceParam[SourceId].SetValue(WetLevel);
+		});
+	}
+
+	void FMixerSourceManager::SetChannelMap(const int32 SourceId, const TArray<float>& ChannelMap)
+	{
+		AUDIO_MIXER_CHECK(SourceId < NumTotalSources);
+		AUDIO_MIXER_CHECK(GameThreadInfo.bIsBusy[SourceId]);
+		AUDIO_MIXER_CHECK_GAME_THREAD(MixerDevice);
+
+		AudioMixerThreadCommand([this, SourceId, ChannelMap]()
+		{
+			AUDIO_MIXER_CHECK_AUDIO_PLAT_THREAD(MixerDevice);
+
+			ChannelMapParam[SourceId].SetChannelMap(ChannelMap);
+		});
+	}
+
+	void FMixerSourceManager::SetLPFFrequency(const int32 SourceId, const float InLPFFrequency)
+	{
+		AUDIO_MIXER_CHECK(SourceId < NumTotalSources);
+		AUDIO_MIXER_CHECK(GameThreadInfo.bIsBusy[SourceId]);
+		AUDIO_MIXER_CHECK_GAME_THREAD(MixerDevice);
+
+		AudioMixerThreadCommand([this, SourceId, InLPFFrequency]()
+		{
+			AUDIO_MIXER_CHECK_AUDIO_PLAT_THREAD(MixerDevice);
+
+			float SampleRate = MixerDevice->GetSampleRate();
+			AUDIO_MIXER_CHECK(SampleRate > 0.0f);
+
+			const float NormalizedFrequency = 2.0f * InLPFFrequency / SampleRate;
+			LPFCutoffFrequencyParam[SourceId].SetValue(NormalizedFrequency);
+		});
+	}
+	 
+	void FMixerSourceManager::SubmitBuffer(const int32 SourceId, FMixerSourceBufferPtr InSourceVoiceBuffer)
+	{
+		AUDIO_MIXER_CHECK(SourceId < NumTotalSources);
+		AUDIO_MIXER_CHECK(GameThreadInfo.bIsBusy[SourceId]);
+		AUDIO_MIXER_CHECK(InSourceVoiceBuffer->AudioBytes <= (uint32)InSourceVoiceBuffer->AudioData.Num());
+		AUDIO_MIXER_CHECK_GAME_THREAD(MixerDevice);
+
+		AudioMixerThreadCommand([this, SourceId, InSourceVoiceBuffer]()
+		{
+			AUDIO_MIXER_CHECK_AUDIO_PLAT_THREAD(MixerDevice);
+
+			AUDIO_MIXER_CHECK(InSourceVoiceBuffer->AudioBytes <= (uint32)InSourceVoiceBuffer->AudioData.Num());
+			BufferQueue[SourceId].Enqueue(InSourceVoiceBuffer);
+		});
+	}
+
+	void FMixerSourceManager::SubmitBufferAudioThread(const int32 SourceId, FMixerSourceBufferPtr InSourceVoiceBuffer)
+	{
+		AUDIO_MIXER_CHECK(SourceId < NumTotalSources);
+		AUDIO_MIXER_CHECK_AUDIO_PLAT_THREAD(MixerDevice);
+
+		// Immediately push the source buffer onto the audio thread buffer queue.
+		BufferQueue[SourceId].Enqueue(InSourceVoiceBuffer);
+	}
+
+	int64 FMixerSourceManager::GetNumFramesPlayed(const int32 SourceId) const
+	{
+		AUDIO_MIXER_CHECK_GAME_THREAD(MixerDevice);
+		return NumFramesPlayed[SourceId];
+	}
+
+	bool FMixerSourceManager::IsDone(const int32 SourceId) const
+	{
+		AUDIO_MIXER_CHECK_GAME_THREAD(MixerDevice);
+		return GameThreadInfo.bIsDone[SourceId];
+	}
+
+	void FMixerSourceManager::ReadSourceFrame(const int32 SourceId)
+	{
+		const int32 NumChannels = NumInputChannels[SourceId];
+
+		// Check if the next frame index is out of range of the total number of frames we have in our current audio buffer
+		bool bNextFrameOutOfRange = (CurrentFrameIndex[SourceId] + 1) >= CurrentAudioChunkNumFrames[SourceId];
+		bool bCurrentFrameOutOfRange = CurrentFrameIndex[SourceId] >= CurrentAudioChunkNumFrames[SourceId];
+
+		bool bReadCurrentFrame = true;
+
+		// Check the boolean conditions that determine if we need to pop buffers from our queue (in PCMRT case) *OR* loop back (looping PCM data)
+		while (bNextFrameOutOfRange || bCurrentFrameOutOfRange)
+		{
+			// If our current frame is in range, but next frame isn't, read the current frame now to avoid pops when transitioning between buffers
+			if (bNextFrameOutOfRange && !bCurrentFrameOutOfRange)
+			{
+				// Don't need to read the current frame audio after reading new audio chunk
+				bReadCurrentFrame = false;
+
+				AUDIO_MIXER_CHECK(CurrentPCMBuffer[SourceId].IsValid());
+				const int16* AudioData = (const int16*)CurrentPCMBuffer[SourceId]->AudioData.GetData();
+				const int32 CurrentSampleIndex = CurrentFrameIndex[SourceId] * NumChannels;
+
+				for (int32 Channel = 0; Channel < NumChannels; ++Channel)
+				{
+					CurrentFrameValues[SourceId][Channel] = (float)AudioData[CurrentSampleIndex + Channel] * ONEOVERSHORTMAX;
+				}
+			}
+
+			// If this is our first PCM buffer, we don't need to do a callback to get more audio
+			if (CurrentPCMBuffer[SourceId].IsValid())
+			{
+				if (CurrentPCMBuffer[SourceId]->LoopCount == Audio::LOOP_FOREVER && !CurrentPCMBuffer[SourceId]->bRealTimeBuffer)
+				{
+					AUDIO_MIXER_DEBUG_LOG(SourceId, TEXT("Hit Loop boundary, looping."));
+
+					CurrentFrameIndex[SourceId] = CurrentFrameIndex[SourceId] - CurrentAudioChunkNumFrames[SourceId];
+					break;
+				}
+
+				BufferQueueListener[SourceId]->OnSourceBufferEnd();
+			}
+
+			// If we have audio in our queue, we're still playing
+			if (!BufferQueue[SourceId].IsEmpty())
+			{
+				FMixerSourceBufferPtr NewBufferPtr;
+				BufferQueue[SourceId].Dequeue(NewBufferPtr);
+				CurrentPCMBuffer[SourceId] = NewBufferPtr;
+
+				AUDIO_MIXER_CHECK(MixerSources[SourceId]->NumBuffersQueued.GetValue() > 0);
+				AUDIO_MIXER_CHECK(CurrentPCMBuffer[SourceId]->AudioBytes <= (uint32)CurrentPCMBuffer[SourceId]->AudioData.Num());
+
+				MixerSources[SourceId]->NumBuffersQueued.Decrement();
+
+				CurrentAudioChunkNumFrames[SourceId] = CurrentPCMBuffer[SourceId]->AudioBytes / (NUM_BYTES_PER_SAMPLE * NumChannels);
+
+				// Subtract the number of frames in the current buffer from our frame index.
+				// Note: if this is the first time we're playing, CurrentFrameIndex will be 0
+				CurrentFrameIndex[SourceId] = FMath::Max(CurrentFrameIndex[SourceId] - CurrentAudioChunkNumFrames[SourceId] - 1, 0);
+			}
+			else
+			{
+				if (!bIsDone[SourceId])
+				{
+					AUDIO_MIXER_DEBUG_LOG(SourceId, TEXT("Is now done."));
+
+					bIsDone[SourceId] = true;
+					BufferQueue[SourceId].Empty();
+					MixerSources[SourceId]->NumBuffersQueued.Set(0);
+					CurrentFrameValues[SourceId].Reset();
+					NextFrameValues[SourceId].Reset();
+					CurrentPCMBuffer[SourceId] = nullptr;
+				}
+				return;
+			}
+
+			bNextFrameOutOfRange = (CurrentFrameIndex[SourceId] + 1) >= CurrentAudioChunkNumFrames[SourceId];
+			bCurrentFrameOutOfRange = CurrentFrameIndex[SourceId] >= CurrentAudioChunkNumFrames[SourceId];
+		}
+		
+		if (CurrentPCMBuffer[SourceId].IsValid())
+		{
+			// Grab the 16-bit PCM audio data (which could be a new audio chunk from previous ReadSourceFrame call)
+			const int16* AudioData = (const int16*)CurrentPCMBuffer[SourceId]->AudioData.GetData();
+			const int32 NextSampleIndex = (CurrentFrameIndex[SourceId] + 1)  * NumChannels;
+
+			if (bReadCurrentFrame)
+			{
+				const int32 CurrentSampleIndex = CurrentFrameIndex[SourceId] * NumChannels;
+				for (int32 Channel = 0; Channel < NumChannels; ++Channel)
+				{
+					CurrentFrameValues[SourceId][Channel] = (float)AudioData[CurrentSampleIndex + Channel] * ONEOVERSHORTMAX;
+					NextFrameValues[SourceId][Channel] = (float)AudioData[NextSampleIndex + Channel] * ONEOVERSHORTMAX;
+				}
+			}
+			else
+			{
+				for (int32 Channel = 0; Channel < NumChannels; ++Channel)
+				{
+					NextFrameValues[SourceId][Channel] = (float)AudioData[NextSampleIndex + Channel] * ONEOVERSHORTMAX;
+				}
+			}
+		}
+	}
+
+	void FMixerSourceManager::ComputeSourceBuffers()
+	{
+		AUDIO_MIXER_CHECK_AUDIO_PLAT_THREAD(MixerDevice);
+
+		const int32 NumFrames = MixerDevice->GetNumOutputFrames();
+
+		// Prepare the source buffers for writing source samples
+		for (int32 SourceId = 0; SourceId < NumTotalSources; ++SourceId)
+		{
+			const int32 NumSourceSamples = NumFrames*NumInputChannels[SourceId];
+			SourceBuffer[SourceId].Reset(NumSourceSamples);
+		}
+
+		// Local variable used for sample values
+		float SampleValue = 0.0f;
+
+		for (int32 SourceId = 0; SourceId < NumTotalSources; ++SourceId)
+		{
+			if (!bIsBusy[SourceId] || !bIsPlaying[SourceId] || bIsPaused[SourceId])
+			{
+				continue;
+			}
+
+			for (int32 Frame = 0; Frame < NumFrames; ++Frame)
+			{
+				// If the source is done, then we'll just write out 0s
+				if (!bIsDone[SourceId])
+				{
+					// Whether or not we need to read another sample from the source buffers
+					// If we haven't yet played any frames, then we will need to read the first source samples no matter what
+					bool bReadNextSample = !bHasStarted[SourceId];
+
+					// Reset that we've started generating audio
+					bHasStarted[SourceId] = true;
+
+					// Update the PrevFrameIndex value for the source based on alpha value
+					while (CurrentFrameAlpha[SourceId] >= 1.0f)
+					{
+						// Our inter-frame alpha lerping value is causing us to read new source frames
+						bReadNextSample = true;
+
+						// Bump up the current frame index
+						CurrentFrameIndex[SourceId]++;
+
+						// Bump up the frames played -- this is tracking the total frames in source file played
+						// CurrentFrameIndex can wrap for looping sounds so won't be accurate in that case
+						NumFramesPlayed[SourceId]++;
+
+						CurrentFrameAlpha[SourceId] -= 1.0f;
+					}
+
+					// If our alpha parameter caused us to jump to a new source frame, we need
+					// read new samples into our prev and next frame sample data
+					if (bReadNextSample)
+					{
+						ReadSourceFrame(SourceId);
+					}
+				}
+
+				const int32 NumSourceChannels = NumInputChannels[SourceId];
+
+				// If we've finished reading all buffer data, then just write out 0s
+				if (bIsDone[SourceId])
+				{
+					for (int32 Channel = 0; Channel < NumSourceChannels; ++Channel)
+					{
+						SourceBuffer[SourceId].Add(0.0f);
+					}
+				}
+				else
+				{
+					// Get the volume value of the source at this frame index
+					for (int32 Channel = 0; Channel < NumSourceChannels; ++Channel)
+					{
+						const float CurrFrameValue = CurrentFrameValues[SourceId][Channel];
+						const float NextFrameValue = NextFrameValues[SourceId][Channel];
+						const float CurrentAlpha = CurrentFrameAlpha[SourceId];
+
+						SampleValue = FMath::Lerp(CurrFrameValue, NextFrameValue, CurrentAlpha);
+						SourceBuffer[SourceId].Add(SampleValue);
+					}
+					const float CurrentPitchScale = PitchSourceParam[SourceId]();
+
+					CurrentFrameAlpha[SourceId] += CurrentPitchScale;
+				}
+			}
+		}
+	}
+
+	void FMixerSourceManager::ComputePostSourceEffectBuffers()
+	{
+		AUDIO_MIXER_CHECK_AUDIO_PLAT_THREAD(MixerDevice);
+
+		// First run each source buffer through it's default source effects (e.g. LPF)
+
+		const int32 NumFrames = MixerDevice->GetNumOutputFrames();
+		class IAudioSpatializationAlgorithm* SpatializeProcessor = MixerDevice->SpatializeProcessor;
+
+		const bool bIsDebugModeEnabled = DebugSoloSources.Num() > 0;
+
+		for (int32 SourceId = 0; SourceId < NumTotalSources; ++SourceId)
+		{
+			if (!bIsBusy[SourceId] || !bIsPlaying[SourceId] || bIsPaused[SourceId])
+			{
+				continue;
+			}
+
+			// Get the source buffer
+			TArray<float>& Buffer = SourceBuffer[SourceId];
+
+			// First apply the LPF filter (and any other default per-source effecct) to the source before spatializing
+			FSourceParam& LPFFrequencyParam = LPFCutoffFrequencyParam[SourceId];
+			const int32 NumInputChans = NumInputChannels[SourceId];
+			int32 SampleIndex = 0;
+			for (int32 Frame = 0; Frame < NumFrames; ++Frame)
+			{
+				const float LPFFreq = LPFFrequencyParam();
+				float CurrentVolumeValue = VolumeSourceParam[SourceId]();
+				
+#if AUDIO_MIXER_ENABLE_DEBUG_MODE				
+				if (bIsDebugModeEnabled && !bIsDebugMode[SourceId])
+				{
+					CurrentVolumeValue *= 0.0f;
+				}
+#endif
+
+				for (int32 Channel = 0; Channel < NumInputChans; ++Channel)
+				{
+					SampleIndex = NumInputChans * Frame + Channel;
+
+					// Update the frequency
+					LowPassFilters[SourceId][Channel].SetFrequency(LPFFreq);
+
+					// Process the source through the filter
+					float SourceSample = Buffer[SampleIndex];
+					
+					SourceSample = LowPassFilters[SourceId][Channel](SourceSample);
+
+					// Process the effect chain
+					// TODO
+					// SourceSample = MyEffect[SourceId][Channel](SourceSample);
+
+					SourceSample *= CurrentVolumeValue;
+
+					// Write back out the sample after the effect processing and volume scaling
+					Buffer[SampleIndex] = SourceSample;
+				}
+			}
+
+			// If the source has HRTF processing enabled, run it through the spatializer
+			if (bUseHRTFSpatializer[SourceId])
+			{
+				AUDIO_MIXER_CHECK(SpatializeProcessor);
+				AUDIO_MIXER_CHECK(NumInputChannels[SourceId] == 1);
+
+				// Reset our scratch buffer and make sure it has enough data to hold 2 channels of interleaved data
+				ScratchBuffer.Reset();
+				ScratchBuffer.AddZeroed(2 * NumFrames);
+
+				FSpatializationParams& SourceSpatParams = SpatParams[SourceId];
+				SpatializeProcessor->SetSpatializationParameters(SourceId, SourceSpatParams);
+				SpatializeProcessor->ProcessSpatializationForVoice(SourceId, Buffer.GetData(), ScratchBuffer.GetData());
+
+				// We are now a 2-channel file and should not be spatialized using normal 3d spatialization
+				NumPostEffectChannels[SourceId] = 2;
+
+				// Copy the output scratch buffer to the post-effect buffers
+				PostEffectBuffers[SourceId] = ScratchBuffer;
+			}
+			else
+			{
+				PostEffectBuffers[SourceId] = Buffer;
+
+				// Otherwise our pre- and post-effect channels are the same as the input channels
+				NumPostEffectChannels[SourceId] = NumInputChannels[SourceId];
+			}
+		}
+	}
+
+	void FMixerSourceManager::ComputeOutputBuffers()
+	{
+		AUDIO_MIXER_CHECK_AUDIO_PLAT_THREAD(MixerDevice);
+
+		const int32 NumFrames = MixerDevice->GetNumOutputFrames();
+		const int32 NumOutputChannels = MixerDevice->GetNumDeviceChannels();
+
+		// Reset the dry/wet buffers for all the sources
+		const int32 NumOutputSamples = NumFrames * NumOutputChannels;
+
+		for (int32 SourceId = 0; SourceId < NumTotalSources; ++SourceId)
+		{
+			DryBuffer[SourceId].Reset();
+			DryBuffer[SourceId].AddZeroed(NumOutputSamples);
+
+			WetBuffer[SourceId].Reset();
+			WetBuffer[SourceId].AddZeroed(NumOutputSamples);
+		}
+
+		for (int32 SourceId = 0; SourceId < NumTotalSources; ++SourceId)
+		{
+			// Don't need to compute anything if the source is not playing or paused (it will remain at 0.0 volume)
+			// Note that effect chains will still be able to continue to compute audio output. The source output 
+			// will simply stop being read from.
+			if (!bIsBusy[SourceId] || !bIsPlaying[SourceId] || bIsPaused[SourceId])
+			{
+				continue;
+			}
+
+			for (int32 Frame = 0; Frame < NumFrames; ++Frame)
+			{
+				const int32 PostEffectChannels = NumPostEffectChannels[SourceId];
+
+				// Compute the constant-power values of the dry-wet buffer scalars
+				const float WetLevelParamValue = WetLevelSourceParam[SourceId]();
+
+				const float WetLevel = FMath::Sin(WetLevelParamValue * 0.5f * PI);
+				const float DryLevel = FMath::Cos(WetLevelParamValue * 0.5f * PI);
+
+#if ENABLE_AUDIO_OUTPUT_DEBUGGING
+				FSineOsc& SineOsc = DebugOutputGenerators[SourceId];
+				float SourceSampleValue = SineOsc();
+#else
+				float SourceSampleValue = 0.0f;
+#endif
+
+				// For each source channel, compute the output channel mapping
+				for (int32 SourceChannel = 0; SourceChannel < PostEffectChannels; ++SourceChannel)
+				{
+
+#if !ENABLE_AUDIO_OUTPUT_DEBUGGING
+					const int32 SourceSampleIndex = Frame * PostEffectChannels + SourceChannel;
+					SourceSampleValue = PostEffectBuffers[SourceId][SourceSampleIndex];
+#endif
+					// Make sure that our channel map is appropriate for the source channel and output channel count!
+					ChannelMapParam[SourceId].GetChannelMap(ScratchChannelMap);
+					AUDIO_MIXER_CHECK(ScratchChannelMap.Num() == PostEffectChannels * NumOutputChannels);
+
+					for (int32 OutputChannel = 0; OutputChannel < NumOutputChannels; ++OutputChannel)
+					{
+						// Look up the channel map value (maps input channels to output channels) for the source
+						// This is the step that either applies the spatialization algorithm or just maps a 2d sound
+						const int32 ChannelMapIndex = PostEffectChannels * OutputChannel + SourceChannel;
+						const float ChannelMapValue = ScratchChannelMap[ChannelMapIndex];
+
+						AUDIO_MIXER_CHECK(ChannelMapValue >= 0.0f && ChannelMapValue <= 1.0f);
+
+						// Scale the input source sample for this source channel value
+						const float SampleValue = SourceSampleValue * ChannelMapValue;
+
+						// If we have a non-zero sample value, write it out. Note that most 3d audio channel maps
+						// for surround sound will result in 0.0 sample values so this branch should save a bunch of multiplies + adds
+						const int32 OutputSampleIndex = Frame * NumOutputChannels + OutputChannel;
+
+						DryBuffer[SourceId][OutputSampleIndex] += SampleValue * DryLevel;
+						WetBuffer[SourceId][OutputSampleIndex] += SampleValue * WetLevel;
+					}
+				}
+			}
+		}
+	}
+
+	void FMixerSourceManager::MixOutputBuffers(const int32 SourceId, TArray<float>& OutDryBuffer, TArray<float>& OutWetBuffer) const
+	{
+		AUDIO_MIXER_CHECK_AUDIO_PLAT_THREAD(MixerDevice);
+		AUDIO_MIXER_CHECK(SourceId < NumTotalSources);
+		AUDIO_MIXER_CHECK(OutDryBuffer.Num() == OutWetBuffer.Num());
+		AUDIO_MIXER_CHECK(OutDryBuffer.Num() == DryBuffer[SourceId].Num());
+
+		// Mix the source's dry and wet buffers into the output buffers. Note that the source's dry/wet buffers need to have been already computed
+		// earlier in the audio update frame.
+		for (int32 SampleIndex = 0; SampleIndex < OutDryBuffer.Num(); ++SampleIndex)
+		{
+			OutDryBuffer[SampleIndex] += DryBuffer[SourceId][SampleIndex];
+			OutWetBuffer[SampleIndex] += WetBuffer[SourceId][SampleIndex];
+		}
+	}
+
+	void FMixerSourceManager::ComputeNextBlockOfSamples()
+	{
+		AUDIO_MIXER_CHECK_AUDIO_PLAT_THREAD(MixerDevice);
+
+		if (bPumpQueue)
+		{
+			bPumpQueue = false;
+			PumpCommandQueue();
+		}
+
+		// Get the next block of frames from the source buffers
+		ComputeSourceBuffers();
+
+		// Compute the audio source buffers after their individual effect chain processing
+		ComputePostSourceEffectBuffers();
+
+		// Get the audio for the output buffers
+		ComputeOutputBuffers();
+
+		// Update the game thread copy of source doneness
+		for (int32 SourceId = 0; SourceId < NumTotalSources; ++SourceId)
+		{
+			GameThreadInfo.bIsDone[SourceId] = bIsDone[SourceId];
+		}
+	}
+
+	void FMixerSourceManager::AudioMixerThreadCommand(TFunction<void()> InFunction)
+	{
+		// Add the function to the command queue
+		SourceCommandQueue.Enqueue(MoveTemp(InFunction));
+	}
+
+	void FMixerSourceManager::PumpCommandQueue()
+	{
+		AUDIO_MIXER_CHECK_AUDIO_PLAT_THREAD(MixerDevice);
+
+		// Pop and execute all the commands that came since last update tick
+		TFunction<void()> CommandFunction;
+		while (SourceCommandQueue.Dequeue(CommandFunction))
+		{
+			CommandFunction();
+		}
+	}
+}
diff --git a/Engine/Source/Runtime/AudioMixer/Private/AudioMixerSourceManager.h b/Engine/Source/Runtime/AudioMixer/Private/AudioMixerSourceManager.h
new file mode 100644
--- /dev/null
+++ b/Engine/Source/Runtime/AudioMixer/Private/AudioMixerSourceManager.h
@@ -0,0 +1,275 @@
+// Copyright 1998-2016 Epic Games, Inc. All Rights Reserved.
+
+#pragma once
+
+/* Public dependencies
+*****************************************************************************/
+
+#include "SoundDefinitions.h"
+#include "AudioMixerDevice.h"
+#include "AudioMixerBuffer.h"
+#include "AudioMixerSourceVoice.h"
+#include "AudioMixerSource.h"
+#include "DSP/OnePole.h"
+#include "IAudioExtensionPlugin.h"
+
+#define ENABLE_AUDIO_OUTPUT_DEBUGGING 0
+
+#if AUDIO_MIXER_ENABLE_DEBUG_MODE
+
+// Macro which checks if the source id is in debug mode, avoids having a bunch of #ifdefs in code
+#define AUDIO_MIXER_DEBUG_LOG(SourceId, Format, ...)																							\
+	if (bIsDebugMode[SourceId])																													\
+	{																																			\
+		FString CustomMessage = FString::Printf(Format, ##__VA_ARGS__);																			\
+		FString LogMessage = FString::Printf(TEXT("<Debug Sound Log> [Id=%d][Name=%s]: %s"), SourceId, *DebugName[SourceId], *CustomMessage);	\
+		UE_LOG(LogAudioMixer, Log, TEXT("%s"), *LogMessage);																								\
+	}
+
+#else
+
+#define AUDIO_MIXER_DEBUG_LOG(SourceId, Message)
+
+#endif
+
+
+namespace Audio
+{
+	class FMixerSubmix;
+	class FMixerDevice;
+	class FMixerSourceVoice;
+
+	struct FMixerSourceVoiceBuffer;
+
+	typedef TSharedPtr<FMixerSourceVoiceBuffer, ESPMode::ThreadSafe> FMixerSourceBufferPtr;
+
+	class ISourceBufferQueueListener
+	{
+	public:
+		// Called when the current buffer is finished and a new one needs to be queued
+		virtual void OnSourceBufferEnd() = 0;
+	};
+
+	struct FMixerSourceVoiceInitParams
+	{
+		ISourceBufferQueueListener* BufferQueueListener;
+		FMixerSubmix* OwningSubmix;
+		FMixerSourceVoice* SourceVoice;
+		int32 NumInputChannels;
+		int32 NumInputFrames;
+		FString DebugName;
+		bool bUseHRTFSpatialization;
+		bool bIsDebugMode;
+
+		FMixerSourceVoiceInitParams()
+			: BufferQueueListener(nullptr)
+			, OwningSubmix(nullptr)
+			, SourceVoice(nullptr)
+			, NumInputChannels(0)
+			, NumInputFrames(0)
+			, bUseHRTFSpatialization(false)
+			, bIsDebugMode(false)
+		{}
+	};
+
+	class FSourceParam
+	{
+	public:
+		FSourceParam(float InNumInterpFrames);
+
+		void Reset();
+
+		void SetValue(float InValue);
+		float operator()();
+
+	private:
+		float StartValue;
+		float EndValue;
+		float CurrentValue;
+		float NumInterpFrames;
+		float Frame;
+		bool bIsInit;
+	};
+
+
+	class FSourceChannelMap
+	{
+	public:
+		FSourceChannelMap(float InNumInterpFrames)
+			: NumInterpFrames(InNumInterpFrames)
+		{}
+
+		void Reset()
+		{
+			ChannelValues.Reset();
+		}
+
+		void SetChannelMap(const TArray<float>& ChannelMap)
+		{
+			if (!ChannelValues.Num())
+			{
+				for (int32 i = 0; i < ChannelMap.Num(); ++i)
+				{
+					ChannelValues.Add(FSourceParam(NumInterpFrames));
+					ChannelValues[i].SetValue(ChannelMap[i]);
+				}
+			}
+			else
+			{
+				for (int32 i = 0; i < ChannelMap.Num(); ++i)
+				{
+					ChannelValues[i].SetValue(ChannelMap[i]);
+				}
+			}
+		}
+
+		void GetChannelMap(TArray<float>& OutChannelMap)
+		{
+			OutChannelMap.Reset();
+			for (int32 i = 0; i < ChannelValues.Num(); ++i)
+			{
+				OutChannelMap.Add(ChannelValues[i]());
+			}
+		}
+
+	private:
+		TArray<FSourceParam> ChannelValues;
+		float NumInterpFrames;
+	};
+
+	class FMixerSourceManager
+	{
+	public:
+		FMixerSourceManager(FMixerDevice* InMixerDevice);
+		~FMixerSourceManager();
+
+		void Init(const int32 InNumSources);
+		void Update();
+
+		bool GetFreeSourceId(int32& OutSourceId);
+		int32 GetNumActiveSources() const;
+
+		void ReleaseSourceId(const int32 SourceId);
+		void InitSource(const int32 SourceId, const FMixerSourceVoiceInitParams& InitParams);
+
+		void Play(const int32 SourceId);
+		void Stop(const int32 SourceId);
+		void Pause(const int32 SourceId);
+		void SetPitch(const int32 SourceId, const float Pitch);
+		void SetVolume(const int32 SourceId, const float Volume);
+		void SetSpatializationParams(const int32 SourceId, const FSpatializationParams& InParams);
+		void SetWetLevel(const int32 SourceId, const float WetLevel);
+		void SetChannelMap(const int32 SourceId, const TArray<float>& InChannelMap);
+		void SetLPFFrequency(const int32 SourceId, const float Frequency);
+		
+		void SubmitBuffer(const int32 SourceId, FMixerSourceBufferPtr InSourceVoiceBuffer);
+		void SubmitBufferAudioThread(const int32 SourceId, FMixerSourceBufferPtr InSourceVoiceBuffer);
+
+		int64 GetNumFramesPlayed(const int32 SourceId) const;
+		bool IsDone(const int32 SourceId) const;
+
+		void ComputeNextBlockOfSamples();
+		void MixOutputBuffers(const int32 SourceId, TArray<float>& OutDryBuffer, TArray<float>& OutWetBuffer) const;
+
+	private:
+
+		void ReleaseSource(const int32 SourceId);
+		void ReadSourceFrame(const int32 SourceId);
+		void ComputeSourceBuffers();
+		void ComputePostSourceEffectBuffers();
+		void ComputeOutputBuffers();
+
+		void AudioMixerThreadCommand(TFunction<void()> InFunction);
+		void PumpCommandQueue();
+
+		static const int32 NUM_BYTES_PER_SAMPLE = 2;
+
+		FMixerDevice* MixerDevice;
+
+		// Array of pointers to game thread audio source objects
+		TArray<FMixerSourceVoice*> MixerSources;
+
+		// A command queue to execute commands from audio thread (or game thread) to audio mixer device thread.
+		TQueue<TFunction<void()>> SourceCommandQueue;
+
+		/////////////////////////////////////////////////////////////////////////////
+		// Vectorized source voice info
+		/////////////////////////////////////////////////////////////////////////////
+
+		// Raw PCM buffer data
+		TArray<TQueue<FMixerSourceBufferPtr>> BufferQueue;
+		TArray<ISourceBufferQueueListener*> BufferQueueListener;
+		TArray<FThreadSafeCounter> NumBuffersQueued;
+
+		// Debugging source generators
+#if ENABLE_AUDIO_OUTPUT_DEBUGGING
+		TArray<FSineOsc> DebugOutputGenerators;
+#endif
+		// Array of sources which are currently debug solo'd
+		TArray<int32> DebugSoloSources;
+		TArray<bool> bIsDebugMode;
+		TArray<FString> DebugName;
+
+		// Data used for rendering sources
+		TArray<FMixerSourceBufferPtr> CurrentPCMBuffer;
+		TArray<int32> CurrentAudioChunkNumFrames;
+		TArray<TArray<float>> SourceBuffer;
+		TArray<TArray<float>> CurrentFrameValues;
+		TArray<TArray<float>> NextFrameValues;
+		TArray<float> CurrentFrameAlpha;
+		TArray<int32> CurrentFrameIndex;
+		TArray<int64> NumFramesPlayed;
+
+		TArray<FSourceParam> PitchSourceParam;
+		TArray<FSourceParam> VolumeSourceParam;
+		TArray<FSourceParam> WetLevelSourceParam;
+		TArray<FSourceParam> LPFCutoffFrequencyParam;
+
+		// Simple LPFs for all sources (all channels of all sources)
+		TArray<TArray<FOnePoleLPF>> LowPassFilters;
+
+		TArray<FSourceChannelMap> ChannelMapParam;
+		TArray<FSpatializationParams> SpatParams;
+		TArray<float> ScratchChannelMap;
+
+		// Output data, after computing a block of sample data, this is read back from mixers
+		TArray<TArray<float>> PostEffectBuffers;
+		TArray<TArray<float>> DryBuffer;
+		TArray<TArray<float>> WetBuffer;
+
+		// Buffer used as an intermediate buffer between effects
+		TArray<float> ScratchBuffer;
+
+		// State management
+		TArray<bool> bIsActive;
+		TArray<bool> bIsPlaying;
+		TArray<bool> bIsPaused;
+		TArray<bool> bHasStarted;
+		TArray<bool> bIsBusy;
+		TArray<bool> bUseHRTFSpatializer;
+		TArray<bool> bIsDone;
+
+		// Source format info
+		TArray<int32> NumInputChannels;
+		TArray<int32> NumPostEffectChannels;
+		TArray<int32> NumInputFrames;
+
+		// General information about sources in source manager accessible from game thread
+		struct FGameThreadInfo
+		{
+			TArray<int32> FreeSourceIndices;
+			TArray<bool> bIsBusy;
+			TArray<FThreadSafeBool> bIsDone;
+			TArray<bool> bIsDebugMode;
+		} GameThreadInfo;
+
+		int32 NumActiveSources;
+		int32 NumTotalSources;
+		uint8 bInitialized : 1;
+
+		friend class FMixerSourceVoice;
+
+		// Set to true when the audio source manager should pump the command queue
+		FThreadSafeBool bPumpQueue;
+	};
+}
diff --git a/Engine/Source/Runtime/AudioMixer/Private/AudioMixerSourceVoice.cpp b/Engine/Source/Runtime/AudioMixer/Private/AudioMixerSourceVoice.cpp
new file mode 100644
--- /dev/null
+++ b/Engine/Source/Runtime/AudioMixer/Private/AudioMixerSourceVoice.cpp
@@ -0,0 +1,214 @@
+﻿// Copyright 1998-2016 Epic Games, Inc. All Rights Reserved.
+
+#include "AudioMixerPCH.h"
+#include "AudioMixerSource.h"
+#include "AudioMixerSourceManager.h"
+
+namespace Audio
+{
+
+	/**
+	* FMixerSourceVoice Implementation
+	*/
+
+	FMixerSourceVoice::FMixerSourceVoice(FMixerDevice* InMixerDevice, FMixerSourceManager* InSourceManager)
+		: OwningSubmix(nullptr)
+		, SourceManager(InSourceManager)
+		, MixerDevice(InMixerDevice)
+		, NumBuffersQueued(0)
+		, Pitch(-1.0f)
+		, Volume(-1.0f)
+		, Distance(-1.0f)
+		, WetLevel(-1.0f)
+		, SourceId(INDEX_NONE)
+		, bIsPlaying(false)
+		, bIsPaused(false)
+		, bIsActive(false)
+	{
+		AUDIO_MIXER_CHECK(SourceManager != nullptr);
+	}
+
+	FMixerSourceVoice::~FMixerSourceVoice()
+	{
+	}
+
+	bool FMixerSourceVoice::Init(const FMixerSourceVoiceInitParams& InitParams)
+	{
+		AUDIO_MIXER_CHECK_GAME_THREAD(MixerDevice);
+
+		if (SourceManager->GetFreeSourceId(SourceId))
+		{
+			AUDIO_MIXER_CHECK(InitParams.BufferQueueListener != nullptr);
+			AUDIO_MIXER_CHECK(InitParams.NumInputChannels > 0);
+
+			OwningSubmix = InitParams.OwningSubmix;
+
+			SourceManager->InitSource(SourceId, InitParams);
+			return true;
+		}
+		return false;
+	}
+
+	void FMixerSourceVoice::Release()
+	{
+		AUDIO_MIXER_CHECK_GAME_THREAD(MixerDevice);
+
+		SourceManager->ReleaseSourceId(SourceId);
+	}
+
+	void FMixerSourceVoice::SubmitBuffer(FMixerSourceBufferPtr InSourceVoiceBuffer)
+	{
+		AUDIO_MIXER_CHECK_GAME_THREAD(MixerDevice);
+
+		NumBuffersQueued.Increment();
+		SourceManager->SubmitBuffer(SourceId, InSourceVoiceBuffer);
+	}
+
+	void FMixerSourceVoice::SubmitBufferAudioThread(FMixerSourceBufferPtr InSourceVoiceBuffer)
+	{
+		AUDIO_MIXER_CHECK_AUDIO_PLAT_THREAD(MixerDevice);
+
+		NumBuffersQueued.Increment();
+		SourceManager->SubmitBufferAudioThread(SourceId, InSourceVoiceBuffer);
+	}
+
+	int32 FMixerSourceVoice::GetNumBuffersQueued() const
+	{
+		return NumBuffersQueued.GetValue();
+	}
+
+	void FMixerSourceVoice::SetPitch(const float InPitch)
+	{
+		AUDIO_MIXER_CHECK_GAME_THREAD(MixerDevice);
+
+		if (Pitch != InPitch)
+		{
+			Pitch = InPitch;
+			SourceManager->SetPitch(SourceId, InPitch);
+		}
+	}
+
+	void FMixerSourceVoice::SetVolume(const float InVolume)
+	{
+		AUDIO_MIXER_CHECK_GAME_THREAD(MixerDevice);
+
+		if (Volume != InVolume)
+		{
+			Volume = InVolume;
+			SourceManager->SetVolume(SourceId, InVolume);
+		}
+	}
+
+	void FMixerSourceVoice::SetWetLevel(const float InWetLevel)
+	{
+		AUDIO_MIXER_CHECK_GAME_THREAD(MixerDevice);
+
+		if (WetLevel != InWetLevel)
+		{
+			WetLevel = InWetLevel;
+			SourceManager->SetWetLevel(SourceId, InWetLevel);
+		}
+	}
+
+	void FMixerSourceVoice::SetLPFFrequency(const float InLPFFrequency)
+	{
+		AUDIO_MIXER_CHECK_GAME_THREAD(MixerDevice);
+
+		if (LPFFrequency != InLPFFrequency)
+		{
+			LPFFrequency = InLPFFrequency;
+			SourceManager->SetLPFFrequency(SourceId, LPFFrequency);
+		}
+	}
+
+	void FMixerSourceVoice::SetChannelMap(TArray<float>& InChannelMap)
+	{
+		AUDIO_MIXER_CHECK_GAME_THREAD(MixerDevice);
+
+		if (ChannelMap != InChannelMap)
+		{
+			ChannelMap = InChannelMap;
+			SourceManager->SetChannelMap(SourceId, InChannelMap);
+		}
+	}
+
+	void FMixerSourceVoice::SetSpatializationParams(const FSpatializationParams& InParams)
+	{
+		AUDIO_MIXER_CHECK_GAME_THREAD(MixerDevice);
+
+		SourceManager->SetSpatializationParams(SourceId, InParams);
+	}
+
+	void FMixerSourceVoice::Play()
+	{
+		AUDIO_MIXER_CHECK_GAME_THREAD(MixerDevice);
+
+		bIsPlaying = true;
+		bIsPaused = false;
+		bIsActive = true;
+		SourceManager->Play(SourceId);
+	}
+
+	void FMixerSourceVoice::Stop()
+	{
+		AUDIO_MIXER_CHECK_GAME_THREAD(MixerDevice);
+
+		bIsPlaying = false;
+		bIsPaused = false;
+		bIsActive = false;
+		SourceManager->Stop(SourceId);
+	}
+
+	void FMixerSourceVoice::Pause()
+	{
+		AUDIO_MIXER_CHECK_GAME_THREAD(MixerDevice);
+
+		bIsPaused = true;
+		bIsActive = false;
+		SourceManager->Pause(SourceId);
+	}
+
+	bool FMixerSourceVoice::IsPlaying() const
+	{
+		AUDIO_MIXER_CHECK_GAME_THREAD(MixerDevice);
+
+		return bIsPlaying;
+	}
+
+	bool FMixerSourceVoice::IsPaused() const
+	{
+		AUDIO_MIXER_CHECK_GAME_THREAD(MixerDevice);
+
+		return bIsPaused;
+	}
+
+	bool FMixerSourceVoice::IsActive() const
+	{
+		AUDIO_MIXER_CHECK_GAME_THREAD(MixerDevice);
+
+		return bIsActive;
+	}
+
+	bool FMixerSourceVoice::IsDone() const
+	{
+		AUDIO_MIXER_CHECK_GAME_THREAD(MixerDevice);
+
+		return SourceManager->IsDone(SourceId);
+	}
+
+	int64 FMixerSourceVoice::GetNumFramesPlayed() const
+	{
+		AUDIO_MIXER_CHECK_GAME_THREAD(MixerDevice);
+
+		return SourceManager->GetNumFramesPlayed(SourceId);
+	}
+
+	void FMixerSourceVoice::MixOutputBuffers(TArray<float>& OutDryBuffer, TArray<float>& OutWetBuffer) const
+	{
+		AUDIO_MIXER_CHECK_AUDIO_PLAT_THREAD(MixerDevice);
+
+		return SourceManager->MixOutputBuffers(SourceId, OutDryBuffer, OutWetBuffer);
+	}
+
+
+}
\ No newline at end of file
diff --git a/Engine/Source/Runtime/AudioMixer/Private/AudioMixerSourceVoice.h b/Engine/Source/Runtime/AudioMixer/Private/AudioMixerSourceVoice.h
new file mode 100644
--- /dev/null
+++ b/Engine/Source/Runtime/AudioMixer/Private/AudioMixerSourceVoice.h
@@ -0,0 +1,110 @@
+// Copyright 1998-2016 Epic Games, Inc. All Rights Reserved.
+
+#pragma once
+
+#include "AudioMixerDevice.h"
+#include "AudioMixerBuffer.h"
+#include "AudioMixerSource.h"
+#include "AudioMixerSubmix.h"
+#include "AudioMixerSourceManager.h"
+
+namespace Audio
+{
+	struct FMixerSourceVoiceBuffer;
+	struct FMixerSourceVoiceFilterParams;
+	class FMixerSubmix;
+	class FMixerSource;
+	class FMixerSourceManager;
+	class ISourceBufferQueueListener;
+
+
+	class FMixerSourceVoice
+	{
+	public:
+		FMixerSourceVoice(FMixerDevice* InMixerDevice, FMixerSourceManager* InSourceManager);
+		~FMixerSourceVoice();
+
+		// Initializes the mixer source voice
+		bool Init(const FMixerSourceVoiceInitParams& InFormat);
+
+		// Releases the source voice back to the source buffer pool
+		void Release();
+
+		// Queues the given buffer data to the internal queue of audio buffers.
+		void SubmitBuffer(FMixerSourceBufferPtr InSourceVoiceBuffer);
+
+		// Quees the given buffer data to teh internal queue of audio buffers on the audio thread.
+		void SubmitBufferAudioThread(FMixerSourceBufferPtr InSourceVoiceBuffer);
+
+		// Returns the number of buffers that are currently queued to be played.
+		int32 GetNumBuffersQueued() const;
+
+		// Sets the source voice pitch value.
+		void SetPitch(const float InPitch);
+
+		// Sets the source voice volume value.
+		void SetVolume(const float InVolume);
+	
+		// Sets the source voice wet-level value (for submix mixing).
+		void SetWetLevel(const float InWetLevel);
+
+		// Sets the source voice's LPF filter frequency.
+		void SetLPFFrequency(const float InFrequency);
+
+		// Sets the source voice's channel map (2d or 3d).
+		void SetChannelMap(TArray<float>& InChannelMap);
+
+		// Sets params used by HRTF spatializer
+		void SetSpatializationParams(const FSpatializationParams& InParams);
+
+		// Starts the source voice generating audio output into it's submix.
+		void Play();
+
+		// Pauses the source voice (i.e. stops generating output but keeps its state as "active and playing". Can be restarted.)
+		void Pause();
+
+		// Stops the source voice (no longer playing or active, can't be restarted.)
+		void Stop();
+
+		// Queries if the voice is playing
+		bool IsPlaying() const;
+
+		// Queries if the voice is paused
+		bool IsPaused() const;
+
+		// Queries if the source voice is active.
+		bool IsActive() const;
+
+		// Queries if the source voice has finished playing all its audio.
+		bool IsDone() const;
+
+		// Retrieves the total number of samples played.
+		int64 GetNumFramesPlayed() const;
+
+		// Mixes the dry and wet buffer audio into the given buffers.
+		void MixOutputBuffers(TArray<float>& OutDryBuffer, TArray<float>& OutWetBuffer) const;
+
+		// Returns the submix that owns this source voice.
+		FMixerSubmix* GetOwningSubmix() const { return OwningSubmix; }
+
+	private:
+
+		friend class FMixerSourceManager;
+
+		FMixerSubmix* OwningSubmix;
+		FMixerSourceManager* SourceManager;
+		FMixerDevice* MixerDevice;
+		TArray<float> ChannelMap;
+		FThreadSafeCounter NumBuffersQueued;
+		float Pitch;
+		float Volume;
+		float Distance;
+		float WetLevel;
+		float LPFFrequency;
+		int32 SourceId;
+		uint16 bIsPlaying : 1;
+		uint16 bIsPaused : 1;
+		uint16 bIsActive : 1;
+	};
+
+}
diff --git a/Engine/Source/Runtime/AudioMixer/Private/AudioMixerSubmix.cpp b/Engine/Source/Runtime/AudioMixer/Private/AudioMixerSubmix.cpp
new file mode 100644
--- /dev/null
+++ b/Engine/Source/Runtime/AudioMixer/Private/AudioMixerSubmix.cpp
@@ -0,0 +1,172 @@
+// Copyright 1998-2016 Epic Games, Inc. All Rights Reserved.
+
+#include "AudioMixerPCH.h"
+#include "AudioMixerSubmix.h"
+#include "Sound/SoundSubmix.h"
+
+namespace Audio
+{
+
+	FMixerSubmix::FMixerSubmix(USoundSubmix* InSoundSubmix, FMixerDevice* InMixerDevice)
+		: SoundSubmix(nullptr)
+		, MixerDevice(InMixerDevice)
+		, Parent(nullptr)
+		, WetLevel(0.5f)
+	{
+		if (InSoundSubmix != nullptr)
+		{
+			SoundSubmix = InSoundSubmix;
+
+			// Loop through the submix's presets and make new instances of effects in the same order as the presets
+			EffectSubmixChain.Reset();
+
+			// If this is the master submix, we have some default effects we're going to apply (like master EQ, master compression, master verb, etc)
+			if (IsMasterSubmix())
+			{
+
+			}
+
+
+			for (USoundEffectSubmixPreset* EffectPreset : SoundSubmix->SubmixEffectChain)
+			{
+				// Create a new effect instance using the preset
+				USoundEffectSubmix* SubmixEffect = EffectPreset->CreateNewEffect();
+
+				// Now set the preset
+				SubmixEffect->SetPreset(EffectPreset);
+
+				// Add the effect to this submix's chain
+				EffectSubmixChain.Add(SubmixEffect);
+			}
+		}
+	}
+
+	FMixerSubmix::~FMixerSubmix()
+	{
+	}
+
+	bool FMixerSubmix::IsMasterSubmix() const
+	{
+		// If we were given a nullptr sound submix, then it is the master sound mix
+		return SoundSubmix == nullptr;
+	}
+
+	int32 FMixerSubmix::GetNumSourceVoices() const
+	{
+		return MixerSourceVoices.Num();
+	}
+
+	int32 FMixerSubmix::GetNumEffects() const
+	{
+		return EffectSubmixChain.Num();
+	}
+
+	void FMixerSubmix::AddSourceVoice(FMixerSourceVoice* InSourceVoice)
+	{
+		AUDIO_MIXER_CHECK_AUDIO_PLAT_THREAD(MixerDevice);
+		MixerSourceVoices.Add(InSourceVoice);
+	}
+
+	void FMixerSubmix::RemoveSourceVoice(FMixerSourceVoice* InSourceVoice)
+	{
+		AUDIO_MIXER_CHECK_AUDIO_PLAT_THREAD(MixerDevice);
+		int32 NumRemoved = MixerSourceVoices.Remove(InSourceVoice);
+		AUDIO_MIXER_CHECK(NumRemoved == 1);
+	}
+
+	void FMixerSubmix::ProcessAudio(TArray<float>& OutAudioBuffer)
+	{
+		AUDIO_MIXER_CHECK_AUDIO_PLAT_THREAD(MixerDevice);
+
+		// Create a zero'd scratch buffer to get the audio from this submix's children
+		const int32 NumSamples = OutAudioBuffer.Num();
+
+		WetBuffer.Reset(NumSamples);
+		WetBuffer.AddZeroed(NumSamples);
+
+		DryBuffer.Reset(NumSamples);
+		DryBuffer.AddZeroed(NumSamples);
+
+		if (SoundSubmix != nullptr)
+		{
+			// First loop this submix's child submixes mixing in their output into this submix's dry/wet buffers.
+			for (USoundSubmix* ChildSubmix : SoundSubmix->ChildSubmixes)
+			{
+				FMixerSubmix* ChildMixerSubmix = MixerDevice->GetSubmixInstance(ChildSubmix);
+				AUDIO_MIXER_CHECK(ChildMixerSubmix != nullptr);
+
+				ScratchBuffer.Reset(NumSamples);
+
+				ChildMixerSubmix->ProcessAudio(ScratchBuffer);
+
+				// Get this child submix's wet level
+				float ChildWetLevel = ChildSubmix->OutputWetLevel;
+				float ChildDryLevel = 1.0f - WetLevel;
+
+				// Get a nice constant-power x-fade of the dry/wet levels
+				ChildWetLevel = FMath::Sin(WetLevel);
+				ChildDryLevel = FMath::Cos(WetLevel);
+
+				// Mix the child's submix results into this submix's dry and wet levels
+				for (int32 i = 0; i < NumSamples; ++i)
+				{
+					WetBuffer[i] += ChildDryLevel * ScratchBuffer[i];
+					DryBuffer[i] += ChildDryLevel* ScratchBuffer[i];
+				}
+			}
+		}
+
+		// Loop through this submix's sound sources, removing any sound sources that have finished
+		for (int32 i = MixerSourceVoices.Num() - 1; i >= 0; --i)
+		{
+			MixerSourceVoices[i]->MixOutputBuffers(DryBuffer, WetBuffer);
+		}
+
+		// Setup the input data buffer
+		FSoundEffectSubmixInputData InputData;
+		InputData.AudioBuffer = &WetBuffer;
+		InputData.NumChannels = MixerDevice->GetNumDeviceChannels();
+		InputData.AudioClock = MixerDevice->GetAudioTime();
+
+		FSoundEffectSubmixOutputData OutputData;
+		OutputData.AudioBuffer = &ScratchBuffer;
+
+		for (USoundEffectSubmix* SubmixEffect : EffectSubmixChain)
+		{
+			SubmixEffect->ProcessAudio(InputData, OutputData);
+
+			// Copy the output audio buffer to the old wet buffer
+			WetBuffer = *OutputData.AudioBuffer;
+		}
+
+		// Mix the dry and wet buffers together for the final output
+		for (int32 i = 0; i < NumSamples; ++i)
+		{
+			OutAudioBuffer[i] = (DryBuffer[i] + WetBuffer[i]);
+		}
+	}
+
+	void FMixerSubmix::SetOutputWetLevel(const float InWetLevel)
+	{
+		WetLevel = InWetLevel;
+	}
+
+	float FMixerSubmix::GetOutputWetLevel() const
+	{
+		return WetLevel;
+	}
+
+	int32 FMixerSubmix::GetSampleRate() const
+	{
+		return MixerDevice->GetDeviceSampleRate();
+	}
+
+	int32 FMixerSubmix::GetNumOutputChannels() const
+	{
+		return MixerDevice->GetNumDeviceChannels();
+	}
+
+	void FMixerSubmix::Update()
+	{
+	}
+}
\ No newline at end of file
diff --git a/Engine/Source/Runtime/AudioMixer/Private/DSP/AllPassFilter.cpp b/Engine/Source/Runtime/AudioMixer/Private/DSP/AllPassFilter.cpp
new file mode 100644
--- /dev/null
+++ b/Engine/Source/Runtime/AudioMixer/Private/DSP/AllPassFilter.cpp
@@ -0,0 +1,45 @@
+// Copyright 1998-2016 Epic Games, Inc. All Rights Reserved.
+
+#include "AudioMixerPCH.h"
+#include "DSP/AllPassFilter.h"
+
+namespace Audio
+{
+	FAllPassFilter::FAllPassFilter()
+	{
+	}
+
+	FAllPassFilter::~FAllPassFilter()
+	{
+	}
+
+	void FAllPassFilter::SetAPFGain(const float InGain)
+	{
+		APF_G = InGain;
+	}
+
+	float FAllPassFilter::operator()(const float InSample)
+	{
+		// If read and write index are identical then pass input to output
+		if (ReadIndex == WriteIndex)
+		{
+			TapIn(InSample);
+			return InSample;
+		}
+
+		// Using difference equations:
+		// w(n) = x(n) + g * w(n - D)
+		// y(n) = -g * w(n) + w(n - D)
+
+		// w(n - D) is the sample value at the current read index
+		const float Wn_D = Read();
+		const float Wn = InSample + APF_G * Wn_D;
+		const float Yn = -APF_G * Wn + Wn_D;
+
+		// Write Wn (not Yn) into the delay line
+		TapIn(Wn);
+
+		return UnderflowClamp(Yn);
+	}
+
+}
diff --git a/Engine/Source/Runtime/AudioMixer/Private/DSP/BiQuadFilter.cpp b/Engine/Source/Runtime/AudioMixer/Private/DSP/BiQuadFilter.cpp
new file mode 100644
--- /dev/null
+++ b/Engine/Source/Runtime/AudioMixer/Private/DSP/BiQuadFilter.cpp
@@ -0,0 +1,221 @@
+// Copyright 1998-2016 Epic Games, Inc. All Rights Reserved.
+
+#include "AudioMixerPCH.h"
+#include "DSP/BiQuadFilter.h"
+
+namespace Audio
+{
+	FBiquadFilter::FBiquadFilter()
+	{
+		FMemory::Memzero(this, sizeof(FBiquadFilter));
+	}
+
+	FBiquadFilter::FBiquadFilter(const EBiquadFilter::Type InFilterType, const float InCutoffFrequency, const float InQ, const float InGainDB)
+		: X1(0.0f)
+		, X2(0.0f)
+		, Y1(0.0f)
+		, Y2(0.0f)
+	{
+		SetParams(InFilterType, InGainDB, InCutoffFrequency, InQ);
+	}
+
+	FBiquadFilter::~FBiquadFilter()
+	{
+	}
+
+	void FBiquadFilter::SetParams(const EBiquadFilter::Type InFilterType, const float InGainDB, const float InCutoffFrequency, const float InQ)
+	{
+		FilterType = InFilterType;
+		GainDB = InGainDB;
+		CutoffFrequency = InCutoffFrequency;
+		Q = InQ;
+
+		float Omega = PI * CutoffFrequency;
+		float SinOmega = FMath::Sin(Omega);
+		float CosOmega = FMath::Cos(Omega);
+
+		static const float NaturalLog2 = 0.69314718056f;
+		float Alpha = 0.5f * SinOmega / Q;
+
+		// Temporary coefficients
+		float a0, a1, a2, b0, b1, b2;
+
+		switch (FilterType)
+		{
+		default:
+		case EBiquadFilter::None:
+			UE_LOG(LogTemp, Error, TEXT("Invalid Filter Parameter for FBiquadFilter"));
+			return;
+			break;
+
+		case EBiquadFilter::LowPass:
+		{
+			// H(s) = 1 / (s^2 + s/Q + 1)
+			b0 = 0.5f * (1.0f - CosOmega);
+			b1 = 1.0f - CosOmega;
+			b2 = b0;
+			a0 = 1.0f + Alpha;
+			a1 = -2.0f * CosOmega;
+			a2 = 1.0f - Alpha;
+		}
+		break;
+
+		case EBiquadFilter::HighPass:
+		{
+			// H(s) = s^2 / (s^2 + s/Q + 1)
+			b0 = 0.5f * (1.0f + CosOmega);
+			b1 = -1.0f - CosOmega;
+			b2 = b0;
+			a0 = 1.0f + Alpha;
+			a1 = -2.0f * CosOmega;
+			a2 = 1.0f - Alpha;
+		}
+		break;
+
+		case EBiquadFilter::BandPass:
+		{
+			// H(s) = (s / Q) / (s ^ 2 + s / Q + 1)      (constant 0 dB peak gain)
+			b0 = Alpha;
+			b1 = 0.0f;
+			b2 = -Alpha;
+			a0 = 1.0f + Alpha;
+			a1 = -2.0f * CosOmega;
+			a2 = 1.0f - Alpha;
+		}
+		break;
+
+		case EBiquadFilter::Notch:
+		{
+			// H(s) = (s^2 + 1) / (s^2 + s/Q + 1)
+			b0 = 1.0f;
+			b1 = -2.0f * CosOmega;
+			b2 = 1.0f;
+			a0 = 1.0f + Alpha;
+			a1 = -2.0f * CosOmega;
+			a2 = 1.0f - Alpha;
+		}
+		break;
+
+		case EBiquadFilter::AllPass:
+		{
+			// H(s) = (s^2 - s/Q + 1) / (s^2 + s/Q + 1)
+			b0 = 1.0f - Alpha;
+			b1 = -2.0f * CosOmega;
+			b2 = 1.0f + Alpha;
+			a0 = 1.0f + Alpha;
+			a1 = -2.0f * CosOmega;
+			a2 = 1.0f - Alpha;
+		}
+		break;
+
+		case EBiquadFilter::PeakEQ:
+		{
+			// H(s) = (s^2 + s*(A/Q) + 1) / (s^2 + s/(A*Q) + 1)
+			// Divide by 40 for peaking and shelving EQ filters
+			float A = FMath::Pow(10.0f, GainDB / 40.0f);
+			b0 = 1.0f + Alpha * A;
+			b1 = -2.0f * CosOmega;
+			b2 = 1.0f - Alpha * A;
+			a0 = 1.0f + Alpha / A;
+			a1 = -2.0f * CosOmega;
+			a2 = 1.0f - Alpha / A;
+		}
+		break;
+
+		case EBiquadFilter::LowShelf:
+		{
+			// H(s) = A * (s^2 + (sqrt(A)/Q)*s + A)/(A*s^2 + (sqrt(A)/Q)*s + 1)
+			// Divide by 40 for peaking and shelving EQ filters
+			float A = FMath::Pow(10.0f, GainDB / 40.0f);
+			float SqrtA = FMath::Sqrt(A);
+
+			b0 = A * ((A + 1.0f) - (A - 1.0f) * CosOmega + 2.0f * SqrtA * Alpha);
+			b1 = 2.0f * A * ((A - 1.0f) - (A + 1.0f) * CosOmega);
+			b2 = A * ((A + 1.0f) - (A - 1.0f) * CosOmega - 2.0f * SqrtA * Alpha);
+			a0 = (A + 1.0f) + (A - 1.0f) * CosOmega + 2.0f * SqrtA * Alpha;
+			a1 = -2.0f * ((A - 1.0f) + (A + 1.0f) * CosOmega);
+			a2 = (A + 1.0f) + (A - 1.0f) * CosOmega - 2.0f * SqrtA * Alpha;
+		}
+		break;
+
+		case EBiquadFilter::HighShelf:
+		{
+			float A = FMath::Pow(10.0f, GainDB / 40.0f);
+			float SqrtA = FMath::Sqrt(A);
+
+			b0 = A * ((A + 1.0f) + (A - 1.0f) * CosOmega + 2.0f * SqrtA * Alpha);
+			b1 = -2.0f * A * ((A - 1.0f) + (A + 1.0f) * CosOmega);
+			b2 = A * ((A + 1.0f) + (A - 1.0f) * CosOmega - 2.0f * SqrtA * Alpha);
+			a0 = (A + 1.0f) - (A - 1.0f) * CosOmega + 2.0f * SqrtA * Alpha;
+			a1 = -2.0f * ((A - 1.0f) - (A + 1.0f) * CosOmega);
+			a2 = (A + 1.0f) - (A - 1.0f) * CosOmega - 2.0f * SqrtA * Alpha;
+		}
+		break;
+		}
+
+		// Set the filter coefficients normalizing by a0
+		B0 = b0 / a0;
+		B1 = b1 / a0;
+		B2 = b2 / a0;
+		A1 = a1 / a0;
+		A2 = a2 / a0;
+	}
+
+	void FBiquadFilter::SetType(const EBiquadFilter::Type InFilterType)
+	{
+		if (FilterType != InFilterType)
+		{
+			SetParams(InFilterType, GainDB, CutoffFrequency, Q);
+		}
+	}
+
+	void FBiquadFilter::SetGain(const float InGainDB)
+	{
+		if (GainDB != InGainDB)
+		{
+			SetParams(FilterType, InGainDB, CutoffFrequency, Q);
+		}
+	}
+
+	void FBiquadFilter::SetCutoffFrequency(const float InCutoffFrequency)
+	{
+		if (CutoffFrequency != InCutoffFrequency)
+		{
+			SetParams(FilterType, GainDB, InCutoffFrequency, Q);
+		}
+	}
+
+	void FBiquadFilter::SetQ(const float InQ)
+	{
+		if (Q != InQ)
+		{
+			SetParams(FilterType, GainDB, CutoffFrequency, InQ);
+		}
+	}
+
+	void FBiquadFilter::SetBandwidth(const float InBandwidth)
+	{
+		if (Bandwidth != InBandwidth)
+		{
+			Bandwidth = InBandwidth;
+			const float InQ = GetQFromBandWidth(InBandwidth, CutoffFrequency);
+			SetQ(InQ);
+		}
+	}
+
+	FBiquadFilter FBiquadFilter::CreateLowPassFilter(const float InCutoffFrequency, const float InQ)
+	{
+		return FBiquadFilter(EBiquadFilter::LowPass, InCutoffFrequency, InQ);
+	}
+
+	FBiquadFilter FBiquadFilter::CreateHighPassFilter(const float InCutoffFrequency, const float InQ)
+	{
+		return FBiquadFilter(EBiquadFilter::HighPass, InCutoffFrequency, InQ);
+	}
+
+	FBiquadFilter FBiquadFilter::CreateBandPassFilter(const float InCutoffFrequency, const float InQ)
+	{
+		return FBiquadFilter(EBiquadFilter::BandPass, InCutoffFrequency, InQ);
+	}
+
+}
diff --git a/Engine/Source/Runtime/AudioMixer/Private/DSP/Chorus.cpp b/Engine/Source/Runtime/AudioMixer/Private/DSP/Chorus.cpp
new file mode 100644
--- /dev/null
+++ b/Engine/Source/Runtime/AudioMixer/Private/DSP/Chorus.cpp
@@ -0,0 +1,28 @@
+// Copyright 1998-2016 Epic Games, Inc. All Rights Reserved.
+
+#include "AudioMixerPCH.h"
+#include "DSP/Chorus.h"
+
+namespace Audio
+{
+	FChorus::FChorus()
+	{
+	}
+
+	FChorus::~FChorus()
+	{
+	}
+
+	void FChorus::Init(int32 InSampleRate)
+	{
+		// Call super class init
+		FModulatedDelay::Init(InSampleRate);
+
+		// Set delay time values and modulation params for flangers
+		MinDelayTimeMsec = 5.0f;
+		MaxDelayTimeMsec = 30.0f;
+		SetModulationFeedback(0.5f);
+		SetModulationDelayOffsetMsec(15.0f);
+		SetWetLevel(0.5f);
+	}
+}
diff --git a/Engine/Source/Runtime/AudioMixer/Private/DSP/Delay.cpp b/Engine/Source/Runtime/AudioMixer/Private/DSP/Delay.cpp
new file mode 100644
--- /dev/null
+++ b/Engine/Source/Runtime/AudioMixer/Private/DSP/Delay.cpp
@@ -0,0 +1,265 @@
+// Copyright 1998-2016 Epic Games, Inc. All Rights Reserved.
+
+#include "AudioMixerPCH.h"
+#include "DSP/Delay.h"
+
+namespace Audio
+{
+	FDelay::FDelay()
+		: Feedback(0.0f)
+		, DelaySamples(0.0f)
+		, OutputGain(0.0f)
+		, ReadIndex(0)
+		, WriteIndex(0)
+		, SampleRate(0)
+	{
+	}
+
+	FDelay::~FDelay()
+	{
+	}
+
+	void FDelay::Init(const int32 InSampleRate, const int32 InDelaySizeMsec)
+	{
+		AUDIO_MIXER_CHECK(InSampleRate != 0);
+		SampleRate = InSampleRate;
+
+		// Get the number of samples based on sample rate
+		const int32 NumSamples = InSampleRate * ((float)InDelaySizeMsec / 1000.0f);
+		AudioBuffer.Reset();
+		AudioBuffer.AddDefaulted(NumSamples);
+
+		// Assume the current delay is the entire buffer size
+		DelaySamples = NumSamples;
+	}
+
+	void FDelay::Reset()
+	{
+		AUDIO_MIXER_CHECK(SampleRate != 0);
+
+		const int32 BufferSize = AudioBuffer.Num();
+		AudioBuffer.Reset();
+		AudioBuffer.AddDefaulted(BufferSize);
+
+		// Reset the write index
+		WriteIndex = 0;
+
+		// Set the read index based on the current delay in samples
+		ReadIndex = -(int32)DelaySamples;
+
+		while (ReadIndex < 0)
+		{
+			ReadIndex += BufferSize;
+		}
+	}
+
+	void FDelay::SetFeedback(const float InFeedback)
+	{
+		AUDIO_MIXER_CHECK(InFeedback >= -1.0f && InFeedback <= 1.0f);
+		Feedback = InFeedback;
+	}
+
+	void FDelay::SetOutputGainDB(const float InOutputGainDB)
+	{
+		AUDIO_MIXER_CHECK(InOutputGainDB <= 0.0f);
+
+		OutputGain = FMath::Pow(10.0f, InOutputGainDB / 20.0f);
+	}
+
+	void FDelay::SetOutputGain(const float InOutputGain)
+	{
+		AUDIO_MIXER_CHECK(InOutputGain >= 0.0f && InOutputGain <= 1.0f);
+
+		OutputGain = InOutputGain;
+	}
+
+	void FDelay::SetDelayMsec(const float InDelayMsec)
+	{
+		AUDIO_MIXER_CHECK(SampleRate != 0);
+
+		// Set the current delay in terms of fractional samples
+		DelaySamples = InDelayMsec * ((float)SampleRate / 1000.0f);
+
+		// Update the read index based on the current write index
+		// Don't touch the current write index, this will allow modulated delay lines.
+		ReadIndex = WriteIndex - (int32)DelaySamples;
+
+		// Make sure to wrap the read index to within the bounds of the delay line
+		while (ReadIndex < 0)
+		{
+			ReadIndex += AudioBuffer.Num();
+		}
+	}
+
+	float FDelay::Read() const
+	{
+		// Get the previous read index, wrap to buffer size
+		const int32 ReadIndexPrev = (ReadIndex == 0) ? AudioBuffer.Num() - 1 : ReadIndex - 1;
+
+		// Get the output at the current read index
+		const float Yn = AudioBuffer[ReadIndex];
+
+		// Get sample value at previous read index
+		const float YnPrev = AudioBuffer[ReadIndexPrev];
+
+		// Get the lerp alpha for interpolating sample reading
+		const float Alpha = DelaySamples - (int32)DelaySamples;
+
+		// Lerp from current read to prev read output
+		return FMath::Lerp(Yn, YnPrev, Alpha);
+	}
+
+	float FDelay::TapOut(const float TapOutMsec) const
+	{
+		const float TapDelaySamples = TapOutMsec * ((float)SampleRate / 1000.0f);
+		int32 TapReadIndex = WriteIndex - (int32)TapDelaySamples;
+		while (TapReadIndex < 0)
+		{
+			TapReadIndex += AudioBuffer.Num();
+		}
+
+		// Get the previous read index, wrap to buffer size
+		const int32 TapReadIndexPrev = (TapReadIndex == 0) ? AudioBuffer.Num() - 1 : TapReadIndex - 1;
+
+		// Get current and previous tap read sample values
+		const float Yn = AudioBuffer[ReadIndex];
+		const float YnPrev = AudioBuffer[TapReadIndexPrev];
+
+		// Get lerp alpha
+		const float Alpha = TapDelaySamples - (int32)TapDelaySamples;
+
+		// Lerp from current read to prev read indices
+		return FMath::Lerp(Yn, YnPrev, Alpha);
+	}
+
+	void FDelay::TapIn(const float InSample)
+	{
+		AudioBuffer[WriteIndex++] = InSample;
+
+		if (WriteIndex >= AudioBuffer.Num())
+		{
+			WriteIndex = 0;
+		}
+
+		++ReadIndex;
+
+		if (ReadIndex >= AudioBuffer.Num())
+		{
+			ReadIndex = 0;
+		}
+	}
+
+	float FDelay::operator()(const float InSample)
+	{
+		// If there is no delay, then copy input to output
+		const float Yn = (DelaySamples == 0.0f) ? InSample : Read();
+
+		// Write in the new sample into the delay line with possible feedback
+		TapIn(InSample + Feedback * Yn);
+
+		// Return the output applied with the gain
+		return OutputGain * Yn;
+	}
+
+	/**
+	* FModulatedDelay
+	*/
+	FModulatedDelay::FModulatedDelay()
+		: CurrentWaveTable(&SinLFO)
+		, ModulationDepth(0.0f)
+		, ModulationFrequency(0.0f)
+		, MinDelayTimeMsec(0.0f)
+		, MaxDelayTimeMsec(0.0f)
+		, DelayOffsetMsec(0.0f)
+		, ModulationFeedback(0.0f)
+		, WetLevel(1.0f)
+		, bIsInQuadPhase(false)
+	{
+	}
+
+	FModulatedDelay::~FModulatedDelay()
+	{
+	}
+
+	void FModulatedDelay::Init(int32 InSampleRate)
+	{
+		SawLFO.Init(InSampleRate);
+		SinLFO.Init(InSampleRate);
+		Delay.Init(InSampleRate);
+	}
+
+	void FModulatedDelay::SetLFO(const EWaveTable::Type InLFOType)
+	{
+		switch (InLFOType)
+		{
+		default:
+		case EWaveTable::SineWaveTable:
+			CurrentWaveTable = &SinLFO;
+			break;
+		case EWaveTable::SawWaveTable:
+			CurrentWaveTable = &SawLFO;
+			break;
+		case EWaveTable::TriangleWaveTable:
+			CurrentWaveTable = &TriLFO;
+			break;
+		case EWaveTable::SquareWaveTable:
+			CurrentWaveTable = &SquareLFO;
+			break;
+		}
+	}
+
+	void FModulatedDelay::SetModulationFeedback(const float InModulationFeedback)
+	{
+		Delay.SetFeedback(InModulationFeedback);
+	}
+
+	void FModulatedDelay::SetWetLevel(const float InWetLevel)
+	{
+		AUDIO_MIXER_CHECK(InWetLevel >= 0.0f && InWetLevel <= 1.0f);
+		WetLevel = InWetLevel;
+	}
+
+	void FModulatedDelay::SetMinAndMaxDelayMSec(const float InMinDelayTimeMsec, const float InMaxDelayTimeMsec)
+	{
+		AUDIO_MIXER_CHECK(InMinDelayTimeMsec >= 0.0f && InMaxDelayTimeMsec <= 1.0f);
+		AUDIO_MIXER_CHECK(InMinDelayTimeMsec < InMaxDelayTimeMsec);
+
+		MinDelayTimeMsec = InMinDelayTimeMsec;
+		MaxDelayTimeMsec = InMaxDelayTimeMsec;
+	}
+
+	void FModulatedDelay::SetModulationDelayOffsetMsec(const float InDelayOffsetMsec)
+	{
+		AUDIO_MIXER_CHECK(InDelayOffsetMsec >= 0.0f);
+		DelayOffsetMsec = InDelayOffsetMsec;
+	}
+
+	void FModulatedDelay::SetModulationFrequency(const float InModulationFrequency)
+	{
+		SawLFO.SetFrequency(ModulationFrequency);
+		SinLFO.SetFrequency(ModulationFrequency);
+	}
+
+	float FModulatedDelay::operator()(const float InSample)
+	{
+		// Get the current value of our LFO wavetable
+		float Yn = 0.0f;			// Normal phase output
+		float Yqn = 0.0f;			// Quad phase output
+
+		(*CurrentWaveTable)(&Yn, &Yqn);
+
+		// Use modulation depth and min/max delay time range to compute a new delay time 
+		const float LFOSample = bIsInQuadPhase ? Yqn : Yn;
+		const float NewDelayMsec = ModulationDepth * LFOSample * (MaxDelayTimeMsec - MinDelayTimeMsec) + MinDelayTimeMsec + DelayOffsetMsec;
+
+		// Set the new delay time on our delay line
+		Delay.SetDelayMsec(NewDelayMsec);
+
+		// Feed the input sample through the now modulated delay line 
+		const float OutDelay = Delay(InSample);
+		const float OutputSample = OutDelay * WetLevel + InSample * (1.0f - WetLevel);
+
+		return OutputSample;
+	}
+
+}
diff --git a/Engine/Source/Runtime/AudioMixer/Private/DSP/Flanger.cpp b/Engine/Source/Runtime/AudioMixer/Private/DSP/Flanger.cpp
new file mode 100644
--- /dev/null
+++ b/Engine/Source/Runtime/AudioMixer/Private/DSP/Flanger.cpp
@@ -0,0 +1,78 @@
+// Copyright 1998-2016 Epic Games, Inc. All Rights Reserved.
+
+#include "AudioMixerPCH.h"
+#include "DSP/Flanger.h"
+
+namespace Audio
+{
+	FFlanger::FFlanger()
+	{
+	}
+
+	FFlanger::~FFlanger()
+	{
+	}
+
+	void FFlanger::Init(int32 InSampleRate)
+	{
+		// Call super class init
+		FModulatedDelay::Init(InSampleRate);
+
+		// Set delay time values and modulation params for flangers
+		MinDelayTimeMsec = 0.0f;
+		MaxDelayTimeMsec = 10.0f;
+		SetModulationFeedback(0.5f);
+		SetWetLevel(0.5f);
+	}
+
+	/** 
+	* FStereoFlanger
+	*/
+	FStereoFlanger::FStereoFlanger()
+	{
+	}
+
+	FStereoFlanger::~FStereoFlanger()
+	{
+	}
+
+	void FStereoFlanger::Init(int32 InSampleRate)
+	{
+		Flanger[0].Init(InSampleRate);
+		Flanger[1].Init(InSampleRate);
+
+		// Make sure both flangers are in quad phase w/ each other.
+		Flanger[0].SetQuadraturePhase(true);
+		Flanger[1].SetQuadraturePhase(false);
+	}
+
+	void FStereoFlanger::SetLFO(const EWaveTable::Type InLFOType)
+	{
+		Flanger[0].SetLFO(InLFOType);
+		Flanger[1].SetLFO(InLFOType);
+	}
+
+	void FStereoFlanger::SetModulationFeedback(const float InModulationFeedback)
+	{
+		Flanger[0].SetModulationFeedback(InModulationFeedback);
+		Flanger[1].SetModulationFeedback(InModulationFeedback);
+	}
+
+	void FStereoFlanger::SetModulationMinAndMaxDelayMSec(const float InMinDelayTimeMsec, const float InMaxDelayTimeMsec)
+	{
+		Flanger[0].SetMinAndMaxDelayMSec(InMinDelayTimeMsec, InMaxDelayTimeMsec);
+		Flanger[1].SetMinAndMaxDelayMSec(InMinDelayTimeMsec, InMaxDelayTimeMsec);
+	}
+
+	void FStereoFlanger::SetModulationFrequency(const float InModulationFrequency)
+	{
+		Flanger[0].SetModulationFrequency(InModulationFrequency);
+		Flanger[1].SetModulationFrequency(InModulationFrequency);
+	}
+
+	void FStereoFlanger::operator()(const float InLeft, const float InRight, float* OutLeft, float* OutRight)
+	{
+		*OutLeft = Flanger[0](InLeft);
+		*OutRight = Flanger[1](InRight);
+	}
+}
diff --git a/Engine/Source/Runtime/AudioMixer/Private/DSP/Noise.cpp b/Engine/Source/Runtime/AudioMixer/Private/DSP/Noise.cpp
new file mode 100644
--- /dev/null
+++ b/Engine/Source/Runtime/AudioMixer/Private/DSP/Noise.cpp
@@ -0,0 +1,80 @@
+// Copyright 1998-2016 Epic Games, Inc. All Rights Reserved.
+
+#pragma once
+
+#include "AudioMixerPCH.h"
+#include "DSP/Noise.h"
+
+namespace Audio
+{
+	FWhiteNoise::FWhiteNoise(const float InScale, const float InAdd)
+		: Scale(InScale)
+		, Add(InAdd)
+	{
+
+	}
+
+	void FWhiteNoise::SetScaleAdd(const float InScale, const float InAdd)
+	{
+		Scale = InScale;
+		Add = InAdd;
+	}
+
+	float FWhiteNoise::operator()()
+	{
+		return Add + Scale * FMath::FRandRange(-1.0f, 1.0f);
+	}
+
+	FPinkNoise::FPinkNoise(const float InScale, const float InAdd)
+	{
+		InitFilter();
+		Noise.SetScaleAdd(InScale, InAdd);
+	}
+
+	void FPinkNoise::InitFilter()
+	{
+		for (int32 i = 0; i < 4; ++i)
+		{
+			X[i] = 0.0f;
+			Y[i] = 0.0f;
+		}
+
+		A[0] = 0.0f;
+		A[1] = -2.495f;
+		A[2] = 2.017;
+		A[3] = -0.522f;
+
+		B[0] = 0.041f;
+		B[1] = -0.96f;
+		B[2] = 0.051f;
+		B[3] = -0.004f;
+	}
+
+	/** Sets the output scale and add parameter. */
+	void FPinkNoise::SetScaleAdd(const float InScale, const float InAdd)
+	{
+		Noise.SetScaleAdd(InScale, InAdd);
+	}
+
+	float FPinkNoise::operator()()
+	{
+		X[0] = Noise();
+
+		float Output = 0.0f;
+		for (int32 i = 0; i < 4; ++i)
+		{
+			Output += (B[i] * X[i] - A[i] * Y[i]);
+		}
+
+		X[3] = X[2];
+		X[2] = X[1];
+		X[1] = X[0];
+
+		Y[3] = Y[2];
+		Y[2] = Y[1];
+		Y[1] = Y[0];
+		Y[0] = Output;
+
+		return Output;
+	}
+}
