Copying //UE4/Dev-Framework to //UE4/Dev-Main (Source: //UE4/Dev-Framework @ 3136612)

diff --git a/Engine/Source/Runtime/AudioMixer/Private/DSP/ParamInterpolator.cpp b/Engine/Source/Runtime/AudioMixer/Private/DSP/ParamInterpolator.cpp
new file mode 100644
--- /dev/null
+++ b/Engine/Source/Runtime/AudioMixer/Private/DSP/ParamInterpolator.cpp
@@ -0,0 +1,66 @@
+// Copyright 1998-2016 Epic Games, Inc. All Rights Reserved.
+
+#include "AudioMixerPCH.h"
+#include "DSP/ParamInterpolator.h"
+
+namespace Audio
+{
+	/**
+	* FParamInterpolator implementation
+	*/
+
+	FParamInterpolator::FParamInterpolator()
+		: StartValue(0.0f)
+		, EndValue(0.0f)
+		, CurrentTick(1.0f)
+		, NumTicksToNewValue(1.0f)
+	{
+	}
+
+	FParamInterpolator::~FParamInterpolator()
+	{
+	}
+
+	void FParamInterpolator::InitValue(const float InValue)
+	{
+		StartValue = InValue;
+		EndValue = InValue;
+		CurrentTick = 1.0f;
+		NumTicksToNewValue = 1.0f;
+	}
+
+	void FParamInterpolator::SetValue(const float InValue, const uint32 InNumInterpolationTicks)
+	{
+		EndValue = InValue;
+		if (InNumInterpolationTicks == 0)
+		{
+			NumTicksToNewValue = 1.0f;
+			CurrentTick = 1.0f;
+		}
+		else
+		{
+			NumTicksToNewValue = (float)InNumInterpolationTicks;
+			CurrentTick = 0.0f;
+		}
+		check(NumTicksToNewValue > 0.0f);
+	}
+
+	float FParamInterpolator::GetValue() const
+	{
+		float Alpha = FMath::Min(1.0f, CurrentTick / NumTicksToNewValue);
+		float Result = FMath::Lerp(StartValue, EndValue, Alpha);
+		return Result;
+	}
+
+	float FParamInterpolator::operator()()
+	{
+		if (CurrentTick >= NumTicksToNewValue)
+		{
+			return EndValue;
+		}
+
+		float Result = GetValue();
+		CurrentTick += 1.0f;
+		return Result;
+	}
+}
diff --git a/Engine/Source/Runtime/AudioMixer/Private/DSP/ParametricEqualizer.cpp b/Engine/Source/Runtime/AudioMixer/Private/DSP/ParametricEqualizer.cpp
new file mode 100644
--- /dev/null
+++ b/Engine/Source/Runtime/AudioMixer/Private/DSP/ParametricEqualizer.cpp
@@ -0,0 +1,37 @@
+// Copyright 1998-2016 Epic Games, Inc. All Rights Reserved.
+
+
+#include "AudioMixerPCH.h"
+#include "DSP/ParametricEqualizer.h"
+
+namespace Audio
+{
+	FParametricEqualizer::FParametricEqualizer()
+	{
+	}
+
+	FParametricEqualizer::~FParametricEqualizer()
+	{
+	}
+
+	void FParametricEqualizer::SetParameters(const TArray<FEqualizerBandParams>& InParams)
+	{
+		Filters.Reset();
+		for (int32 i = 0; i < InParams.Num(); ++i)
+		{
+			const float Q = GetQFromBandWidth(InParams[i].Bandwidth, InParams[i].CutoffFrequency);
+			Filters.Add(FBiquadFilter(EBiquadFilter::PeakEQ, InParams[i].CutoffFrequency, Q, InParams[i].Gain));
+		}
+	}
+
+	bool FParametricEqualizer::SetBand(const int32 InBand, const FEqualizerBandParams& InParams)
+	{
+		if (InBand >= 0 && InBand < Filters.Num())
+		{
+			const float Q = GetQFromBandWidth(InParams.Bandwidth, InParams.CutoffFrequency);
+			Filters[InBand].SetParams(EBiquadFilter::PeakEQ, InParams.CutoffFrequency, Q, InParams.Gain);
+			return true;
+		}
+		return false;
+	}
+}
diff --git a/Engine/Source/Runtime/AudioMixer/Private/DSP/Reverb.cpp b/Engine/Source/Runtime/AudioMixer/Private/DSP/Reverb.cpp
new file mode 100644
--- /dev/null
+++ b/Engine/Source/Runtime/AudioMixer/Private/DSP/Reverb.cpp
@@ -0,0 +1,238 @@
+// Copyright 1998-2016 Epic Games, Inc. All Rights Reserved.
+
+#pragma once
+
+#include "AudioMixerPCH.h"
+#include "DSP/Reverb.h"
+
+namespace Audio
+{
+	const float FFreeVerb::FixedGain = 0.015f;
+	const float FFreeVerb::ScaleWet = 3.0f;
+	const float FFreeVerb::ScaleDry = 2.0f;
+	const float FFreeVerb::ScaleDamp = 0.4f;
+	const float FFreeVerb::ScaleRoom = 0.28f;
+	const float FFreeVerb::OffsetRoom = 0.7;
+
+	FFreeVerb::FFreeVerb(const float InSampleRate)
+	{
+		// Hard-coded delay lines tuned to 44.1 kHz sample rate
+		CombFilterDelayLengths.AddDefaulted(NumCombFilters);
+
+		const float OneOverSamplesPerMs = 1000.0f / InSampleRate;
+
+		CombFilterDelayLengths[0] = OneOverSamplesPerMs * 1557.0f;
+		CombFilterDelayLengths[1] = OneOverSamplesPerMs * 1617.0f;
+		CombFilterDelayLengths[2] = OneOverSamplesPerMs * 1491.0f;
+		CombFilterDelayLengths[3] = OneOverSamplesPerMs * 1422.0f;
+		CombFilterDelayLengths[4] = OneOverSamplesPerMs * 1277.0f;
+		CombFilterDelayLengths[5] = OneOverSamplesPerMs * 1356.0f;
+		CombFilterDelayLengths[6] = OneOverSamplesPerMs * 1188.0f;
+		CombFilterDelayLengths[7] = OneOverSamplesPerMs * 1116.0f;
+
+		AllpassFilterDelayLengths.AddDefaulted(NumAllPassFilters);
+		AllpassFilterDelayLengths[0] = OneOverSamplesPerMs * 225.0f;
+		AllpassFilterDelayLengths[1] = OneOverSamplesPerMs * 556.0f;
+		AllpassFilterDelayLengths[2] = OneOverSamplesPerMs * 441.0f;
+		AllpassFilterDelayLengths[3] = OneOverSamplesPerMs * 341.0f;
+
+		// Scale these delay line lengths according to the actual sample rate
+		float SampleRateScale = InSampleRate / 44100.0f;
+		for (int32 i = 0; i < NumCombFilters; ++i)
+		{
+			CombFilterDelayLengths[i] = FMath::FloorToInt(SampleRateScale * CombFilterDelayLengths[i]);
+		}
+
+		for (int32 i = 0; i < NumAllPassFilters; ++i)
+		{
+			AllpassFilterDelayLengths[i] = FMath::FloorToInt(SampleRateScale * AllpassFilterDelayLengths[i]);
+		}
+
+		CombDelayLeft.AddDefaulted(NumCombFilters);
+		CombDelayRight.AddDefaulted(NumCombFilters);
+		for (int32 i = 0; i < NumCombFilters; ++i)
+		{
+			CombDelayLeft[i].Init(InSampleRate, CombFilterDelayLengths[i]);
+			CombDelayLeft[i].SetDelayMsec(CombFilterDelayLengths[i]);
+			CombDelayRight[i].Init(InSampleRate, CombFilterDelayLengths[i]);
+			CombDelayRight[i].SetDelayMsec(CombFilterDelayLengths[i]);
+		}
+
+		AllPassDelayLeft.AddDefaulted(NumAllPassFilters);
+		AllPassDelayRight.AddDefaulted(NumAllPassFilters);
+		for (int32 i = 0; i < NumAllPassFilters; ++i)
+		{
+			AllPassDelayRight[i].Init(InSampleRate, CombFilterDelayLengths[i]);
+			AllPassDelayRight[i].SetDelayMsec(CombFilterDelayLengths[i]);
+			AllPassDelayLeft[i].Init(InSampleRate, CombFilterDelayLengths[i]);
+			AllPassDelayLeft[i].SetDelayMsec(CombFilterDelayLengths[i]);
+		}
+
+		SetWetLevel(0.75f);
+
+		// Feedback attenuation in low pass feedback comb filter
+		RoomSizeSetting = (0.75f * ScaleRoom) + OffsetRoom;
+		DampSetting = -0.25f * ScaleDamp;
+
+		Width = 1.0f;
+		bIsFrozen = false;
+
+		CombLowPassLeft.AddDefaulted(NumCombFilters);
+		CombLowPassRight.AddDefaulted(NumCombFilters);
+
+		UpdateParameters();
+
+		AllPassCoefficientG = 0.5f;
+	}
+
+	void FFreeVerb::SetWetLevel(const float InWetLevel)
+	{
+		WetLevel = InWetLevel;
+		UpdateParameters();
+	}
+
+	void FFreeVerb::SetRoomSize(const float InRoomSize)
+	{
+		RoomSizeSetting = (InRoomSize * ScaleRoom) + OffsetRoom;
+		UpdateParameters();
+	}
+
+	float FFreeVerb::GetRoomSize() const
+	{
+		return (RoomSizeSetting - OffsetRoom) / ScaleRoom;
+	}
+
+	void FFreeVerb::SetDamping(const float InDamping)
+	{
+		DampSetting = InDamping * ScaleDamp;
+		UpdateParameters();
+	}
+
+	float FFreeVerb::GetDamping() const
+	{
+		return DampSetting / ScaleDamp;
+	}
+
+	void FFreeVerb::SetStereoWidth(float InWidth)
+	{
+		Width = InWidth;
+		UpdateParameters();
+	}
+
+	float FFreeVerb::GetStereoWidth() const
+	{
+		return Width;
+	}
+
+	void FFreeVerb::SetFreeze(const bool bInIsFrozen)
+	{
+		bIsFrozen = bInIsFrozen;
+		UpdateParameters();
+	}
+
+	bool FFreeVerb::IsFrozen() const
+	{
+		return bIsFrozen;
+	}
+
+	void FFreeVerb::Clear()
+	{
+		for (int32 i = 0; i < NumCombFilters; ++i)
+		{
+			CombDelayLeft[i].Reset();
+			CombDelayRight[i].Reset();
+		}
+
+		for (int32 i = 0; i < NumAllPassFilters; ++i)
+		{
+			AllPassDelayLeft[i].Reset();
+			AllPassDelayRight[i].Reset();
+		}
+	}
+
+	void FFreeVerb::UpdateParameters()
+	{
+		float ActualWet = ScaleWet * WetLevel;
+		Dry = ScaleDry * (1.0f - Wet1);
+
+		ActualWet /= (ActualWet + Dry);
+		Dry /= (ActualWet + Dry);
+
+		Wet1 = ActualWet * (0.5f * Width + 0.5f);
+		Wet2 = ActualWet * (-0.5f * Width + 0.5f);
+
+		if (bIsFrozen)
+		{
+			RoomSize = 1.0f;
+			Damp = 0.0f;
+			Gain = 0.0f;
+		}
+		else
+		{
+			RoomSize = RoomSizeSetting;
+			Damp = DampSetting;
+			Gain = FixedGain;
+		}
+
+		// Set the low pass filters
+		for (int i = 0; i < NumCombFilters; ++i)
+		{
+			CombLowPassLeft[i].SetCoefficients(1.0f - Damp, -Damp);
+			CombLowPassRight[i].SetCoefficients(1.0f - Damp, -Damp);
+		}
+	}
+
+	void FFreeVerb::operator()(const TArray<float>& InAudio, TArray<float>& OutAudio)
+	{
+		// Sum left and right channels to a mono channel
+		const float InputSum = (InAudio[0] + InAudio[1]) * Gain;
+
+		float OutputLeft = 0.0f;
+		float OutputRight = 0.0f;
+
+		// Process the bank of parallel lowpass-feedback comb filters
+		for (int32 i = 0; i < NumCombFilters; ++i)
+		{
+			// Get the comb filter's next sample value
+			float DelayNextOut = CombDelayLeft[i].Read();
+
+			// Feed it through the LPF, scale with the room size add to the input
+			float Yn = InputSum + (RoomSize * CombLowPassLeft[i](DelayNextOut));
+
+			// Add the sample to the comb filter delay line
+			CombDelayLeft[i](Yn);
+
+			// Add the result of this comb filter to the left channel output
+			OutputLeft += Yn;
+
+			DelayNextOut = CombDelayRight[i].Read();
+
+			Yn = InputSum + (RoomSize * CombLowPassRight[i](DelayNextOut));
+
+			CombDelayRight[i](Yn);
+
+			OutputRight += Yn;
+		}
+
+		// Now process the series of all pass filters
+		for (int32 i = 0; i < NumAllPassFilters; ++i)
+		{
+			// Compute left channel output
+			float NextAllpassOutput = AllPassDelayLeft[i].Read();
+			float AllpassInput = OutputLeft + AllPassCoefficientG * NextAllpassOutput;
+			AllPassDelayLeft[i](AllpassInput);
+			OutputLeft = -AllpassInput + (1.0f + AllPassCoefficientG) * NextAllpassOutput;
+
+			// Compute right channel output
+			NextAllpassOutput = AllPassDelayRight[i].Read();
+			AllpassInput = OutputRight + (AllPassCoefficientG * NextAllpassOutput);
+			AllPassDelayRight[i](AllpassInput);
+			OutputRight = -AllpassInput + (1.0f + AllPassCoefficientG) * NextAllpassOutput;
+		}
+
+		// Now blend the left/right channels to the final output
+		OutAudio[0] = OutputLeft * Wet1 + OutputRight * Wet2 + InAudio[0] * Dry;
+		OutAudio[1] = OutputRight * Wet1 + OutputLeft * Wet2 + InAudio[1] * Dry;
+	}
+
+}
diff --git a/Engine/Source/Runtime/AudioMixer/Private/DSP/WaveTableOsc.cpp b/Engine/Source/Runtime/AudioMixer/Private/DSP/WaveTableOsc.cpp
new file mode 100644
--- /dev/null
+++ b/Engine/Source/Runtime/AudioMixer/Private/DSP/WaveTableOsc.cpp
@@ -0,0 +1,291 @@
+// Copyright 1998-2016 Epic Games, Inc. All Rights Reserved.
+
+
+#include "AudioMixerPCH.h"
+#include "DSP/WaveTableOsc.h"
+
+namespace Audio
+{
+
+	FWaveTableOsc::FWaveTableOsc()
+		: SampleRate(0)
+		, ReadIndex(0.0f)
+		, QuadPhaseReadIndex(0.0f)
+		, ReadDelta(0.0f)
+	{
+	}
+
+	FWaveTableOsc::~FWaveTableOsc()
+	{
+	}
+
+	void FWaveTableOsc::Init(int32 InSampleRate, int32 TableSize)
+	{
+		AUDIO_MIXER_CHECK(InSampleRate > 0);
+
+		SampleRate = InSampleRate;
+		
+		Generate(TableSize);
+	}
+
+	void FWaveTableOsc::SetFrequency(const float InFrequencyHz)
+	{
+		AUDIO_MIXER_CHECK(InFrequencyHz > 0);
+		ReadDelta = (InFrequencyHz / (float)SampleRate) * (float)WaveTable.Num();
+	}
+
+	void FWaveTableOsc::SetPolarity(const bool bInIsUnipolar)
+	{
+		bIsUnipolar = bInIsUnipolar;
+	}
+
+	void FWaveTableOsc::operator()(float* OutSample, float* OutQuadSample)
+	{
+		const int32 CurrIndex = (int32)ReadIndex;
+		const int32 NextIndex = (CurrIndex + 1) % WaveTable.Num();
+
+		const int32 CurrQuadIndex = (int32)QuadPhaseReadIndex;
+		const int32 NextQuadIndex = (CurrQuadIndex + 1) % WaveTable.Num();
+
+		const float Alpha = ReadIndex - CurrIndex;
+		const float QuadAlpha = QuadPhaseReadIndex - CurrQuadIndex;
+
+		*OutSample = FMath::Lerp(WaveTable[CurrIndex], WaveTable[NextIndex], Alpha);
+		*OutQuadSample = FMath::Lerp(WaveTable[CurrQuadIndex], WaveTable[NextQuadIndex], Alpha);
+
+		if (bIsUnipolar)
+		{
+			// Scale from [-1.0, 1.0] to [0.0 1.0]
+			*OutSample = 0.5f * (*OutSample + 1.0f);
+			*OutQuadSample = 0.5f * (*OutQuadSample + 1.0f);
+		}
+
+		ReadIndex += ReadDelta;
+		QuadPhaseReadIndex += ReadIndex;
+
+		// Wrap the read indices to within the size of the table
+		while (ReadIndex >= (float)WaveTable.Num())
+		{
+			ReadIndex -= (float)WaveTable.Num();
+		}
+
+		while (QuadPhaseReadIndex >= (float)WaveTable.Num())
+		{
+			QuadPhaseReadIndex -= (float)WaveTable.Num();
+		}
+	}
+
+	/*
+	* FSineWaveTable
+	*/
+	void FSineWaveTable::Generate(const int32 InTableSize)
+	{
+		WaveTable.Reset(InTableSize);
+
+		for (int32 i = 0; i < InTableSize; ++i)
+		{
+			WaveTable.Add(FMath::Sin(2.0f * PI * (float)i / InTableSize));
+		}
+	}
+
+	/*
+	* FSawWaveTable
+	*/
+	void FSawWaveTable::Generate(const int32 InTableSize)
+	{
+		WaveTable.Reset(InTableSize);
+
+		/*      
+				 /|   /|
+		        / |  / |
+		       /  | /  |
+		      /   |/   |
+		*/
+
+		// Using y = mx + b slope/intercept form
+		const int32 HalfPoint = 0.5f * InTableSize;
+		const float Slope = 1.0f / HalfPoint;
+
+		for (int32 i = 0; i < InTableSize; ++i)
+		{
+			// Rise if less than half point
+			if (i < HalfPoint)
+			{
+				WaveTable.Add(Slope * i);
+			}
+			else
+			{
+				WaveTable.Add(Slope * (i - HalfPoint - 1) - 1.0f);
+			}
+		}
+	}
+
+	/*
+	* FBandLimitedSawWaveTable
+	*/
+	void FBandLimitedSawWaveTable::Generate(const int32 InTableSize)
+	{
+		WaveTable.Reset(InTableSize);
+		WaveTable.AddDefaulted(InTableSize);
+
+		float MaxValue = 0.0f;
+
+		// Generate according to sum of sines, Fourier theorem
+		for (int32 i = 0; i < InTableSize; ++i)
+		{
+			const float TwoPiFraction = 2.0f * PI * (float)i / InTableSize;
+
+			for (int32 k = 1; k <= 6; ++k)
+			{
+				const float SinusoidComponent = FMath::Sin(TwoPiFraction * k);
+				const float Sign = FMath::Pow(-1.0f, (float)(k + 1));
+				const float Amplitude = (1.0f / k);
+				WaveTable[i] += Sign * Amplitude * SinusoidComponent;
+			}
+
+			if (WaveTable[i] > MaxValue)
+			{
+				MaxValue = WaveTable[i];
+			}
+		}
+
+		// Normalize
+		for (int32 i = 0; i < InTableSize; ++i)
+		{
+			WaveTable[i] /= MaxValue;
+		}
+	}
+
+	/*
+	* FTriangleWaveTable
+	*/
+	void FTriangleWaveTable::Generate(const int32 InTableSize)
+	{
+		WaveTable.Reset(InTableSize);
+
+		/*
+		          /\
+		         /  \    
+		             \  /
+		              \/
+		*/
+
+		const int32 QuarterPoint = 0.25f * InTableSize;
+		const int32 ThirdQuarterPoint = 0.75f * InTableSize;
+
+		const float RisingSlope = 1.0f / QuarterPoint;
+		const float FallingSlope = -2.0f / (0.5f * InTableSize);
+
+		for (int32 i = 0; i < InTableSize; ++i)
+		{
+			// Rising first edge
+			if (i < QuarterPoint)
+			{
+				WaveTable.Add(RisingSlope * i);
+			}
+			// Falling middle part
+			else if (i < ThirdQuarterPoint)
+			{
+				WaveTable.Add(FallingSlope * (i - QuarterPoint) + 1.0f);
+			}
+			// Rising last part
+			else
+			{
+				WaveTable.Add(RisingSlope * (i - ThirdQuarterPoint) - 1.0f);
+			}
+		}
+	}
+
+	/*
+	* FBandLimitedTriangleWaveTable
+	*/
+	void FBandLimitedTriangleWaveTable::Generate(const int32 InTableSize)
+	{
+		WaveTable.Reset(InTableSize);
+		WaveTable.AddDefaulted(InTableSize);
+
+		float MaxValue = 0.0f;
+
+		// Generate according to sum of sines, Fourier theorem
+		// triangle wave only has odd-harmonics
+		for (int32 i = 0; i < InTableSize; ++i)
+		{
+			const float TwoPiFraction = 2.0f * PI * (float)i / InTableSize;
+
+			for (int32 k = 0; k <= 3; ++k)
+			{
+				const float SinusoidComponent = FMath::Sin(TwoPiFraction * (2.0f * k + 1.0f));
+				const float Sign = FMath::Pow(-1.0f, (float)k);
+				const float Amplitude = (1.0f / FMath::Pow((float)(2 * k + 1), 2.0f));
+
+				WaveTable[i] += Sign * Amplitude * SinusoidComponent;
+			}
+
+			if (WaveTable[i] > MaxValue)
+			{
+				MaxValue = WaveTable[i];
+			}
+		}
+
+		// Normalize
+		for (int32 i = 0; i < InTableSize; ++i)
+		{
+			WaveTable[i] /= MaxValue;
+		}
+	}
+
+	/*
+	* FBandLimitedSquareWaveTable
+	*/
+	void FSquareWaveTable::Generate(const int32 InTableSize)
+	{
+		WaveTable.Reset(InTableSize);
+		const int32 HalfPoint = 0.5f * InTableSize;
+
+		for (int i = 0; i < InTableSize; ++i)
+		{
+			if (i < HalfPoint)
+			{
+				WaveTable.Add(1.0f);
+			}
+			else
+			{
+				WaveTable.Add(-1.0f);
+			}
+		}
+	}
+
+	void FBandLimitedSquareWaveTable::Generate(const int32 InTableSize)
+	{
+		WaveTable.Reset(InTableSize);
+		WaveTable.AddDefaulted(InTableSize);
+
+		float MaxValue = 0.0f;
+
+		// Generate according to sum of sines, Fourier theorem
+		for (int32 i = 0; i < InTableSize; ++i)
+		{
+			const float TwoPiFraction = 2.0f * PI * (float)i / InTableSize;
+
+			for (int32 k = 1; k <= 5; ++k)
+			{
+				const float SinusoidComponent = FMath::Sin(TwoPiFraction * (2.0f * k - 1.0f));
+				const float Amplitude = 1.0f / (2.0f * k - 1.0f);
+
+				WaveTable[i] += Amplitude * SinusoidComponent;
+			}
+
+			if (WaveTable[i] > MaxValue)
+			{
+				MaxValue = WaveTable[i];
+			}
+		}
+
+		// Normalize
+		for (int32 i = 0; i < InTableSize; ++i)
+		{
+			WaveTable[i] /= MaxValue;
+		}
+	}
+
+}
diff --git a/Engine/Source/Runtime/AudioMixer/Private/Effects/AudioMixerSourceEffectFilter.cpp b/Engine/Source/Runtime/AudioMixer/Private/Effects/AudioMixerSourceEffectFilter.cpp
new file mode 100644
--- /dev/null
+++ b/Engine/Source/Runtime/AudioMixer/Private/Effects/AudioMixerSourceEffectFilter.cpp
@@ -0,0 +1,20 @@
+// Copyright 1998-2016 Epic Games, Inc. All Rights Reserved.
+
+#include "AudioMixerPCH.h"
+
+// Includes for audio dsp helper classes
+#include "SourceEffects/AudioMixerSourceEffectFilter.h"
+
+
+// ========================================================================
+// USoundEffectSourceFilter Implementation
+// ========================================================================
+
+void USourceEffectFilter::Init(const FSoundEffectSourceInitData& InSampleRate)
+{
+}
+
+void USourceEffectFilter::OnProcessAudio(const FSoundEffectSourceInputData& InData, FSoundEffectSourceOutputData& OutData)
+{
+
+}
diff --git a/Engine/Source/Runtime/AudioMixer/Private/Effects/AudioMixerSubmixEffectEQ.cpp b/Engine/Source/Runtime/AudioMixer/Private/Effects/AudioMixerSubmixEffectEQ.cpp
new file mode 100644
--- /dev/null
+++ b/Engine/Source/Runtime/AudioMixer/Private/Effects/AudioMixerSubmixEffectEQ.cpp
@@ -0,0 +1,20 @@
+// Copyright 1998-2016 Epic Games, Inc. All Rights Reserved.
+
+
+#include "AudioMixerPCH.h"
+
+// Includes for audio dsp helper classes
+#include "SubmixEffects/AudioMixerSubmixEffectEQ.h"
+
+void USoundEffectSubmixEQ::Init(const FSoundEffectSourceInitData& InSampleRate)
+{
+
+}
+
+void USoundEffectSubmixEQ::OnProcessAudio(const FSoundEffectSourceInputData& InData, FSoundEffectSourceOutputData& OutData)
+{
+
+}
+
+
+
diff --git a/Engine/Source/Runtime/AudioMixer/Public/AudioMixer.h b/Engine/Source/Runtime/AudioMixer/Public/AudioMixer.h
new file mode 100644
--- /dev/null
+++ b/Engine/Source/Runtime/AudioMixer/Public/AudioMixer.h
@@ -0,0 +1,317 @@
+// Copyright 1998-2016 Epic Games, Inc. All Rights Reserved.
+
+#pragma once
+
+#include "Core.h"
+#include "Engine.h"
+
+AUDIOMIXER_API DECLARE_LOG_CATEGORY_EXTERN(LogAudioMixer, Log, All);
+AUDIOMIXER_API DECLARE_LOG_CATEGORY_EXTERN(LogAudioMixerDebug, Log, All);
+
+#include "SoundDefinitions.h"
+#include "AudioDecompress.h"
+#include "AudioEffect.h"
+#include "AudioMixerTypes.h"
+
+// defines used for AudioMixer.h
+#define AUDIO_PLATFORM_ERROR(INFO)			(OnAudioMixerPlatformError(INFO, FString(__FILE__), __LINE__))
+
+#ifndef AUDIO_MIXER_ENABLE_DEBUG_MODE
+// This define enables a bunch of more expensive debug checks and logging capabilities that are intended to be off most of the time even in debug builds of game/editor.
+#define AUDIO_MIXER_ENABLE_DEBUG_MODE 0
+#endif
+
+// Enable debug checking for audio mixer
+
+#if AUDIO_MIXER_ENABLE_DEBUG_MODE
+#define AUDIO_MIXER_CHECK(expr) check(expr)
+#define AUDIO_MIXER_CHECK_GAME_THREAD(_MixerDevice)			(_MixerDevice->CheckGameThread())
+#define AUDIO_MIXER_CHECK_AUDIO_PLAT_THREAD(_MixerDevice)	(_MixerDevice->CheckAudioPlatformThread())
+#else
+#define AUDIO_MIXER_CHECK(expr)
+#define AUDIO_MIXER_CHECK_GAME_THREAD(_MixerDevice)
+#define AUDIO_MIXER_CHECK_AUDIO_PLAT_THREAD(_MixerDevice)
+#endif
+
+#define AUDIO_MIXER_MIN_PITCH						0.1f
+#define AUDIO_MIXER_MAX_PITCH						4.0f
+
+namespace Audio
+{
+
+	/** Structure to hold platform device information **/
+	struct FAudioPlatformDeviceInfo
+	{
+		/** The name of the audio device */
+		FString Name;
+
+		/** The number of channels supported by the audio device */
+		int32 NumChannels;
+
+		/** The sample rate of the audio device */
+		int32 SampleRate;
+
+		/** The default sample of the audio device */
+		int32 DefaultSampleRate;
+
+		/** The number of frames processed by audio device */
+		int32 NumFrames;
+
+		/** The number of samples processed by the audio device (NumChannels * NumFrames) */
+		int32 NumSamples;
+
+		/** The data format of the audio stream */
+		EAudioMixerStreamDataFormat::Type Format;
+
+		/** The output channel array of the audio device */
+		TArray<EAudioMixerChannel::Type> OutputChannelArray;
+
+		/** The estimated latency of the audio device */
+		uint32 Latency;
+
+		/** Whether or not this device is the system default */
+		uint8 bIsSystemDefault : 1;
+
+		FAudioPlatformDeviceInfo()
+		{
+			Reset();
+		}
+
+		void Reset()
+		{
+			Name = TEXT("Unknown");
+			NumChannels = 0;
+			SampleRate = 0;
+			DefaultSampleRate = 0;
+			NumFrames = 0;
+			NumSamples = 0;
+			Format = EAudioMixerStreamDataFormat::Unknown;
+			OutputChannelArray.Reset();
+			Latency = 0;
+			bIsSystemDefault = false;
+		}
+
+	};
+
+	/** Platform independent audio mixer interface. */
+	class IAudioMixer
+	{
+	public:
+		/** Callback to generate a new audio stream buffer. */
+		virtual bool OnProcessAudioStream(TArray<float>& OutputBuffer) = 0;
+	};
+
+	/** Defines parameters needed for opening a new audio stream to device. */
+	struct FAudioMixerOpenStreamParams
+	{
+		/** The audio device index to open. */
+		uint32 OutputDeviceIndex;
+
+		/** The number of desired audio frames in audio callback. */
+		uint32 NumFrames;
+		
+		/** Owning platform independent audio mixer ptr.*/
+		IAudioMixer* AudioMixer;
+		
+		/** The desired sample rate */
+		uint32 SampleRate;
+
+		FAudioMixerOpenStreamParams()
+			: OutputDeviceIndex(INDEX_NONE)
+			, NumFrames(1024)
+			, AudioMixer(nullptr)
+			, SampleRate(44100)
+		{}
+	};
+
+	struct FAudioOutputStreamInfo
+	{
+		/** The index of the output device for the audio stream. */
+		uint32 OutputDeviceIndex;
+
+		FAudioPlatformDeviceInfo DeviceInfo;
+
+		/** The requested sample rate. */
+		uint32 RequestedSampleRate;
+
+		/** The state of the output audio stream. */
+		FOutputStreamState StreamState;
+
+		/** The callback to use for platform-independent layer. */
+		IAudioMixer* AudioMixer;
+
+		/** The number of frames used in each audio callback. */
+		uint32 NumOutputFrames;
+
+		/** Whether or not we need to perform a format conversion of the audio for this output device. */
+		uint8 bPerformFormatConversion : 1;
+
+		/** Whether or not we need to perform a byte swap for this output device. */
+		uint8 bPerformByteSwap : 1;
+
+		FAudioOutputStreamInfo()
+		{
+			Reset();
+		}
+
+		~FAudioOutputStreamInfo()
+		{
+
+		}
+
+		void Reset()
+		{
+			OutputDeviceIndex = 0;
+			DeviceInfo.Reset();
+			StreamState = EAudioOutputStreamState::Closed;
+			AudioMixer = nullptr;
+			NumOutputFrames = 0;
+			bPerformFormatConversion = false;
+			bPerformByteSwap = false;
+		}
+	};
+
+	class IAudioMixerPlatformInterface;
+
+	/** Class for doing async callbacks to generate audio buffers. */
+	class FAudioMixerCallbackWorker : public FNonAbandonableTask
+	{
+	public:
+
+		/** Constructor. */
+		FAudioMixerCallbackWorker(IAudioMixerPlatformInterface* InMixerPlatform, TArray<float>* InOutBuffer);
+
+		/** Performs task work */
+		void DoWork();
+
+		/** Returns the generated buffer */
+		const TArray<float>* GetBuffer() const { return OutBuffer; }
+
+		/** Returns the stat ID of the task */
+		FORCEINLINE TStatId GetStatId() const
+		{
+			RETURN_QUICK_DECLARE_CYCLE_STAT(FAudioMixerCallbackWorker, STATGROUP_ThreadPoolAsyncTasks);
+		}
+
+	private:
+
+		/** Ptr to platform mixer. */
+		IAudioMixerPlatformInterface* MixerPlatform;
+
+		/** Ptr to buffer to fill. */
+		TArray<float>* OutBuffer;
+	};
+
+	typedef FAsyncTask<FAudioMixerCallbackWorker> FAsyncAudioMixerCallbackTask;
+
+	/** Abstract interface for mixer platform. */
+	class AUDIOMIXER_API IAudioMixerPlatformInterface
+	{
+
+	public: // Virtual functions
+
+		/** Virtual destructor. */
+		virtual ~IAudioMixerPlatformInterface() {}
+
+		/** Returns the platform API enumeration. */
+		virtual EAudioMixerPlatformApi::Type GetPlatformApi() const = 0;
+
+		/** Initialize the hardware. */
+		virtual bool InitializeHardware() = 0;
+
+		/** Teardown the hardware. */
+		virtual bool TeardownHardware() = 0;
+		
+		/** Is the hardware initialized. */
+		virtual bool IsInitialized() const = 0;
+
+		/** Returns the number of output devices. */
+		virtual bool GetNumOutputDevices(uint32& OutNumOutputDevices) = 0;
+
+		/** Gets the device information of the given device index. */
+		virtual bool GetOutputDeviceInfo(const uint32 InDeviceIndex, FAudioPlatformDeviceInfo& OutInfo) = 0;
+
+		/** Returns the default device index. */
+		virtual bool GetDefaultOutputDeviceIndex(uint32& OutDefaultDeviceIndex) const = 0;
+
+		/** Opens up a new audio stream with the given parameters. */
+		virtual bool OpenAudioStream(const FAudioMixerOpenStreamParams& Params) = 0;
+
+		/** Closes the audio stream (if it's open). */
+		virtual bool CloseAudioStream() = 0;
+
+		/** Starts the audio stream processing and generating audio. */
+		virtual bool StartAudioStream() = 0;
+
+		/** Stops the audio stream (but keeps the audio stream open). */
+		virtual bool StopAudioStream() = 0;
+
+		/** Returns the platform device info of the currently open audio stream. */
+		virtual FAudioPlatformDeviceInfo GetPlatformDeviceInfo() const = 0;
+
+		/** Submit the given buffer to the platform's output audio device. */
+		virtual void SubmitBuffer(const TArray<float>& Buffer) = 0;
+
+		/** Returns the name of the format of the input sound wave. */
+		virtual FName GetRuntimeFormat(USoundWave* InSoundWave) = 0;
+
+		/** Checks if the platform has a compressed audio format for sound waves. */
+		virtual bool HasCompressedAudioInfoClass(USoundWave* InSoundWave) = 0;
+
+		/** Creates a Compressed audio info class suitable for decompressing this SoundWave. */
+		virtual ICompressedAudioInfo* CreateCompressedAudioInfo(USoundWave* SoundWave) = 0;
+
+	public: // Public Functions
+
+		/** Constructor. */
+		IAudioMixerPlatformInterface();
+
+		/** Generate the next audio buffer into the given buffer array. */
+		void GenerateBuffer(TArray<float>& Buffer);
+
+		/** Retrieves the next generated buffer and feeds it to the platform mixer output stream. */
+		void ReadNextBuffer();
+
+		/** Returns the last error generated. */
+		FString GetLastError() const { return LastError; }
+
+	protected:
+
+		/** Is called when an error is generated. */
+		inline void OnAudioMixerPlatformError(const FString& ErrorDetails, const FString& FileName, int32 LineNumber)
+		{
+			LastError = FString::Printf(TEXT("Audio Platform Device Error: %s (File %s, Line %d)"), *ErrorDetails, *FileName, LineNumber);
+			UE_LOG(LogAudioMixer, Error, TEXT("%s"), *LastError);
+		}
+
+		/** Returns the number of bytes for the given audio stream format. */
+		uint32 GetNumBytesForFormat(const EAudioMixerStreamDataFormat::Type DataFormat);
+
+		/** Start generating audio callbacks from our mixer. */
+		void BeginGeneratingAudio();
+
+	protected:
+
+		/** The audio device stream info. */
+		FAudioOutputStreamInfo AudioStreamInfo;
+
+		/** The number of mixer buffers. */
+		static const int32 NumMixerBuffers = 3;
+
+		/** Task for generating new audio buffers from audio mixer. */
+		FAsyncAudioMixerCallbackTask* CallbackTask;
+
+		/** List of generated output buffers. */
+		TArray<float> OutputBuffers[NumMixerBuffers];
+
+		/** The current buffer index. */
+		int32 CurrentBufferIndex;
+
+		/** String containing the last generated error. */
+		FString LastError;
+	};
+
+
+
+
+}
diff --git a/Engine/Source/Runtime/AudioMixer/Public/AudioMixerDevice.h b/Engine/Source/Runtime/AudioMixer/Public/AudioMixerDevice.h
new file mode 100644
--- /dev/null
+++ b/Engine/Source/Runtime/AudioMixer/Public/AudioMixerDevice.h
@@ -0,0 +1,184 @@
+// Copyright 1998-2016 Epic Games, Inc. All Rights Reserved.
+
+#pragma once
+
+#include "AudioMixer.h"
+#include "AudioMixerDevice.h"
+#include "AudioMixerSourceVoice.h"
+#include "AudioMixerSourceManager.h"
+
+class IAudioMixerPlatformInterface;
+
+namespace Audio
+{
+	struct FChannelPositionInfo
+	{
+		EAudioMixerChannel::Type Channel;
+		int32 Azimuth;
+
+		FChannelPositionInfo()
+			: Channel(EAudioMixerChannel::Unused)
+			, Azimuth(0)
+		{}
+
+		FChannelPositionInfo(EAudioMixerChannel::Type InChannel, int32 InAzimuth)
+			: Channel(InChannel)
+			, Azimuth(InAzimuth)
+		{}
+	};
+
+	class AUDIOMIXER_API FMixerDevice :	public FAudioDevice,
+										public IAudioMixer
+	{
+	public:
+		FMixerDevice(IAudioMixerPlatformInterface* InAudioMixerPlatform);
+		~FMixerDevice();
+
+		//~ Begin FAudioDevice
+		void GetAudioDeviceList(TArray<FString>& OutAudioDeviceNames) const override;
+		bool InitializeHardware() override;
+		void TeardownHardware() override;
+		void UpdateHardware() override;
+
+		double GetAudioTime() const override;
+
+		FAudioEffectsManager* CreateEffectsManager() override;
+		FSoundSource* CreateSoundSource() override;
+
+		FName GetRuntimeFormat(USoundWave* SoundWave) override;
+		bool HasCompressedAudioInfoClass(USoundWave* SoundWave) override;
+		bool SupportsRealtimeDecompression() const override;
+		class ICompressedAudioInfo* CreateCompressedAudioInfo(USoundWave* SoundWave) override;
+		bool ValidateAPICall(const TCHAR* Function, uint32 ErrorCode) override;
+		bool Exec(UWorld* InWorld, const TCHAR* Cmd, FOutputDevice& Ar) override;
+
+		void CountBytes(class FArchive& Ar) override;
+		bool IsExernalBackgroundSoundActive() override;
+		void Precache(USoundWave* InSoundWave, bool bSynchronous = false, bool bTrackMemory = true) override;
+		void ResumeContext() override;
+		void SetMaxChannels(int) override;
+		void StopAllSounds(bool) override;
+		void SuspendContext() override;
+		void EnableDebugAudioOutput() override;
+
+		void InitSoundSubmixes() override;
+		void RegisterSoundSubmix(USoundSubmix* SoundSubmix) override;
+		void UnregisterSoundSubmix(USoundSubmix* SoundSubmix) override;
+		FMixerSubmix* GetSubmixInstance(USoundSubmix* SoundSubmix);
+
+		int32 GetNumActiveSources() const override;
+		//~ End FAudioDevice
+
+		//~ Begin IAudioMixer
+		bool OnProcessAudioStream(TArray<float>& OutputBuffer) override;
+		//~ End IAudioMixer
+
+		// Functions which check the thread it's called on and helps make sure functions are called from correct threads
+		void CheckGameThread();
+		void CheckAudioPlatformThread();
+
+		// Public Functions
+		FMixerSourceVoice* GetMixerSourceVoice(const FWaveInstance* InWaveInstance, ISourceBufferQueueListener* InBufferQueueListener, bool bUseHRTFSpatialization);
+		int32 GetNumSources() const;
+
+		const FAudioPlatformDeviceInfo& GetPlatformDeviceInfo() const { return PlatformInfo; };
+
+		int32 GetNumDeviceChannels() const { return PlatformInfo.NumChannels; }
+
+		int32 GetNumSpatialChannels() const { return NumSpatialChannels; }
+
+		int32 GetNumOutputFrames() const { return PlatformInfo.NumFrames; }
+
+		const TArray<FChannelPositionInfo>& GetCurrentChannelPositions() const { return CurrentChannelAzimuthPositions; }
+
+		void Get3DChannelMap(const FWaveInstance* InWaveInstance, const float EmitterAzimuth, const float NormalizedOmniRadius, TArray<float>& OutChannelMap);
+		void Get2DChannelMap(const int32 NumSourceChannels, TArray<float>& OutChannelMap);
+
+		void SetChannelAzimuth(EAudioMixerChannel::Type ChannelType, int32 Azimuth);
+
+		int32 GetDeviceSampleRate() const;
+		int32 GetDeviceOutputChannels() const;
+
+		FMixerSourceManager* GetSourceManager();
+
+		FMixerSubmix* GetMasterSubmix() { return MasterSubmix; }
+		FMixerSubmix** GetSubmix(USoundSubmix* InSoundSubmix) { return Submixes.Find(InSoundSubmix); }
+
+	private:
+		void ResetAudioPlatformThreadId();
+
+		void InitializeChannelAzimuthMap(const int32 NumChannels);
+		void Initialize2DChannelMaps();
+
+		int32 GetAzimuthForChannelType(EAudioMixerChannel::Type ChannelType);
+
+		void WhiteNoiseTest(TArray<float>& Output);
+		void SineOscTest(TArray<float>& Output);
+
+	private:
+
+		struct FChannelMapEntry
+		{
+			EChannelMap::Type ChannelMapType;
+			TArray<float> ChannelWeights;
+
+			FChannelMapEntry(EChannelMap::Type InType, float* InputWeights, int32 Count)
+				: ChannelMapType(InType)
+			{
+				ChannelWeights.Append(InputWeights, Count);
+			}
+		};
+
+		/** Ptr to the platform interface, which handles streaming audio to the hardware device. */
+		IAudioMixerPlatformInterface* AudioMixerPlatform;
+		
+		/** Contains a map of channel/speaker azimuth positions. */
+		FChannelPositionInfo DefaultChannelAzimuthPosition[EAudioMixerChannel::MaxSupportedChannel];
+
+		/** The azimuth positions of the current output speakers. */
+		TArray<FChannelPositionInfo> CurrentChannelAzimuthPositions;
+
+		/** 2D channel maps of input-output channel maps. */
+		TMap<int32, FChannelMapEntry> ChannelMaps;
+
+		/** The audio output stream parameters used to initialize the audio hardware. */
+		FAudioMixerOpenStreamParams OpenStreamParams;
+
+		/** The number of output channels which are used for calculating 3d audio. */
+		int32 NumSpatialChannels;
+
+		/** The amount of volume per speaker for fully omni-directional 3d audio. */
+		float OmniPanFactor;
+	
+		/** The time delta for each callback block. */
+		double AudioClockDelta;
+
+		/** The audio clock from device initialization, updated at block rate. */
+		double AudioClock;
+
+		/** The platform device info for this mixer device. */
+		FAudioPlatformDeviceInfo PlatformInfo;
+
+		/** The true root master submix which will always exist. */
+		FMixerSubmix* MasterSubmix;
+
+		/** Map of USoundSubmix static data objects to the dynamic audio mixer submix. */
+		TMap<USoundSubmix*, FMixerSubmix*> Submixes;
+
+		/** List of mixer source voices. */
+		TArray<FMixerSourceVoice*> SourceVoices;
+
+		/** The mixer source manager. */
+		FMixerSourceManager SourceManager;
+
+		/** ThreadId for the game thread (or if audio is running a seperate thread, that ID) */
+		int32 GameOrAudioThreadId;
+
+		/** ThreadId for the low-level platform audio mixer. */
+		int32 AudioPlatformThreadId;
+
+		/** Whether or not we generate output audio to test multi-platform mixer. */
+		bool bDebugOutputEnabled;
+	};
+}
+
diff --git a/Engine/Source/Runtime/AudioMixer/Public/AudioMixerModule.h b/Engine/Source/Runtime/AudioMixer/Public/AudioMixerModule.h
new file mode 100644
--- /dev/null
+++ b/Engine/Source/Runtime/AudioMixer/Public/AudioMixerModule.h
@@ -0,0 +1,20 @@
+// Copyright 1998-2016 Epic Games, Inc. All Rights Reserved.
+
+#pragma once
+
+
+/* Public dependencies
+*****************************************************************************/
+
+#include "Core.h"
+#include "Engine.h"
+
+
+/* Public includes
+*****************************************************************************/
+
+#include "SoundDefinitions.h"
+#include "AudioDecompress.h"
+#include "AudioEffect.h"
+
+
diff --git a/Engine/Source/Runtime/AudioMixer/Public/AudioMixerSubmix.h b/Engine/Source/Runtime/AudioMixer/Public/AudioMixerSubmix.h
new file mode 100644
--- /dev/null
+++ b/Engine/Source/Runtime/AudioMixer/Public/AudioMixerSubmix.h
@@ -0,0 +1,82 @@
+// Copyright 1998-2016 Epic Games, Inc. All Rights Reserved.
+
+#pragma once
+
+#include "AudioMixer.h"
+#include "AudioMixerDevice.h"
+#include "AudioMixerSourceVoice.h"
+#include "Sound/SoundSubmix.h"
+
+namespace Audio
+{
+	class IAudioMixerEffect;
+	class FMixerSourceVoice;
+	class FMixerDevice;
+
+	class FMixerSubmix
+	{
+	public:
+		FMixerSubmix(USoundSubmix* InSoundSubmix, FMixerDevice* InMixerDevice);
+		~FMixerSubmix();
+
+		// Queries if this submix is the master submix.
+		bool IsMasterSubmix() const;
+
+		// Returns the number of source voices currently a part of this submix.
+		int32 GetNumSourceVoices() const;
+
+		// Returns the number of wet effects in this submix.
+		int32 GetNumEffects() const;
+
+		// Adds the given source voice to the submix.
+		void AddSourceVoice(FMixerSourceVoice* InSourceVoice);
+
+		/** Removes the given source voice from the submix. */
+		void RemoveSourceVoice(FMixerSourceVoice* InSourceVoice);
+
+		// Function which processes audio.
+		void ProcessAudio(TArray<float>& OutAudio);
+
+		// Sets the wet-level to use when outputting the submix's audio to it's parent submixes (does not apply for master submix).
+		void SetOutputWetLevel(const float InWetLevel);
+
+		// Returns the submix's wet level, used by parent submixes to determine where to route output audio of children.
+		float GetOutputWetLevel() const;
+
+		// Returns the device sample rate this submix is rendering to
+		int32 GetSampleRate() const;
+
+		// Returns the output channels this submix is rendering to
+		int32 GetNumOutputChannels() const;
+
+		// Updates the submix from the main thread.
+		void Update();
+
+	private:
+
+		// UObject static data with graph information.
+		USoundSubmix* SoundSubmix;
+
+		// The effect chain of this submix, based on the sound submix preset chain
+		TArray<USoundEffectSubmix*> EffectSubmixChain;
+
+		// Owning mixer device. 
+		FMixerDevice* MixerDevice;
+
+		// Parent submix. 
+		FMixerSubmix* Parent;
+
+		// The source voices to be mixed by this submix. 
+		TArray<FMixerSourceVoice*> MixerSourceVoices;
+
+		TArray<float> DryBuffer;
+		TArray<float> WetBuffer;
+		TArray<float> ScratchBuffer;
+
+		// The output wet level of this submix. Used by parent submix to determine how to route the output of this submix. 
+		float WetLevel;
+	};
+
+
+
+}
diff --git a/Engine/Source/Runtime/AudioMixer/Public/AudioMixerTypes.h b/Engine/Source/Runtime/AudioMixer/Public/AudioMixerTypes.h
new file mode 100644
--- /dev/null
+++ b/Engine/Source/Runtime/AudioMixer/Public/AudioMixerTypes.h
@@ -0,0 +1,203 @@
+// Copyright 1998-2016 Epic Games, Inc. All Rights Reserved.
+
+#pragma once
+
+namespace Audio {
+
+	namespace EAudioMixerPlatformApi
+	{
+		enum Type
+		{
+			Wasapi,
+			XAudio2,
+			Ngs2,
+			CoreAudio,
+			OpenAL,
+			Html5,
+			Null,
+		};
+	}
+
+	namespace EAudioMixerStreamDataFormat
+	{
+		enum Type
+		{
+			Unknown,
+			Float,
+			Double,
+			Int16,
+			Int24,
+			Int32,
+			Unsupported
+		};
+	}
+
+	namespace EAudioMixerChannel
+	{
+		/** Enumeration values represent sound file or speaker channel types. */
+		enum Type
+		{
+			FrontLeft,
+			FrontRight,
+			FrontCenter,
+			LowFrequency,
+			BackLeft,
+			BackRight,
+			FrontLeftOfCenter,
+			FrontRightOfCenter,
+			BackCenter,
+			SideLeft,
+			SideRight,
+			TopCenter,
+			TopFrontLeft,
+			TopFrontCenter,
+			TopFrontRight,
+			TopBackLeft,
+			TopBackCenter,
+			TopBackRight,
+			Unused,
+			ChannelTypeCount
+		};
+
+		static const int32 MaxSupportedChannel = EAudioMixerChannel::TopCenter;
+
+		inline const TCHAR* ToString(EAudioMixerChannel::Type InType)
+		{
+			switch (InType)
+			{
+				case FrontLeft:				return TEXT("FrontLeft");
+				case FrontRight:			return TEXT("FrontRight");
+				case FrontCenter:			return TEXT("FrontCenter");
+				case LowFrequency:			return TEXT("LowFrequency");
+				case BackLeft:				return TEXT("BackLeft");
+				case BackRight:				return TEXT("BackRight");
+				case FrontLeftOfCenter:		return TEXT("FrontLeftOfCenter");
+				case FrontRightOfCenter:	return TEXT("FrontRightOfCenter");
+				case BackCenter:			return TEXT("BackCenter");
+				case SideLeft:				return TEXT("SideLeft");
+				case SideRight:				return TEXT("SideRight");
+				case TopCenter:				return TEXT("TopCenter");
+				case TopFrontLeft:			return TEXT("TopFrontLeft");
+				case TopFrontCenter:		return TEXT("TopFrontCenter");
+				case TopFrontRight:			return TEXT("TopFrontRight");
+				case TopBackLeft:			return TEXT("TopBackLeft");
+				case TopBackCenter:			return TEXT("TopBackCenter");
+				case TopBackRight:			return TEXT("TopBackRight");
+
+				default:
+				checkf(false, TEXT("Unsupport channel type"));
+				return TEXT("UNSUPPORTED");
+			}
+		}
+	}
+
+	namespace EChannelMap
+	{
+		/** Enumeration values represent mappings of input channels to output channels. */
+		enum Type
+		{
+			MonoToStereo,
+			MonoToQuad,
+			MonoTo51,
+			MonoTo71,
+
+			StereoToStereo,
+			StereoToQuad,
+			StereoTo51,
+			StereoTo71,
+
+			QuadToStereo,
+			QuadToQuad,
+			QuadTo51,
+			QuadTo71,
+
+			HexToStereo,
+			HexToQuad,
+			HexTo51,
+			HexTo71,
+
+			NumChannelMaps
+		};
+
+		inline const TCHAR* ToString(EChannelMap::Type InType)
+		{
+			switch (InType)
+			{
+				case MonoToStereo:			return TEXT("MonoToStereo");
+				case MonoToQuad:			return TEXT("MonoToQuad");
+				case MonoTo51:				return TEXT("MonoTo51");
+				case MonoTo71:				return TEXT("MonoTo71");
+				case StereoToStereo:		return TEXT("StereoToStereo");
+				case StereoToQuad:			return TEXT("StereoToQuad");
+				case StereoTo51:			return TEXT("StereoTo51");
+				case StereoTo71:			return TEXT("StereoTo71");
+				case QuadToStereo:			return TEXT("QuadToStereo");
+				case QuadToQuad:			return TEXT("QuadToQuad");
+				case QuadTo51:				return TEXT("QuadTo51");
+				case QuadTo71:				return TEXT("QuadTo71");
+				case HexToStereo:			return TEXT("HexToStereo");
+				case HexToQuad:				return TEXT("HexToQuad");
+				case HexTo51:				return TEXT("HexTo51");
+				case HexTo71:				return TEXT("HexTo71");
+
+				default:
+				checkf(false, TEXT("Unsupport channel map"));
+				return TEXT("UNSUPPORTED");
+			}
+		}
+	}
+
+	/**
+	 * EAudioOutputStreamState
+	 * Specifies the state of the output audio stream.
+	 */
+	namespace EAudioOutputStreamState
+	{
+		enum Type
+		{
+			/* The audio stream is shutdown or not uninitialized. */
+			Closed,
+		
+			/* The audio stream is open but not running. */
+			Open,
+
+			/** The audio stream is open but stopped. */
+			Stopped,
+		
+			/** The audio output stream is stopping. */
+			Stopping,
+
+			/** The audio output stream is open and running. */
+			Running,
+		};
+	}
+
+	class FOutputStreamState
+	{
+	public:
+		FOutputStreamState()
+		{
+			State.Set((int32)EAudioOutputStreamState::Closed);
+		}
+
+		FOutputStreamState& operator=(EAudioOutputStreamState::Type StateType)
+		{
+			State.Set((int32)StateType);
+			return *this;
+		}
+
+		bool operator==(EAudioOutputStreamState::Type StateType)
+		{
+			return State.GetValue() == (int32)StateType;
+		}
+
+		bool operator!=(EAudioOutputStreamState::Type StateType)
+		{
+			return State.GetValue() != (int32)StateType;
+		}
+
+	private:
+		FThreadSafeCounter State;
+	};
+
+}
diff --git a/Engine/Source/Runtime/AudioMixer/Public/DSP/AllPassFilter.h b/Engine/Source/Runtime/AudioMixer/Public/DSP/AllPassFilter.h
new file mode 100644
--- /dev/null
+++ b/Engine/Source/Runtime/AudioMixer/Public/DSP/AllPassFilter.h
@@ -0,0 +1,32 @@
+// Copyright 1998-2016 Epic Games, Inc. All Rights Reserved.
+
+#pragma once
+
+#include "DSP/Delay.h"
+
+namespace Audio
+{
+	/** 
+	* An All Pass Filter (APF) delay line. 
+	* https://en.wikipedia.org/wiki/All-pass_filter
+	* APFs can be created by cascading a feedback and a feedforward combfilter with equal delays.
+	*/
+	class FAllPassFilter : public FDelay
+	{
+	public:
+		FAllPassFilter();
+		~FAllPassFilter();
+
+		/** Set the APF Gain coefficient. */
+		void SetAPFGain(const float InGain);
+
+		/** Override the operator for APF */
+		float operator()(const float InSample) override;
+
+	protected:
+		/** APF Feedback/Feedforward gain coefficient. */
+		float APF_G;
+	};
+
+
+}
diff --git a/Engine/Source/Runtime/AudioMixer/Public/DSP/BiQuadFilter.h b/Engine/Source/Runtime/AudioMixer/Public/DSP/BiQuadFilter.h
new file mode 100644
--- /dev/null
+++ b/Engine/Source/Runtime/AudioMixer/Public/DSP/BiQuadFilter.h
@@ -0,0 +1,132 @@
+// Copyright 1998-2016 Epic Games, Inc. All Rights Reserved.
+
+#pragma once
+
+namespace Audio
+{
+	/**************************************************************************************
+	Digital biquad filter derived from analog prototypes using bi-linear transform.
+
+	Parameters:
+	- SampleRate :			Sample rate (e.g. 44.1 kHz) of the filter runs at.
+	- CutoffFrequency :		Normalized frequency at which the filter (whichever type) does its thing.
+	- GainDB :				Amplitude gain in dB (used only in peak and shelve filters)
+	- Q :					Filter quality, amount of sharpness in the frequency response
+	- Bandwidth :			Bandwidth of filter (used alternatively to Q)
+
+
+	Note: Normalized frequency, Fn, is:
+
+	Fn = 2.0f * Fhz / SampleRate
+
+	where Fn is in range of (0.0f, 1.0f]
+
+	See:
+	https://en.wikipedia.org/wiki/Digital_biquad_filter
+	***************************************************************************************/
+
+	namespace EBiquadFilter
+	{
+		enum Type
+		{
+			None,
+
+			LowPass,
+			HighPass,
+			BandPass,
+			Notch,
+			AllPass,
+			PeakEQ,
+			LowShelf,
+			HighShelf,
+		};
+	};
+
+	/** Helper function to get Q from a given bandwidth */
+	FORCEINLINE float GetQFromBandWidth(const float InBandwidth, const float CutoffFrequency)
+	{
+		static const float NaturalLog2 = 0.69314718056f;
+		const float Omega = PI * CutoffFrequency;
+		float InQ = 1.0f / (2.0f * sinh(0.5f * NaturalLog2 * InBandwidth * Omega / FMath::Sin(Omega)));
+		return InQ;
+	}
+
+	class FBiquadFilter
+	{
+	public:
+		/** Default constructor */
+		FBiquadFilter();
+
+		/** Constructor to initialize type and parameters */
+		FBiquadFilter(const EBiquadFilter::Type InFilterType, const float InCutoffFrequency, const float InQ, const float InGainDB = 0.0f);
+
+		/** Virtual destructor */
+		virtual ~FBiquadFilter();
+
+		/** Process input sample to generate next sample. */
+		FORCEINLINE float operator()(const float InSample)
+		{
+			float Output = B0 * InSample + B1 * X1 + B2 * X2 - A1 * Y1 - A2 * Y2;
+
+			// Only accept full-precision floats, clamp to 0.0 for subnormal precision
+			if ((Output > 0.0f && Output < FLT_MIN) || (Output < 0.0f && Output > -FLT_MIN))
+			{
+				Output = 0.0f;
+			}
+
+			X2 = X1;
+			X1 = InSample;
+			Y2 = Y1;
+			Y1 = Output;
+
+			return Output;
+		}
+
+		/** Parameter setting functions. */
+		void SetParams(const EBiquadFilter::Type InFilterType, const float InGainDB, const float InCutoffFrequency, const float IQ);
+
+		/** Set the biquad filter type. */
+		void SetType(const EBiquadFilter::Type InFilterType);
+
+		/** Set the gain. */
+		void SetGain(const float InGainDB);
+
+		/** Sets the frequency of the filter */
+		void SetCutoffFrequency(const float InCutoffFrequency);
+
+		/** Sets the Q of the filter */
+		void SetQ(const float InQ);
+
+		/** Sets the bandwidth of the filter. (Alternative to setting Q) */
+		void SetBandwidth(const float InBandwidth);
+
+		/** Static functions to create filters of different types */
+		static FBiquadFilter CreateLowPassFilter(const float InCutoffFrequency, const float InQ);
+		static FBiquadFilter CreateHighPassFilter(const float InCutoffFrequency, const float InQ);
+		static FBiquadFilter CreateBandPassFilter(const float InCutoffFrequency, const float InQ);
+
+	protected:
+
+		/** The current filter type. */
+		EBiquadFilter::Type FilterType;
+
+		/** Generic filter parameters. */
+		float CutoffFrequency;
+		float Q;
+		float GainDB;
+		float Bandwidth;
+
+		/** Feedforward and feedback Sample values */
+		float X1; // X(n) * Z^-1
+		float X2; // X(n) * Z^-2
+		float Y1; // Y(n) * Z^-1
+		float Y2; // Y(n) * Z^-2
+
+		/** Filter coefficients */
+		float B0;
+		float B1;
+		float B2;
+		float A1;
+		float A2;
+	};
+}
diff --git a/Engine/Source/Runtime/AudioMixer/Public/DSP/Chorus.h b/Engine/Source/Runtime/AudioMixer/Public/DSP/Chorus.h
new file mode 100644
--- /dev/null
+++ b/Engine/Source/Runtime/AudioMixer/Public/DSP/Chorus.h
@@ -0,0 +1,25 @@
+// Copyright 1998-2016 Epic Games, Inc. All Rights Reserved.
+
+#pragma once
+
+#include "DSP/Delay.h"
+#
+namespace Audio
+{
+	/** 
+	* FChorus
+	* Takes mono samples in and returns flanged mono samples out.
+	*/
+	class FChorus : public FModulatedDelay
+	{
+	public:
+		/** Constructor */
+		FChorus();
+
+		/** Destructor */
+		virtual ~FChorus();
+
+		/** Initialize the flanger with the given sample rate. */
+		void Init(int32 InSampleRate) override;
+	};
+}
diff --git a/Engine/Source/Runtime/AudioMixer/Public/DSP/Delay.h b/Engine/Source/Runtime/AudioMixer/Public/DSP/Delay.h
new file mode 100644
--- /dev/null
+++ b/Engine/Source/Runtime/AudioMixer/Public/DSP/Delay.h
@@ -0,0 +1,151 @@
+// Copyright 1998-2016 Epic Games, Inc. All Rights Reserved.
+
+#pragma once
+
+#include "DSP/Dsp.h"
+#include "DSP/WaveTableOsc.h"
+
+namespace Audio
+{
+
+	/** Represents a circulating delay line. */
+	class FDelay
+	{
+	public:
+
+		/** Constructor. */
+		FDelay();
+
+		/** Destructor. */
+		virtual ~FDelay();
+
+		/** Initializes the delay line with the given number of input samples. */
+		void Init(const int32 InSampleRate, const int32 InDelaySizeMsec = 2000.0f);
+
+		/** Resets the delay line to re-use for audio processing. */
+		void Reset();
+
+		/** Sets the feedback of the delay line. */
+		void SetFeedback(const float InFeedback);
+
+		/** Sets the output gain of the delay line in dB */
+		void SetOutputGainDB(const float InOutputGainDB);
+
+		/** Sets the output gain of the delay line in linear gain. */
+		void SetOutputGain(const float InOutputGain);
+
+		/** Sets the delay length (of read index) by given fractional mSec. */
+		void SetDelayMsec(const float InDelayMsec);
+
+		/** Reads from the delay line at the current read index without incrementing the read index. */
+		float Read() const;
+
+		/** Reads the value of the delay line at the given tap out point in milliseconds. */
+		float TapOut(const float TapOutMsec) const;
+
+		/** Writes the audio sample into the delay line and increments read and write indices. */
+		void TapIn(const float InSample);
+
+		/** Writes the input sample and returns the sample at the output of the delay line. */
+		virtual float operator()(const float InSample);
+
+	protected:
+		/** The delay line audio buffer. */
+		TArray<float> AudioBuffer;
+
+		/** Feedback of the delay line. */
+		float Feedback;
+
+		/** Represents this delay lines current fractional sample delay. */
+		float DelaySamples;
+
+		/** The output gain of the delay. */
+		float OutputGain;
+
+		/** Where the delay line is going to read from next. */
+		int32 ReadIndex;
+
+		/** Where teh delay line is going write to next. */
+		int32 WriteIndex;
+
+		/** The sample rate of the delay line. */
+		int32 SampleRate;
+	};
+
+	/** A delay line which can be modulated by an LFO. */
+	class FModulatedDelay
+	{
+	public:
+		/** Constructor */
+		FModulatedDelay();
+
+		/** Destructor */
+		virtual ~FModulatedDelay();
+
+		/** Initialize the modulated delay with the given sample rate. */
+		virtual void Init(int32 InSampleRate);
+
+		/** Set the modulated delay to use the saw LFO */
+		void SetLFO(const EWaveTable::Type InLFOType);
+
+		/** Sets whether or not this modulated delay's LFOs are in quadrature phase. */
+		void SetQuadraturePhase(const bool bInIsInQuadPhase) { bIsInQuadPhase = bInIsInQuadPhase; }
+
+		/** Set the modulation feedback (between -1.0f and 1.0f) */
+		virtual void SetModulationFeedback(const float InModulationFeedback);
+
+		/** Sets the modulated delay modulation min and max delay times. */
+		void SetMinAndMaxDelayMSec(const float InMinDelayTimeMsec, const float InMaxDelayTimeMsec);
+
+		/** Sets the delay offset time. */
+		void SetModulationDelayOffsetMsec(const float InDelayOffsetMsec);
+
+		/** Sets the modulation frequency of the LFOs. */
+		void SetModulationFrequency(const float InModulationFrequency);
+
+		/** Sets the amount of wet level in the modulated delay. */
+		void SetWetLevel(const float InWetLevel);
+
+		/** Generate new sample from the given input sample. */
+		float operator()(const float InSample);
+
+	protected:
+
+		/** Wave tables used for LFOs. */
+		FSineWaveTable SinLFO;
+		FSawWaveTable SawLFO;
+		FTriangleWaveTable TriLFO;
+		FSquareWaveTable SquareLFO;
+
+		/** Ptr to the currently used wave tables */
+		FWaveTableOsc* CurrentWaveTable;
+
+		/** A delay line */
+		FDelay Delay;
+
+		/** The modulation depth */
+		float ModulationDepth;
+
+		/** The modulation frequency. */
+		float ModulationFrequency;
+
+		/** The min delay time for the flanger. */
+		float MinDelayTimeMsec;
+
+		/** The max delay time for the flanger. */
+		float MaxDelayTimeMsec;
+
+		/** An additional delay offset to use. */
+		float DelayOffsetMsec;
+
+		/** The modulation feedback. */
+		float ModulationFeedback;
+
+		/** The amount of wet level in the modulated delay line. */
+		float WetLevel;
+
+		/** Whether or not the flanger is operating in quad phase. */
+		bool bIsInQuadPhase;
+	};
+
+}
diff --git a/Engine/Source/Runtime/AudioMixer/Public/DSP/Dsp.h b/Engine/Source/Runtime/AudioMixer/Public/DSP/Dsp.h
new file mode 100644
--- /dev/null
+++ b/Engine/Source/Runtime/AudioMixer/Public/DSP/Dsp.h
@@ -0,0 +1,24 @@
+// Copyright 1998-2016 Epic Games, Inc. All Rights Reserved.
+
+#pragma once
+
+
+/** Utility macro to clamp a sample value to 0 if in range outside of full precision. */
+
+namespace Audio
+{
+	FORCEINLINE float UnderflowClamp(const float InValue)
+	{
+		if (InValue > 0.0f && InValue < FLT_MIN)
+		{
+			return 0.0f;
+		}
+		if (InValue < 0.0f && InValue > -FLT_MIN)
+		{
+			return 0.0f;
+		}
+		return InValue;
+	}
+
+}
+
diff --git a/Engine/Source/Runtime/AudioMixer/Public/DSP/Flanger.h b/Engine/Source/Runtime/AudioMixer/Public/DSP/Flanger.h
new file mode 100644
--- /dev/null
+++ b/Engine/Source/Runtime/AudioMixer/Public/DSP/Flanger.h
@@ -0,0 +1,53 @@
+// Copyright 1998-2016 Epic Games, Inc. All Rights Reserved.
+
+#pragma once
+
+#include "DSP/Delay.h"
+
+namespace Audio
+{
+	/** 
+	* FFlanger
+	* A simple modulated delay line.
+	*/
+	class FFlanger : public FModulatedDelay
+	{
+	public:
+		/** Constructor */
+		FFlanger();
+
+		/** Destructor */
+		virtual ~FFlanger();
+
+		/** Initialize the flanger with the given sample rate. */
+		void Init(int32 InSampleRate) override;
+	};
+
+	class FStereoFlanger
+	{
+	public:
+		FStereoFlanger();
+		~FStereoFlanger();
+
+		/** Initialize the stereo flanger with the given sample rate. */
+		void Init(int32 InSampleRate);
+
+		/** Set the flanger to use the saw LFO */
+		void SetLFO(const EWaveTable::Type InLFOType);
+
+		/** Set the modulation feedback (between -1.0f and 1.0f) */
+		void SetModulationFeedback(const float InModulationFeedback);
+
+		/** Sets the flanger modulation min and max delay times. */
+		void SetModulationMinAndMaxDelayMSec(const float InMinDelayTimeMsec, const float InMaxDelayTimeMsec);
+
+		/** Sets the modulation frequency of the LFOs. */
+		void SetModulationFrequency(const float InModulationFrequency);
+
+		void operator()(const float InLeft, const float InRight, float* OutLeft, float* OutRight);
+
+	private:
+		/** Stereo flanger, one for each channel. */
+		FFlanger Flanger[2];
+	};
+}
diff --git a/Engine/Source/Runtime/AudioMixer/Public/DSP/Noise.h b/Engine/Source/Runtime/AudioMixer/Public/DSP/Noise.h
new file mode 100644
--- /dev/null
+++ b/Engine/Source/Runtime/AudioMixer/Public/DSP/Noise.h
@@ -0,0 +1,56 @@
+// Copyright 1998-2016 Epic Games, Inc. All Rights Reserved.
+
+#pragma once
+
+namespace Audio
+{
+	/** 
+	* White noise generator 
+	* Flat spectrum
+	*/
+	class FWhiteNoise
+	{
+	public:
+		/** Constructor with a default scale add parameter */
+		FWhiteNoise(const float InScale = 1.0f, const float InAdd = 0.0f);
+
+		void SetScaleAdd(const float InScale, const float InAdd);
+
+		/** Generate next sample of white noise */
+		float operator()();
+
+	private:
+		float Scale;
+		float Add;
+	};
+
+	/** 
+	* Pink noise generator
+	* 1/Frequency noise spectrum
+	*/
+	class FPinkNoise
+	{
+	public:
+		/** Constructor. */
+		FPinkNoise(const float InScale = 1.0f, const float InAdd = 0.0f);
+
+		/** Sets the output scale and add parameter. */
+		void SetScaleAdd(const float InScale, const float InAdd);
+
+		/** Generate next sample of pink noise. */
+		float operator()();
+
+	private:
+
+		void InitFilter();
+
+		FWhiteNoise Noise;
+
+		float A[4];
+		float B[4];
+		float X[4];
+		float Y[4];
+	};
+
+
+}
diff --git a/Engine/Source/Runtime/AudioMixer/Public/DSP/OnePole.h b/Engine/Source/Runtime/AudioMixer/Public/DSP/OnePole.h
new file mode 100644
--- /dev/null
+++ b/Engine/Source/Runtime/AudioMixer/Public/DSP/OnePole.h
@@ -0,0 +1,113 @@
+// Copyright 1998-2016 Epic Games, Inc. All Rights Reserved.
+
+#pragma once
+
+#include "DSP/Dsp.h"
+
+namespace Audio
+{
+	/** One Pole Feedback Filter */
+	class FOnePole
+	{
+	public:
+		/** Constructor */
+		FOnePole()
+			: A0(1.0f)
+			, B1(0.0f)
+			, Z1(0.0f)
+		{
+		}
+
+		/** Destructor */
+		virtual ~FOnePole()
+		{
+		}
+
+		/** Initialize the LPF */
+		void Init()
+		{
+			A0 = 1.0f;
+			B1 = 0.0f;
+			Z1 = 0.0f;
+		}
+
+		/** Directly set the filter coefficients */
+		void SetCoefficients(const float InA0, const float InB1)
+		{
+			A0 = InA0;
+			B1 = InB1;
+		}
+
+		/** Filter the input sample. */
+		float operator()(const float InSample)
+		{
+			const float Yn = A0 * InSample + B1 * Z1;
+
+			// Keep float values in full precision
+			return Z1 = UnderflowClamp(Yn);
+		}
+
+	protected:
+		// Coefficients
+		float A0;
+		float B1;
+
+		/** One sample delay. */
+		float Z1;
+	};
+
+	/** One Pole Low Pass Filter */
+	class FOnePoleLPF : public FOnePole
+	{
+	public:
+		FOnePoleLPF()
+			: CutoffFrequency(0.0f)
+		{
+			SetFrequency(1.0f);
+		}
+
+		/** Directly sets the G-coefficient. */
+		FORCEINLINE void SetG(const float InG)
+		{
+			B1 = InG;
+			A0 = 1.0f - B1;
+		}
+
+		/** Sets the filter frequency using normalized frequency */
+		FORCEINLINE void SetFrequency(const float InFrequency)
+		{
+			if (CutoffFrequency != InFrequency)
+			{
+				CutoffFrequency = InFrequency;
+				B1 = FMath::Exp(-PI * CutoffFrequency);
+				A0 = 1.0f - B1;
+			}
+		}
+
+		float CutoffFrequency;
+	};
+
+	/** One Pole Highpass Pass Filter */
+	class FOnePoleHPF : public FOnePole
+	{
+	public:
+		FOnePoleHPF()
+			: CutoffFrequency(1.0f)
+		{
+			SetFrequency(0.0f);
+		}
+
+		FORCEINLINE void SetFrequency(const float InFrequency)
+		{
+			if (CutoffFrequency != InFrequency)
+			{
+				CutoffFrequency = InFrequency;
+				B1 = -FMath::Exp(PI * (CutoffFrequency - 0.5f));
+				A0 = 1.0f + B1;
+			}
+		}
+
+		float CutoffFrequency;
+	};
+
+}
diff --git a/Engine/Source/Runtime/AudioMixer/Public/DSP/ParamInterpolator.h b/Engine/Source/Runtime/AudioMixer/Public/DSP/ParamInterpolator.h
new file mode 100644
--- /dev/null
+++ b/Engine/Source/Runtime/AudioMixer/Public/DSP/ParamInterpolator.h
@@ -0,0 +1,39 @@
+// Copyright 1998-2016 Epic Games, Inc. All Rights Reserved.
+
+#pragma once
+
+namespace Audio
+{
+
+	/**
+	* FParamInterpolator
+	* Simple parameter interpolator which interpolates to target values
+	* in the given the required number of update ticks.
+	*/
+	class FParamInterpolator
+	{
+	public:
+
+		FParamInterpolator();
+		~FParamInterpolator();
+
+		/** Initialize the value. */
+		void InitValue(const float InValue);
+
+		/** Set the target value to interpolate to in the given number of update ticks. */
+		void SetValue(const float InValue, const uint32 InNumInterpolationTicks);
+
+		/** Returns the current value without updating the tick. */
+		float GetValue() const;
+
+		/** Functor operator to generate new output value. */
+		virtual float operator()();
+
+	protected:
+		float StartValue;
+		float EndValue;
+		float CurrentTick;
+		float NumTicksToNewValue;
+	};
+
+}
diff --git a/Engine/Source/Runtime/AudioMixer/Public/DSP/ParametricEqualizer.h b/Engine/Source/Runtime/AudioMixer/Public/DSP/ParametricEqualizer.h
new file mode 100644
--- /dev/null
+++ b/Engine/Source/Runtime/AudioMixer/Public/DSP/ParametricEqualizer.h
@@ -0,0 +1,51 @@
+// Copyright 1998-2016 Epic Games, Inc. All Rights Reserved.
+
+#pragma once
+
+#include "DSP/BiQuadFilter.h"
+
+namespace Audio
+{
+	/** Parameteric equalizer. A band of parallel biquad PeakEQ filters. */
+	class FParametricEqualizer
+	{
+	public:
+		/** Constructor */
+		FParametricEqualizer();
+
+		/** Destructor */
+		~FParametricEqualizer();
+
+		/** Struct defining each band's parameters. */
+		struct FEqualizerBandParams
+		{
+			float CutoffFrequency;
+			float Bandwidth;
+			float Gain;
+		};
+
+		/** Sets (or overrides) the parameters of each of the bands. The number of filters used will be the size of the array. */
+		void SetParameters(const TArray<FEqualizerBandParams>& InParams);
+
+		/** Sets given band's parameters. Returns false if the band doesn't exist. **/
+		bool SetBand(const int32 InBand, const FEqualizerBandParams& InParams);
+
+		/** Process input sample to generate new output sample. */
+		FORCEINLINE float operator()(const float Input)
+		{
+			float Output = 0.0f;
+			for (int32 i = 0; i < Filters.Num(); ++i)
+			{
+				Output += Filters[i](Input);
+			}
+			return Output;
+		}
+
+	private:
+
+		TArray<FBiquadFilter> Filters;
+	};
+
+
+
+}
diff --git a/Engine/Source/Runtime/AudioMixer/Public/DSP/Reverb.h b/Engine/Source/Runtime/AudioMixer/Public/DSP/Reverb.h
new file mode 100644
--- /dev/null
+++ b/Engine/Source/Runtime/AudioMixer/Public/DSP/Reverb.h
@@ -0,0 +1,72 @@
+// Copyright 1998-2016 Epic Games, Inc. All Rights Reserved.
+
+#pragma once
+
+#include "DSP/Delay.h"
+#include "DSP/OnePole.h"
+
+namespace Audio
+{
+	/** 
+	* FreeVerb/Schroeder verb implementation
+	* https://ccrma.stanford.edu/~jos/pasp/Freeverb.html
+	*/
+	class FFreeVerb
+	{
+	public:
+		FFreeVerb(const float InSampleRate);
+
+		void SetWetLevel(const float InWetLevel);
+		void SetRoomSize(const float InRoomSize);
+		float GetRoomSize() const;
+		void SetDamping(const float InDamping);
+		float GetDamping() const;
+		void SetStereoWidth(float InWidth);
+		float GetStereoWidth() const;
+		void SetFreeze(const bool bInIsFrozen);
+		bool IsFrozen() const;
+		void Clear();
+
+		void operator()(const TArray<float>& InAudio, TArray<float>& OutAudio);
+
+	protected:
+		void UpdateParameters();
+
+		static const int32 NumCombFilters = 8;
+		static const int32 NumAllPassFilters = 4;
+		static const int32 StereoSpread = 23;
+
+		static const float FixedGain;
+		static const float ScaleWet;
+		static const float ScaleDry;
+		static const float ScaleDamp;
+		static const float ScaleRoom;
+		static const float OffsetRoom;
+
+		TArray<float> CombFilterDelayLengths;
+		TArray<float> AllpassFilterDelayLengths;
+
+		float AllPassCoefficientG;
+		float Gain;
+		float RoomSizeSetting;
+		float RoomSize;
+		float DampSetting;
+		float Damp;
+		float WetLevel;
+		float Wet1; 
+		float Wet2;
+		float Dry;
+		float Width;
+		bool bIsFrozen;
+
+		TArray<FDelay> CombDelayLeft;
+		TArray<FDelay> CombDelayRight;
+		TArray<FOnePoleLPF> CombLowPassLeft;
+		TArray<FOnePoleLPF> CombLowPassRight;
+
+		TArray<FDelay> AllPassDelayLeft;
+		TArray<FDelay> AllPassDelayRight;
+	};
+
+
+}
diff --git a/Engine/Source/Runtime/AudioMixer/Public/DSP/SinOsc.h b/Engine/Source/Runtime/AudioMixer/Public/DSP/SinOsc.h
new file mode 100644
--- /dev/null
+++ b/Engine/Source/Runtime/AudioMixer/Public/DSP/SinOsc.h
@@ -0,0 +1,49 @@
+// Copyright 1998-2016 Epic Games, Inc. All Rights Reserved.
+
+#pragma once
+
+namespace Audio
+{
+	/** 
+	* FOsc
+	* Direct-form sinusoid oscillator. 
+	* Created with a biquad filter (using only feedback coefficients) with poles directly on unit circle in z-plane.
+	* Setting frequency uses current filter state to compute initial conditions to avoid pops when changing frequency.
+	*/
+	class FSineOsc
+	{
+	public:
+		/** Constructor */
+		FSineOsc();
+
+		/** Non-default constructor */
+		FSineOsc(const int32 InSampleRate, const float InFrequencyHz, const float Scale = 1.0f, const float Add = 0.0f);
+
+		/** Virtual destructor */
+		virtual ~FSineOsc();
+
+		/** Initialize the oscillator with a sample rate and new frequency. Must be called before playing oscillator. */
+		void Init(const int32 InSampleRate, const float InFrequencyHz, const float Scale = 1.0f, const float Add = 0.0f);
+
+		/** Sets the frequency of the oscillator in Hz (based on sample rate). Performs initial condition calculation to avoid pops. */
+		void SetFrequency(const float InFrequencyHz);
+
+		/** Returns the current frequency. */
+		float GetFrequency() const;
+
+		/** Generates the next sample of the oscillator. */
+		float operator()();
+
+	protected:
+		int32 SampleRate;
+		float FrequencyHz;
+		float B1;	// Biquad feedback coefficients
+		float B2;
+		float Yn_1;	// y(n - 1)
+		float Yn_2;	// y(n - 2)
+		float Scale;
+		float Add;
+	};
+
+
+}
diff --git a/Engine/Source/Runtime/AudioMixer/Public/DSP/WaveTableOsc.h b/Engine/Source/Runtime/AudioMixer/Public/DSP/WaveTableOsc.h
new file mode 100644
--- /dev/null
+++ b/Engine/Source/Runtime/AudioMixer/Public/DSP/WaveTableOsc.h
@@ -0,0 +1,136 @@
+// Copyright 1998-2016 Epic Games, Inc. All Rights Reserved.
+
+#pragma once
+
+
+namespace Audio
+{
+	/** The various types of wave tables */
+	namespace EWaveTable
+	{
+		enum Type
+		{
+			SineWaveTable,
+			SawWaveTable,
+			TriangleWaveTable,
+			SquareWaveTable,
+			BandLimitedSawWaveTable,
+			BandLimitedTriangleWaveTable,
+			BandLimitedSquareWaveTable,
+			UserCreated
+		};
+	}
+
+
+	/** Base class for a wave table oscillator. */
+	class FWaveTableOsc
+	{
+	public:
+		FWaveTableOsc();
+		virtual ~FWaveTableOsc();
+
+		/** Initializes the wave table oscillator */
+		void Init(const int32 InSampleRate, const int32 InTableSize = 1024);
+
+		/** Set the frequency of the oscillator in hertz */
+		void SetFrequency(const float InFrequencyHz);
+
+		/** Sets the polarity of the wave table oscillator. True makes it unipolar. False makes it bipolar. */
+		void SetPolarity(const bool bInIsUnipolar);
+
+		/** Function which returns the current in-phase and quad phase sample values. */
+		void operator()(float* OutSample, float* OutQuadSample);
+
+		/** Returns the type of the wave table this is. */
+		virtual EWaveTable::Type GetType() const = 0;
+
+	protected:
+		/** Function to generate the wave table. */
+		virtual void Generate(const int32 InTableSize) = 0;
+
+		/** Audio for the wave table. */
+		TArray<float> WaveTable;
+
+	private:
+
+		/** Sample rate of the oscillator. */
+		int32 SampleRate;
+
+		/** Read index of the wave table oscillator */
+		float ReadIndex;
+
+		/** Quadrature read index of the wave table oscillator */
+		float QuadPhaseReadIndex;
+
+		/** The amount to increment to the read indices whenever read from the oscillator */
+		float ReadDelta;
+
+		/** What polarity the wave table generator works in. */
+		bool bIsUnipolar;
+	};
+
+	class FSineWaveTable : public FWaveTableOsc
+	{
+	public:
+		EWaveTable::Type GetType() const { return EWaveTable::SineWaveTable; }
+
+	protected:
+		virtual void Generate(const int32 InTableSize) override;
+		
+	};
+
+	class FSawWaveTable : public FWaveTableOsc
+	{
+	public:
+		EWaveTable::Type GetType() const { return EWaveTable::SawWaveTable; }
+
+	protected:
+		virtual void Generate(const int32 InTableSize) override;
+	};
+
+	class FBandLimitedSawWaveTable : public FWaveTableOsc
+	{
+	public:
+		EWaveTable::Type GetType() const { return EWaveTable::BandLimitedSawWaveTable; }
+
+	protected:
+		virtual void Generate(const int32 InTableSize) override;
+	};
+
+	class FTriangleWaveTable : public FWaveTableOsc
+	{
+	public:
+		EWaveTable::Type GetType() const { return EWaveTable::TriangleWaveTable; }
+
+	protected:
+		virtual void Generate(const int32 InTableSize) override;
+	};
+
+	class FBandLimitedTriangleWaveTable : public FWaveTableOsc
+	{
+	public:
+		EWaveTable::Type GetType() const { return EWaveTable::BandLimitedTriangleWaveTable; }
+
+	protected:
+		virtual void Generate(const int32 InTableSize) override;
+	};
+
+	class FSquareWaveTable : public FWaveTableOsc
+	{
+	public:
+		EWaveTable::Type GetType() const { return EWaveTable::SquareWaveTable; }
+
+	protected:
+		virtual void Generate(const int32 InTableSize) override;
+	};
+
+	class FBandLimitedSquareWaveTable : public FWaveTableOsc
+	{
+	public:
+		EWaveTable::Type GetType() const { return EWaveTable::BandLimitedSquareWaveTable; }
+
+	protected:
+		virtual void Generate(const int32 InTableSize) override;
+	};
+
+}
diff --git a/Engine/Source/Runtime/CinematicCamera/Private/CameraRig_Crane.cpp b/Engine/Source/Runtime/CinematicCamera/Private/CameraRig_Crane.cpp
--- a/Engine/Source/Runtime/CinematicCamera/Private/CameraRig_Crane.cpp
+++ b/Engine/Source/Runtime/CinematicCamera/Private/CameraRig_Crane.cpp
@@ -43,8 +43,7 @@ ACameraRig_Crane::ACameraRig_Crane(const FObjectInitializer& ObjectInitializer)
 		if (PreviewMesh_CraneBase)
 		{
 			PreviewMesh_CraneBase->SetStaticMesh(CraneBaseMesh.Object);
-			PreviewMesh_CraneBase->AlwaysLoadOnClient = false;
-			PreviewMesh_CraneBase->AlwaysLoadOnServer = false;
+			PreviewMesh_CraneBase->bIsEditorOnly = true;
 			PreviewMesh_CraneBase->SetCollisionProfileName(UCollisionProfile::NoCollision_ProfileName);
 			PreviewMesh_CraneBase->bHiddenInGame = true;
 			PreviewMesh_CraneBase->CastShadow = false;
@@ -58,8 +57,7 @@ ACameraRig_Crane::ACameraRig_Crane(const FObjectInitializer& ObjectInitializer)
 		if (PreviewMesh_CraneArm)
 		{
 			PreviewMesh_CraneArm->SetStaticMesh(CraneArmMesh.Object);
-			PreviewMesh_CraneArm->AlwaysLoadOnClient = false;
-			PreviewMesh_CraneArm->AlwaysLoadOnServer = false;
+			PreviewMesh_CraneArm->bIsEditorOnly = true;
 			PreviewMesh_CraneArm->SetCollisionProfileName(UCollisionProfile::NoCollision_ProfileName);
 			PreviewMesh_CraneArm->bHiddenInGame = true;
 			PreviewMesh_CraneArm->CastShadow = false;
@@ -76,8 +74,7 @@ ACameraRig_Crane::ACameraRig_Crane(const FObjectInitializer& ObjectInitializer)
 		if (PreviewMesh_CraneMount)
 		{
 			PreviewMesh_CraneMount->SetStaticMesh(CraneArmMount.Object);
-			PreviewMesh_CraneMount->AlwaysLoadOnClient = false;
-			PreviewMesh_CraneMount->AlwaysLoadOnServer = false;
+			PreviewMesh_CraneMount->bIsEditorOnly = true;
 			PreviewMesh_CraneMount->SetCollisionProfileName(UCollisionProfile::NoCollision_ProfileName);
 			PreviewMesh_CraneMount->bHiddenInGame = true;
 			PreviewMesh_CraneMount->CastShadow = false;
@@ -94,8 +91,7 @@ ACameraRig_Crane::ACameraRig_Crane(const FObjectInitializer& ObjectInitializer)
 		if (PreviewMesh_CraneCounterWeight)
 		{
 			PreviewMesh_CraneCounterWeight->SetStaticMesh(CraneArmCW.Object);
-			PreviewMesh_CraneCounterWeight->AlwaysLoadOnClient = false;
-			PreviewMesh_CraneCounterWeight->AlwaysLoadOnServer = false;
+			PreviewMesh_CraneCounterWeight->bIsEditorOnly = true;
 			PreviewMesh_CraneCounterWeight->SetCollisionProfileName(UCollisionProfile::NoCollision_ProfileName);
 			PreviewMesh_CraneCounterWeight->bHiddenInGame = true;
 			PreviewMesh_CraneCounterWeight->CastShadow = false;
diff --git a/Engine/Source/Runtime/CinematicCamera/Private/CameraRig_Rail.cpp b/Engine/Source/Runtime/CinematicCamera/Private/CameraRig_Rail.cpp
--- a/Engine/Source/Runtime/CinematicCamera/Private/CameraRig_Rail.cpp
+++ b/Engine/Source/Runtime/CinematicCamera/Private/CameraRig_Rail.cpp
@@ -37,8 +37,7 @@ ACameraRig_Rail::ACameraRig_Rail(const FObjectInitializer& ObjectInitializer)
 		if (PreviewMesh_Mount)
 		{
 			PreviewMesh_Mount->SetStaticMesh(MountMesh.Object);
-			PreviewMesh_Mount->AlwaysLoadOnClient = false;
-			PreviewMesh_Mount->AlwaysLoadOnServer = false;
+			PreviewMesh_Mount->bIsEditorOnly = true;
 			PreviewMesh_Mount->SetCollisionProfileName(UCollisionProfile::NoCollision_ProfileName);
 			PreviewMesh_Mount->bHiddenInGame = true;
 			PreviewMesh_Mount->CastShadow = false;
@@ -56,8 +55,7 @@ USplineMeshComponent* ACameraRig_Rail::CreateSplinePreviewSegment()
 	{
 		Segment->SetStaticMesh(PreviewRailStaticMesh);
 		Segment->SetMobility(EComponentMobility::Movable);
-		Segment->AlwaysLoadOnClient = false;
-		Segment->AlwaysLoadOnServer = false;
+		Segment->bIsEditorOnly = true;
 		Segment->SetCollisionProfileName(UCollisionProfile::NoCollision_ProfileName);
 		Segment->bHiddenInGame = true;
 		Segment->CastShadow = false;
diff --git a/Engine/Source/Runtime/CinematicCamera/Private/CineCameraComponent.cpp b/Engine/Source/Runtime/CinematicCamera/Private/CineCameraComponent.cpp
--- a/Engine/Source/Runtime/CinematicCamera/Private/CineCameraComponent.cpp
+++ b/Engine/Source/Runtime/CinematicCamera/Private/CineCameraComponent.cpp
@@ -295,8 +295,7 @@ void UCineCameraComponent::CreateDebugFocusPlane()
 		{
 			DebugFocusPlaneComponent = NewObject<UStaticMeshComponent>(MyOwner, NAME_None, RF_Transactional | RF_TextExportTransient);
 			DebugFocusPlaneComponent->SetupAttachment(this);
-			DebugFocusPlaneComponent->AlwaysLoadOnClient = false;
-			DebugFocusPlaneComponent->AlwaysLoadOnServer = false;
+			DebugFocusPlaneComponent->bIsEditorOnly = true;
 			DebugFocusPlaneComponent->StaticMesh = DebugFocusPlaneMesh;
 			DebugFocusPlaneComponent->SetCollisionProfileName(UCollisionProfile::NoCollision_ProfileName);
 			DebugFocusPlaneComponent->bHiddenInGame = false;
diff --git a/Engine/Source/Runtime/Core/Private/Math/UnrealMath.cpp b/Engine/Source/Runtime/Core/Private/Math/UnrealMath.cpp
--- a/Engine/Source/Runtime/Core/Private/Math/UnrealMath.cpp
+++ b/Engine/Source/Runtime/Core/Private/Math/UnrealMath.cpp
@@ -586,6 +586,28 @@ FQuat FQuat::MakeFromEuler(const FVector& Euler)
 	return FRotator::MakeFromEuler(Euler).Quaternion();
 }
 
+void FQuat::ToSwingTwist(const FVector& InTwistAxis, FQuat& OutSwing, FQuat& OutTwist) const
+{
+	// Vector part projected onto twist axis
+	FVector Projection = FVector::DotProduct(InTwistAxis, FVector(X, Y, Z)) * InTwistAxis;
+
+	// Twist quaternion
+	OutTwist = FQuat(Projection.X, Projection.Y, Projection.Z, W);
+
+	// Singularity close to 180deg
+	if(OutTwist.SizeSquared() == 0.0f)
+	{
+		OutTwist = FQuat::Identity;
+	}
+	else
+	{
+		OutTwist.Normalize();
+	}
+
+	// Set swing
+	OutSwing = *this * OutTwist.Inverse();
+}
+
 FMatrix FRotationAboutPointMatrix::Make(const FQuat& Rot, const FVector& Origin)
 {
 	return FRotationAboutPointMatrix(Rot.Rotator(), Origin);
diff --git a/Engine/Source/Runtime/Core/Private/Misc/Core.cpp b/Engine/Source/Runtime/Core/Private/Misc/Core.cpp
--- a/Engine/Source/Runtime/Core/Private/Misc/Core.cpp
+++ b/Engine/Source/Runtime/Core/Private/Misc/Core.cpp
@@ -97,6 +97,7 @@ bool GIsReinstancing = false;
 #if WITH_ENGINE
 bool					PRIVATE_GIsRunningCommandlet			= false;				/* Whether this executable is running a commandlet (custom command-line processing code) */
 bool					PRIVATE_GAllowCommandletRendering	= false;				/** If true, initialise RHI and set up scene for rendering even when running a commandlet. */
+bool					PRIVATE_GAllowCommandletAudio 		= false;				/** If true, allow audio even when running a commandlet. */
 #endif	// WITH_ENGINE
 
 #if WITH_EDITORONLY_DATA
diff --git a/Engine/Source/Runtime/Core/Public/CoreGlobals.h b/Engine/Source/Runtime/Core/Public/CoreGlobals.h
--- a/Engine/Source/Runtime/Core/Public/CoreGlobals.h
+++ b/Engine/Source/Runtime/Core/Public/CoreGlobals.h
@@ -106,6 +106,9 @@ extern CORE_API bool PRIVATE_GIsRunningCommandlet;
 
 /** If true, initialize RHI and set up scene for rendering even when running a commandlet. */
 extern CORE_API bool PRIVATE_GAllowCommandletRendering;
+
+/** If true, initialize audio and even when running a commandlet. */
+extern CORE_API bool PRIVATE_GAllowCommandletAudio;
 #endif
 
 #if WITH_EDITORONLY_DATA
@@ -168,6 +171,15 @@ FORCEINLINE bool IsAllowCommandletRendering()
 #endif
 }
 
+FORCEINLINE bool IsAllowCommandletAudio()
+{
+#if WITH_ENGINE
+	return PRIVATE_GAllowCommandletAudio;
+#else
+	return false;
+#endif
+}
+
 extern CORE_API bool GEdSelectionLock;
 extern CORE_API bool GIsClient;
 extern CORE_API bool GIsServer;
diff --git a/Engine/Source/Runtime/Core/Public/Math/Quat.h b/Engine/Source/Runtime/Core/Public/Math/Quat.h
--- a/Engine/Source/Runtime/Core/Public/Math/Quat.h
+++ b/Engine/Source/Runtime/Core/Public/Math/Quat.h
@@ -308,6 +308,16 @@ public:
 	 */
 	void ToAxisAndAngle(FVector& Axis, float& Angle) const;
 
+	/** 
+	 * Get the swing and twist decomposition for a specified axis
+	 *
+	 * @param InTwistAxis Axis to use for decomposition
+	 * @param OutSwing swing component quaternion
+	 * @param OutTwist Twist component quaternion
+	 * @warning assumes normalised quaternion and twist axis
+	 */
+	CORE_API void ToSwingTwist(const FVector& InTwistAxis, FQuat& OutSwing, FQuat& OutTwist) const;
+
 	/**
 	 * Rotate a vector by this quaternion.
 	 *
@@ -937,7 +947,6 @@ FORCEINLINE void FQuat::ToAxisAndAngle(FVector& Axis, float& Angle) const
 	Axis = GetRotationAxis();
 }
 
-
 FORCEINLINE FVector FQuat::GetRotationAxis() const
 {
 	// Ensure we never try to sqrt a neg number
diff --git a/Engine/Source/Runtime/Core/Public/UObject/FrameworkObjectVersion.h b/Engine/Source/Runtime/Core/Public/UObject/FrameworkObjectVersion.h
--- a/Engine/Source/Runtime/Core/Public/UObject/FrameworkObjectVersion.h
+++ b/Engine/Source/Runtime/Core/Public/UObject/FrameworkObjectVersion.h
@@ -52,6 +52,10 @@ struct CORE_API FFrameworkObjectVersion
 
 		// Wheel force offset is now applied at the wheel instead of vehicle COM
 		WheelOffsetIsFromWheel,
+		// Move curve metadata to be saved in skeleton
+		// Individual asset still saves some flag - i.e. disabled curve and editable or not, but 
+		// major flag - i.e. material types - moves to skeleton and handle in one place
+		MoveCurveTypesToSkeleton,
 
 		// -----<new versions can be added above this line>-------------------------------------------------
 		VersionPlusOne,
diff --git a/Engine/Source/Runtime/CoreUObject/Private/UObject/SavePackage.cpp b/Engine/Source/Runtime/CoreUObject/Private/UObject/SavePackage.cpp
--- a/Engine/Source/Runtime/CoreUObject/Private/UObject/SavePackage.cpp
+++ b/Engine/Source/Runtime/CoreUObject/Private/UObject/SavePackage.cpp
@@ -539,7 +539,7 @@ bool IsEditorOnlyObject(const UObject* InObject)
 		return false;
 	}
 	
-	if (InObject->HasAnyMarks(OBJECTMARK_EditorOnly))
+	if (InObject->HasAnyMarks(OBJECTMARK_EditorOnly) || InObject->IsEditorOnly())
 	{
 		return true;
 	}
@@ -548,19 +548,16 @@ bool IsEditorOnlyObject(const UObject* InObject)
 	check(InObject);
 	// If this is a package that is editor only or the object is in editor-only package,
 	// the object is editor-only too.
-	const UPackage* Package = Cast<const UPackage>(InObject);
-	if (!Package)
-	{
-		Package = InObject->GetOutermost();
-	}
+	const bool bIsAPackage = InObject->IsA<UPackage>();
+	const UPackage* Package = (bIsAPackage ? static_cast<const UPackage*>(InObject) : InObject->GetOutermost());
 	if (Package && Package->HasAnyPackageFlags(PKG_EditorOnly))
 	{
 		bResult = true;
 	}
-	if (!bResult && !InObject->IsA(UPackage::StaticClass()))
+	else if (!bIsAPackage)
 	{
 		// Otherwise the object is editor-only if its class is editor-only
-		const UStruct* Struct = InObject->IsA(UStruct::StaticClass()) ? CastChecked<const UStruct>(InObject) : InObject->GetClass();
+		const UStruct* Struct = InObject->IsA<UStruct>() ? static_cast<const UStruct*>(InObject) : InObject->GetClass();
 		bResult = IsEditorOnlyStruct(Struct);
 		if (!bResult && InObject->GetOuter())
 		{
diff --git a/Engine/Source/Runtime/CoreUObject/Private/UObject/UObjectHash.cpp b/Engine/Source/Runtime/CoreUObject/Private/UObject/UObjectHash.cpp
--- a/Engine/Source/Runtime/CoreUObject/Private/UObject/UObjectHash.cpp
+++ b/Engine/Source/Runtime/CoreUObject/Private/UObject/UObjectHash.cpp
@@ -692,7 +692,7 @@ UObjectBase* FindObjectWithOuter(class UObjectBase* Outer, class UClass* ClassTo
 	SCOPE_CYCLE_COUNTER( STAT_Hash_FindObjectWithOuter );
 
 	UObject* Result = nullptr;
-	check( Outer && ClassToLookFor );
+	check( Outer );
 	// We don't want to return any objects that are currently being background loaded unless we're using the object iterator during async loading.
 	EInternalObjectFlags ExclusionInternalFlags = EInternalObjectFlags::Unreachable;
 	if (!IsInAsyncLoadingThread())
@@ -718,7 +718,7 @@ UObjectBase* FindObjectWithOuter(class UObjectBase* Outer, class UClass* ClassTo
 				{
 					continue;
 				}
-				if (!Object->IsA(ClassToLookFor))
+				if (ClassToLookFor && !Object->IsA(ClassToLookFor))
 				{
 					continue;
 				}
diff --git a/Engine/Source/Runtime/CoreUObject/Public/UObject/Object.h b/Engine/Source/Runtime/CoreUObject/Public/UObject/Object.h
--- a/Engine/Source/Runtime/CoreUObject/Public/UObject/Object.h
+++ b/Engine/Source/Runtime/CoreUObject/Public/UObject/Object.h
@@ -358,6 +358,16 @@ public:
 		return false;
 	}
 
+	/**
+	* Called during saving to determine if the object is forced to be editor only or not
+	*
+	* @return	true if this object should never be loaded outside the editor
+	*/
+	virtual bool IsEditorOnly() const
+	{
+		return false;
+	}
+
 	/**
 	* Called during async load to determine if PostLoad can be called on the loading thread.
 	*
diff --git a/Engine/Source/Runtime/CoreUObject/Public/UObject/ObjectMacros.h b/Engine/Source/Runtime/CoreUObject/Public/UObject/ObjectMacros.h
--- a/Engine/Source/Runtime/CoreUObject/Public/UObject/ObjectMacros.h
+++ b/Engine/Source/Runtime/CoreUObject/Public/UObject/ObjectMacros.h
@@ -938,6 +938,9 @@ namespace UP
 
 		/// Property shouldn't be serialized, can still be exported to text
 		SkipSerialization,
+
+		/// Property wont have a 'reset to default' button when displayed in property windows
+		NoResetToDefault,
 	};
 }
 
diff --git a/Engine/Source/Runtime/Engine/Classes/AI/Navigation/NavigationTypes.h b/Engine/Source/Runtime/Engine/Classes/AI/Navigation/NavigationTypes.h
--- a/Engine/Source/Runtime/Engine/Classes/AI/Navigation/NavigationTypes.h
+++ b/Engine/Source/Runtime/Engine/Classes/AI/Navigation/NavigationTypes.h
@@ -314,6 +314,8 @@ struct FNavLocation
 
 	/** checks if location has associated navigation node ref */
 	FORCEINLINE bool HasNodeRef() const { return NodeRef != INVALID_NAVNODEREF; }
+
+	FORCEINLINE operator FVector() const { return Location; }
 };
 
 /** Describes node in navigation path */
diff --git a/Engine/Source/Runtime/Engine/Classes/Animation/AnimInstance.h b/Engine/Source/Runtime/Engine/Classes/Animation/AnimInstance.h
--- a/Engine/Source/Runtime/Engine/Classes/Animation/AnimInstance.h
+++ b/Engine/Source/Runtime/Engine/Classes/Animation/AnimInstance.h
@@ -24,6 +24,7 @@ class UCanvas;
 class UWorld;
 struct FTransform;
 class FDebugDisplayInfo;
+struct FAnimNode_SubInput;
 struct FAnimNode_AssetPlayerBase;
 struct FAnimNode_Base;
 struct FAnimInstanceProxy;
@@ -503,7 +504,7 @@ public:
 
 	/** Play normal animation asset on the slot node by creating a dynamic UAnimMontage. You can only play one asset (whether montage or animsequence) at a time per SlotGroup. */
 	UFUNCTION(BlueprintCallable, Category="Animation")
-	UAnimMontage* PlaySlotAnimationAsDynamicMontage(UAnimSequenceBase* Asset, FName SlotNodeName, float BlendInTime = 0.25f, float BlendOutTime = 0.25f, float InPlayRate = 1.f, int32 LoopCount = 1, float BlendOutTriggerTime = -1.f);
+	UAnimMontage* PlaySlotAnimationAsDynamicMontage(UAnimSequenceBase* Asset, FName SlotNodeName, float BlendInTime = 0.25f, float BlendOutTime = 0.25f, float InPlayRate = 1.f, int32 LoopCount = 1, float BlendOutTriggerTime = -1.f, float InTimeToStartMontageAt = 0.f);
 
 	/** Stops currently playing slot animation slot or all*/
 	UFUNCTION(BlueprintCallable, Category="Animation")
@@ -522,7 +523,7 @@ public:
 public:
 	/** Plays an animation montage. Returns the length of the animation montage in seconds. Returns 0.f if failed to play. */
 	UFUNCTION(BlueprintCallable, Category = "Animation")
-	float Montage_Play(UAnimMontage* MontageToPlay, float InPlayRate = 1.f, EMontagePlayReturnType ReturnValueType = EMontagePlayReturnType::MontageLength);
+	float Montage_Play(UAnimMontage* MontageToPlay, float InPlayRate = 1.f, EMontagePlayReturnType ReturnValueType = EMontagePlayReturnType::MontageLength, float InTimeToStartMontageAt=0.f);
 
 	/** Stops the animation montage. If reference is NULL, it will stop ALL active montages. */
 	UFUNCTION(BlueprintCallable, Category = "Animation")
@@ -651,6 +652,8 @@ public:
 	virtual void OnMontageInstanceStopped(FAnimMontageInstance & StoppedMontageInstance);
 	void ClearMontageInstanceReferences(FAnimMontageInstance& InMontageInstance);
 
+	FAnimNode_SubInput* GetSubInputNode() const;
+
 protected:
 	/** Map between Active Montages and their FAnimMontageInstance */
 	TMap<class UAnimMontage*, struct FAnimMontageInstance*> ActiveMontagesMap;
@@ -1087,7 +1090,7 @@ public:
 	 * Retrieve animation curve list by Curve Flags, it will return list of {UID, value} 
 	 * It will clear the OutCurveList before adding
 	 */
-	void GetAnimationCurveList(int32 CurveFlags, TMap<FName, float>& OutCurveList) const;
+	void GetAnimationCurveList(EAnimCurveType Type, TMap<FName, float>& OutCurveList) const;
 
 #if WITH_EDITORONLY_DATA
 	// Maximum playback position ever reached (only used when debugging in Persona)
@@ -1120,6 +1123,11 @@ public:
 	 */
 	void RecalcRequiredBones();
 
+	/**
+	* Recalculate Required Curves based on Required Bones [RequiredBones]
+	*/
+	void RecalcRequiredCurves();
+
 	/** When RequiredBones mapping has changed, AnimNodes need to update their bones caches. */
 	DEPRECATED(4.11, "This cannot be accessed directly, use FAnimInstanceProxy::bBoneCachesInvalidated")
 	bool bBoneCachesInvalidated;
@@ -1150,7 +1158,7 @@ public:
 	void TriggerSingleAnimNotify(const FAnimNotifyEvent* AnimNotifyEvent);
 
 	/** Add curve float data using a curve Uid, the name of the curve will be resolved from the skeleton **/
-	void AddCurveValue(const USkeleton::AnimCurveUID Uid, float Value, int32 CurveTypeFlags);
+	void AddCurveValue(const USkeleton::AnimCurveUID Uid, float Value);
 
 	/** Given a machine index, record a state machine weight for this frame */
 	void RecordMachineWeight(const int32& InMachineClassIndex, const float& InMachineWeight);
@@ -1158,7 +1166,7 @@ public:
 	 * Add curve float data, using a curve name. External values should all be added using
 	 * The curve UID to the public version of this method
 	 */
-	void AddCurveValue(const FName& CurveName, float Value, int32 CurveTypeFlags);
+	void AddCurveValue(const FName& CurveName, float Value);
 
 	/** Given a machine and state index, record a state weight for this frame */
 	void RecordStateWeight(const int32& InMachineClassIndex, const int32& InStateIndex, const float& InStateWeight);
diff --git a/Engine/Source/Runtime/Engine/Classes/Animation/AnimSequence.h b/Engine/Source/Runtime/Engine/Classes/Animation/AnimSequence.h
--- a/Engine/Source/Runtime/Engine/Classes/Animation/AnimSequence.h
+++ b/Engine/Source/Runtime/Engine/Classes/Animation/AnimSequence.h
@@ -331,7 +331,7 @@ class ENGINE_API UAnimSequence : public UAnimSequenceBase
 	GENERATED_UCLASS_BODY()
 
 	/** Number of raw frames in this sequence (not used by engine - just for informational purposes). */
-	UPROPERTY(AssetRegistrySearchable)
+	UPROPERTY(AssetRegistrySearchable, meta=(DisplayName = "Number of Keys"))
 	int32 NumFrames;
 
 	/**
diff --git a/Engine/Source/Runtime/Engine/Classes/Animation/AnimSequenceBase.h b/Engine/Source/Runtime/Engine/Classes/Animation/AnimSequenceBase.h
--- a/Engine/Source/Runtime/Engine/Classes/Animation/AnimSequenceBase.h
+++ b/Engine/Source/Runtime/Engine/Classes/Animation/AnimSequenceBase.h
@@ -13,9 +13,6 @@
 #include "Animation/AnimCurveTypes.h"
 #include "AnimSequenceBase.generated.h"
 
-#define DEFAULT_SAMPLERATE			30.f
-#define MINIMUM_ANIMATION_LENGTH	(1/DEFAULT_SAMPLERATE)
-
 UENUM()
 enum ETypeAdvanceAnim
 {
diff --git a/Engine/Source/Runtime/Engine/Classes/Animation/AnimationAsset.h b/Engine/Source/Runtime/Engine/Classes/Animation/AnimationAsset.h
--- a/Engine/Source/Runtime/Engine/Classes/Animation/AnimationAsset.h
+++ b/Engine/Source/Runtime/Engine/Classes/Animation/AnimationAsset.h
@@ -865,6 +865,7 @@ public:
 
 	ENGINE_API void SetPreviewMesh(USkeletalMesh* PreviewMesh);
 	ENGINE_API USkeletalMesh* GetPreviewMesh();
+	ENGINE_API USkeletalMesh* GetPreviewMesh() const;
 
 	ENGINE_API virtual int32 GetMarkerUpdateCounter() const { return 0; }
 
diff --git a/Engine/Source/Runtime/Engine/Classes/Animation/PoseAsset.h b/Engine/Source/Runtime/Engine/Classes/Animation/PoseAsset.h
--- a/Engine/Source/Runtime/Engine/Classes/Animation/PoseAsset.h
+++ b/Engine/Source/Runtime/Engine/Classes/Animation/PoseAsset.h
@@ -186,8 +186,8 @@ public:
 
 #if WITH_EDITOR
 	ENGINE_API void AddOrUpdatePose(const FSmartName& PoseName, USkeletalMeshComponent* MeshComponent);
-	ENGINE_API void AddOrUpdatePoseWithUniqueName(USkeletalMeshComponent* MeshComponent);
-	ENGINE_API void AddOrUpdatePose(const FSmartName& PoseName, TArray<FName> TrackNames, TArray<FTransform>& LocalTransform);
+	ENGINE_API bool AddOrUpdatePoseWithUniqueName(USkeletalMeshComponent* MeshComponent, FSmartName* OutPoseName = nullptr);
+	ENGINE_API void AddOrUpdatePose(const FSmartName& PoseName, const TArray<FName>& TrackNames, const TArray<FTransform>& LocalTransform, const TArray<float>& CurveValues);
 
 	ENGINE_API void CreatePoseFromAnimation(class UAnimSequence* AnimSequence, const TArray<FSmartName>* InPoseNames = nullptr);
 	ENGINE_API void UpdatePoseFromAnimation(class UAnimSequence* AnimSequence);
@@ -198,7 +198,7 @@ public:
 	// End AnimationAsset interface
 #endif
 
-	ENGINE_API bool ModifyPoseName(FName OldPoseName, FName NewPoseName, const FSmartNameMapping::UID* NewUID);
+	ENGINE_API bool ModifyPoseName(FName OldPoseName, FName NewPoseName, const SmartName::UID_Type* NewUID);
 
 	ENGINE_API int32 DeletePoses(TArray<FName> PoseNamesToDelete);
 	ENGINE_API int32 DeleteCurves(TArray<FName> CurveNamesToDelete);
diff --git a/Engine/Source/Runtime/Engine/Classes/Animation/PreviewAssetAttachComponent.h b/Engine/Source/Runtime/Engine/Classes/Animation/PreviewAssetAttachComponent.h
--- a/Engine/Source/Runtime/Engine/Classes/Animation/PreviewAssetAttachComponent.h
+++ b/Engine/Source/Runtime/Engine/Classes/Animation/PreviewAssetAttachComponent.h
@@ -82,6 +82,14 @@ public:
 	 */
 	ENGINE_API void AddAttachedObject( UObject* AttachObject, FName AttachPointName );
 
+	/**
+	* Adds the given Name/Object uniquely to the PreviewAttachedObjects list. Allows us to recreate the attached objects later
+	*
+	* @param	AttachObject		The object that is being attached
+	* @param	AttachPointName		The place where the object is attached (bone or socket)
+	*/
+	ENGINE_API void AddUniqueAttachedObject(UObject* AttachObject, FName AttachPointName);
+
 	/**
 	 * Remove the given object from the attached list
 	 *
@@ -112,6 +120,11 @@ public:
 	 */
 	ENGINE_API FPreviewAttachedObjectPair& operator[]( int32 i );
 
+	/**
+	* operator [] - pipes through to AttachedObjects
+	*/
+	ENGINE_API const FPreviewAttachedObjectPair& operator[](int32 i) const;
+
 	/**
 	 * Const iterator creator for AttachedObjects
 	 */
diff --git a/Engine/Source/Runtime/Engine/Classes/Animation/Skeleton.h b/Engine/Source/Runtime/Engine/Classes/Animation/Skeleton.h
--- a/Engine/Source/Runtime/Engine/Classes/Animation/Skeleton.h
+++ b/Engine/Source/Runtime/Engine/Classes/Animation/Skeleton.h
@@ -267,7 +267,7 @@ public:
 	TMap< FName, FReferencePose > AnimRetargetSources;
 
 	// Typedefs for greater smartname UID readability, add one for each smartname category 
-	typedef FSmartNameMapping::UID AnimCurveUID;
+	typedef SmartName::UID_Type AnimCurveUID;
 
 	// Names for smartname mappings, if you're adding a new category of smartnames add a new name here
 	static ENGINE_API const FName AnimCurveMappingName;
@@ -275,19 +275,28 @@ public:
 	// Names for smartname mappings, if you're adding a new category of smartnames add a new name here
 	static ENGINE_API const FName AnimTrackCurveMappingName;
 
-	/** Caches the Animation curve mapping name UIds from SmartNames into CachedAnimCurveMappingNameUids (used in FBlendedCurve::InitFrom) */
-	void CacheAnimCurveMappingNameUids();
+	// these return container of curve meta data, if you modify this container, 
+	// you'll have to call REfreshCAchedAnimationCurveData to apply
+	ENGINE_API FCurveMetaData* GetCurveMetaData(const FName& CurveName);
+	ENGINE_API const FCurveMetaData* GetCurveMetaData(const FName& CurveName) const;
+	ENGINE_API FCurveMetaData* GetCurveMetaData(const FSmartName& CurveName);
+	ENGINE_API const FCurveMetaData* GetCurveMetaData(const FSmartName& CurveName) const;
+	// this is called when you know both flags - called by post serialize
+	ENGINE_API void AccumulateCurveMetaData(FName CurveName, bool bMaterialSet, bool bMorphtargetSet);
 
-	// Returns the cached animation curve mapping Uids, used for initializing FBlendedCurve 
-	const TArray<FSmartNameMapping::UID>& GetCachedAnimCurveMappingNameUids();
+	// return version of AnimCurveUidVersion
+	uint16 GetAnimCurveUidVersion() const { return AnimCurveUidVersion;  }
 
 protected:
 	// Container for smart name mappings
 	UPROPERTY()
 	FSmartNameContainer SmartNames;
 
-	/** Cached animation curves smart name mapping UIDs, only at runtime, not serialized. (accessed in FBlendedCurve::InitFrom) */
-	TArray<FSmartNameMapping::UID> CachedAnimCurveMappingNameUids;
+private:
+	/** Increase the AnimCurveUidVersion so that instances can get the latest information */
+	void IncreaseAnimCurveUidVersion();
+	/** Current  Anim Curve Uid Version. Increase whenever it has to be recalculated */
+	uint16 AnimCurveUidVersion;
 
 public:
 	//////////////////////////////////////////////////////////////////////////
@@ -320,8 +329,9 @@ private:
 
 public:
 	ENGINE_API FAnimSlotGroup* FindAnimSlotGroup(const FName& InGroupName);
+	ENGINE_API const FAnimSlotGroup* FindAnimSlotGroup(const FName& InGroupName) const;
 	ENGINE_API const TArray<FAnimSlotGroup>& GetSlotGroups() const;
-	ENGINE_API bool ContainsSlotName(const FName& InSlotName);
+	ENGINE_API bool ContainsSlotName(const FName& InSlotName) const;
 	ENGINE_API void RegisterSlotNode(const FName& InSlotName);
 	ENGINE_API void SetSlotGroupName(const FName& InSlotName, const FName& InGroupName);
 	/** Returns true if Group is added, false if it already exists */
@@ -345,23 +355,23 @@ public:
 
 	// Renames a smartname in the specified container and modifies the skeleton
 	// return bool - Whether the rename was sucessful
-	ENGINE_API bool RenameSmartnameAndModify(FName ContainerName, FSmartNameMapping::UID Uid, FName NewName);
+	ENGINE_API bool RenameSmartnameAndModify(FName ContainerName, SmartName::UID_Type Uid, FName NewName);
 
 	// Removes a smartname from the specified container and modifies the skeleton
-	ENGINE_API void RemoveSmartnameAndModify(FName ContainerName, FSmartNameMapping::UID Uid);
+	ENGINE_API void RemoveSmartnameAndModify(FName ContainerName, SmartName::UID_Type Uid);
 
 	// Removes smartnames from the specified container and modifies the skeleton
-	ENGINE_API void RemoveSmartnamesAndModify(FName ContainerName, const TArray<FSmartNameMapping::UID>& Uids);
+	ENGINE_API void RemoveSmartnamesAndModify(FName ContainerName, const TArray<SmartName::UID_Type>& Uids);
 #endif// WITH_EDITOR
 
-	// quick wrapper function for Find UID by name, if not found, it will return FSmartNameMapping::MaxUID
-	ENGINE_API FSmartNameMapping::UID GetUIDByName(const FName& ContainerName, const FName& Name);
-	ENGINE_API bool GetSmartNameByUID(const FName& ContainerName, FSmartNameMapping::UID UID, FSmartName& OutSmartName);
+	// quick wrapper function for Find UID by name, if not found, it will return SmartName::MaxUID
+	ENGINE_API SmartName::UID_Type GetUIDByName(const FName& ContainerName, const FName& Name) const;
+	ENGINE_API bool GetSmartNameByUID(const FName& ContainerName, SmartName::UID_Type UID, FSmartName& OutSmartName);
 	ENGINE_API bool GetSmartNameByName(const FName& ContainerName, const FName& InName, FSmartName& OutSmartName);
 
 	// Adds a new name to the smart name container and modifies the skeleton so it can be saved
 	// return bool - Whether a name was added (false if already present)
-	ENGINE_API bool RenameSmartName(FName ContainerName, const FSmartNameMapping::UID& Uid, FName NewName);
+	ENGINE_API bool RenameSmartName(FName ContainerName, const SmartName::UID_Type& Uid, FName NewName);
 
 	// Get or add a smartname container with the given name
 	ENGINE_API const FSmartNameMapping* GetSmartNameContainer(const FName& ContainerName) const;
@@ -380,6 +390,10 @@ private:
 	UPROPERTY(duplicatetransient, AssetRegistrySearchable)
 	TAssetPtr<class USkeletalMesh> PreviewSkeletalMesh;
 
+	/** The additional skeletal meshes to use when previewing this skeleton */
+	UPROPERTY(duplicatetransient, AssetRegistrySearchable)
+	TAssetPtr<class UPreviewMeshCollection> AdditionalPreviewSkeletalMeshes;
+
 	UPROPERTY()
 	FRigConfiguration RigConfig;
 
@@ -448,6 +462,15 @@ public:
 	/** Returns the skeletons preview mesh, loading it if necessary */
 	ENGINE_API void SetPreviewMesh(USkeletalMesh* PreviewMesh, bool bMarkAsDirty=true);
 
+	/** Load any additional meshes we may have */
+	ENGINE_API void LoadAdditionalPreviewSkeletalMeshes();
+
+	/** Get the additional skeletal meshes we use when previewing this skeleton */
+	ENGINE_API UPreviewMeshCollection* GetAdditionalPreviewSkeletalMeshes() const;
+
+	/** Set the additional skeletal meshes we use when previewing this skeleton */
+	ENGINE_API void SetAdditionalPreviewSkeletalMeshes(UPreviewMeshCollection* PreviewMeshCollection);
+
 	/**
 	 * Makes sure all attached objects are valid and removes any that aren't.
 	 *
@@ -544,6 +567,7 @@ public:
 	 *	Anybody modifying BoneTree will corrupt animation data, so will need to make sure it's not modifiable outside of Skeleton
 	 *	You can add new BoneNode but you can't modify current list. The index will be referenced by Animation data.
 	 */
+	DEPRECATED(4.14, "GetBoneTree should not be called. Please use GetBoneTranslationRetargetingMode()/SetBoneTranslationRetargetingMode() functions instead")
 	const TArray<FBoneNode>& GetBoneTree()	const
 	{ 
 		return BoneTree;	
@@ -612,7 +636,7 @@ public:
 	ENGINE_API static void AddReferencedObjects(UObject* InThis, FReferenceCollector& Collector);
 
 	/** 
-	 * Create RefLocalPoses from InSkelMesh
+	 * Create RefLocalPoses from InSkelMesh. Note InSkelMesh cannot be null and this function will assert if it is.
 	 * 
 	 * If bClearAll is false, it will overwrite ref pose of bones that are found in InSkelMesh
 	 * If bClearAll is true, it will reset all Reference Poses 
@@ -705,9 +729,15 @@ public:
 	ENGINE_API void RefreshRigConfig();
 	int32 FindRigBoneMapping(const FName& NodeName) const;
 	ENGINE_API URig * GetRig() const;
+
 #endif
 
+public:
+	ENGINE_API USkeletalMeshSocket* FindSocketAndIndex(FName InSocketName, int32& OutIndex) const;
+	ENGINE_API USkeletalMeshSocket* FindSocket(FName InSocketName) const;
+
 private:
+	/** Regenerate new Guid */
 	void RegenerateGuid();	
 };
 
