Fix include file name capitalization CIS failure for Linux and use of unsupported constexpr

diff --git a/Engine/Source/Runtime/AudioMixer/Private/DSP/Filter.cpp b/Engine/Source/Runtime/AudioMixer/Private/DSP/Filter.cpp
--- a/Engine/Source/Runtime/AudioMixer/Private/DSP/Filter.cpp
+++ b/Engine/Source/Runtime/AudioMixer/Private/DSP/Filter.cpp
@@ -152,7 +152,7 @@ namespace Audio
 
 	void FFilter::CalculateBiquadCoefficients()
 	{
-		constexpr float NaturalLog2 = 0.69314718055994530942f;
+		const float NaturalLog2 = 0.69314718055994530942f;
 
 		const float Omega = 2.0f * PI * Frequency / SampleRate;
 		const float Sn = FMath::Sin(Omega);
diff --git a/Engine/Source/Runtime/AudioMixer/Public/DSP/Filter.h b/Engine/Source/Runtime/AudioMixer/Public/DSP/Filter.h
--- a/Engine/Source/Runtime/AudioMixer/Public/DSP/Filter.h
+++ b/Engine/Source/Runtime/AudioMixer/Public/DSP/Filter.h
@@ -3,7 +3,7 @@
 #pragma once
 
 #include "CoreMinimal.h"
-#include "BiquadFilter.h"
+#include "BiQuadFilter.h"
 
 namespace Audio
 {
