Copying //UE4/Dev-AnimPhys to //UE4/Dev-Main (Source: //UE4/Dev-AnimPhys @ 3624379)

diff --git a/Engine/Source/Runtime/AnimGraphRuntime/Public/AnimNodes/AnimNode_Slot.h b/Engine/Source/Runtime/AnimGraphRuntime/Public/AnimNodes/AnimNode_Slot.h
--- a/Engine/Source/Runtime/AnimGraphRuntime/Public/AnimNodes/AnimNode_Slot.h
+++ b/Engine/Source/Runtime/AnimGraphRuntime/Public/AnimNodes/AnimNode_Slot.h
@@ -24,6 +24,10 @@ struct ANIMGRAPHRUNTIME_API FAnimNode_Slot : public FAnimNode_Base
 	UPROPERTY(EditAnywhere, BlueprintReadWrite, Category=Settings, meta=(CustomizeProperty))
 	FName SlotName;
 
+	//Whether we should continue to update the source pose regardless of whether it would be used.
+	UPROPERTY(EditAnywhere, BlueprintReadWrite, Category = Settings)
+	bool bAlwaysUpdateSourcePose;
+
 protected:
 	FSlotNodeWeightInfo WeightData;
 	FGraphTraversalCounter SlotNodeInitializationCounter;
diff --git a/Engine/Source/Runtime/AnimGraphRuntime/Public/PlayMontageCallbackProxy.h b/Engine/Source/Runtime/AnimGraphRuntime/Public/PlayMontageCallbackProxy.h
--- a/Engine/Source/Runtime/AnimGraphRuntime/Public/PlayMontageCallbackProxy.h
+++ b/Engine/Source/Runtime/AnimGraphRuntime/Public/PlayMontageCallbackProxy.h
@@ -63,6 +63,7 @@ protected:
 private:
 	TWeakObjectPtr<UAnimInstance> AnimInstancePtr;
 	int32 MontageInstanceID;
+	uint32 bInterruptedCalledBeforeBlendingOut : 1;
 
 	bool IsNotifyValid(FName NotifyName, const FBranchingPointNotifyPayload& BranchingPointNotifyPayload) const;
 	void UnbindDelegates();
diff --git a/Engine/Source/Runtime/Apple/AudioMixerAudioUnit/Private/AudioMixerPlatformAudioUnit.cpp b/Engine/Source/Runtime/Apple/AudioMixerAudioUnit/Private/AudioMixerPlatformAudioUnit.cpp
--- a/Engine/Source/Runtime/Apple/AudioMixerAudioUnit/Private/AudioMixerPlatformAudioUnit.cpp
+++ b/Engine/Source/Runtime/Apple/AudioMixerAudioUnit/Private/AudioMixerPlatformAudioUnit.cpp
@@ -41,14 +41,14 @@ namespace Audio
 #endif //#if PLATFORM_IOS || PLATFORM_TVOS
     static const double DefaultSampleRate = 48000.0;
     
+    static int32 SuspendCounter = 0;
+    
 	FMixerPlatformAudioUnit::FMixerPlatformAudioUnit()
 		: bInitialized(false)
         , bInCallback(false)
         , SubmittedBufferPtr(nullptr)
-#if PLATFORM_IOS || PLATFORM_TVOS
         , RemainingBytesInCurrentSubmittedBuffer(0)
         , BytesPerSubmittedBuffer(0)
-#endif //#if PLATFORM_IOS || PLATFORM_TVOS
 	{
 	}
 
@@ -88,6 +88,10 @@ namespace Audio
             BufferSize = DefaultBufferSize;
         }
         
+        BytesPerSubmittedBuffer = BufferSize * NumChannels * sizeof(float);
+        check(BytesPerSubmittedBuffer != 0);
+        UE_LOG(LogAudioMixerAudioUnit, Display, TEXT("Bytes per submitted buffer: %d"), BytesPerSubmittedBuffer);
+        
 #if PLATFORM_IOS || PLATFORM_TVOS
         NSError* error;
         
@@ -95,7 +99,7 @@ namespace Audio
         [AudioSession setPreferredSampleRate:GraphSampleRate error:&error];
         
         //Set the buffer size.
-        Float32 bufferSizeInSec = (float)BufferSize / OpenStreamParams.SampleRate;
+        Float32 bufferSizeInSec = (float)BufferSize / (OpenStreamParams.SampleRate == 0 ? GraphSampleRate : OpenStreamParams.SampleRate);
         [AudioSession setPreferredIOBufferDuration:bufferSizeInSec error:&error];
         
         if (error != nil)
@@ -103,10 +107,6 @@ namespace Audio
             UE_LOG(LogAudioMixerAudioUnit, Display, TEXT("Error setting buffer size to: %f"), bufferSizeInSec);
         }
         
-        BytesPerSubmittedBuffer = BufferSize * NumChannels * sizeof(float);
-        check(BytesPerSubmittedBuffer != 0);
-        UE_LOG(LogAudioMixerAudioUnit, Display, TEXT("Bytes per submitted buffer: %d"), BytesPerSubmittedBuffer);
-        
         [AudioSession setActive:true error:&error];
         
         // Retrieve the actual hardware sample rate
@@ -125,13 +125,6 @@ namespace Audio
         AudioDeviceQuerySize = sizeof(AudioDeviceID);
         Status = AudioObjectGetPropertyData(kAudioObjectSystemObject, &DevicePropertyAddress, 0, nullptr, &AudioDeviceQuerySize, &DeviceAudioObjectID);
         
-        //Query hardware sample rate:
-        DevicePropertyAddress.mSelector = kAudioDevicePropertyNominalSampleRate;
-        DevicePropertyAddress.mScope = kAudioObjectPropertyScopeGlobal;
-        DevicePropertyAddress.mElement = 0;
-        AudioDeviceQuerySize = sizeof(Float64);
-        Status = AudioObjectSetPropertyData(DeviceAudioObjectID, &DevicePropertyAddress, 0, nullptr, AudioDeviceQuerySize, &GraphSampleRate);
-        
         if(Status != 0)
         {
             UE_LOG(LogAudioMixerAudioUnit, Display, TEXT("ERROR setting sample rate to %f"), GraphSampleRate);
@@ -389,9 +382,8 @@ namespace Audio
 	void FMixerPlatformAudioUnit::SubmitBuffer(const uint8* Buffer)
 	{
         SubmittedBufferPtr = (uint8*) Buffer;
-#if PLATFORM_IOS || PLATFORM_TVOS
+        SubmittedBytes = 0;
         RemainingBytesInCurrentSubmittedBuffer = BytesPerSubmittedBuffer;
-#endif //#if PLATFORM_IOS || PLATFORM_TVOS
 	}
 
 	FName FMixerPlatformAudioUnit::GetRuntimeFormat(USoundWave* InSoundWave)
@@ -429,18 +421,20 @@ namespace Audio
 
 	void FMixerPlatformAudioUnit::ResumeContext()
 	{
-		if (bSuspended)
-		{
-			AUGraphStart(AudioUnitGraph);
-			UE_LOG(LogAudioMixerAudioUnit, Display, TEXT("Resuming Audio"));
-			bSuspended = false;
-		}
+        if (SuspendCounter > 0)
+        {
+            FPlatformAtomics::InterlockedDecrement(&SuspendCounter);
+            AUGraphStart(AudioUnitGraph);
+            UE_LOG(LogAudioMixerAudioUnit, Display, TEXT("Resuming Audio"));
+            bSuspended = false;
+        }
 	}
 	
 	void FMixerPlatformAudioUnit::SuspendContext()
 	{
-		if (!bSuspended)
+		if (SuspendCounter == 0)
 		{
+            FPlatformAtomics::InterlockedIncrement(&SuspendCounter);
 			AUGraphStop(AudioUnitGraph);
 			UE_LOG(LogAudioMixerAudioUnit, Display, TEXT("Suspending Audio"));
 			bSuspended = true;
@@ -467,8 +461,39 @@ namespace Audio
                 ReadNextBuffer();
             }
 
-            OutputBufferData->mBuffers[0].mData = (void*) SubmittedBufferPtr;
-            SubmittedBufferPtr = nullptr;
+            // How many bytes we have left over from previous callback
+            int32 SubmittedBufferBytesLeft = BytesPerSubmittedBuffer - SubmittedBytes;
+            int32 OutputBufferBytesLeft = OutputBufferData->mBuffers[0].mDataByteSize;
+            uint8* OutputBufferPtr = (uint8*) OutputBufferData->mBuffers[0].mData;
+            while (OutputBufferBytesLeft > 0)
+            {
+                const int32 BytesToCopy = FMath::Min(SubmittedBufferBytesLeft, OutputBufferBytesLeft);
+                
+                FMemory::Memcpy((void*) OutputBufferPtr, SubmittedBufferPtr + SubmittedBytes, BytesToCopy);
+                
+                OutputBufferBytesLeft -= BytesToCopy;
+                SubmittedBufferBytesLeft -= BytesToCopy;
+                
+                if (SubmittedBufferBytesLeft <= 0)
+                {
+                    ReadNextBuffer();
+                    SubmittedBytes = 0;
+                    SubmittedBufferBytesLeft = BytesPerSubmittedBuffer;
+                }
+                else
+                {
+                    SubmittedBytes += BytesToCopy;
+                }
+                
+                if (OutputBufferBytesLeft <= 0)
+                {
+                    break;
+                }
+                else
+                {
+                    OutputBufferPtr += BytesToCopy;
+                }
+            }
 		}
 		else
 		{
diff --git a/Engine/Source/Runtime/Apple/AudioMixerAudioUnit/Private/AudioMixerPlatformAudioUnit.h b/Engine/Source/Runtime/Apple/AudioMixerAudioUnit/Private/AudioMixerPlatformAudioUnit.h
--- a/Engine/Source/Runtime/Apple/AudioMixerAudioUnit/Private/AudioMixerPlatformAudioUnit.h
+++ b/Engine/Source/Runtime/Apple/AudioMixerAudioUnit/Private/AudioMixerPlatformAudioUnit.h
@@ -60,11 +60,10 @@ namespace Audio
 		AUNode      OutputNode;
 		AudioUnit	OutputUnit;
         uint8*      SubmittedBufferPtr;
-        
-#if PLATFORM_IOS || PLATFORM_TVOS
+        int32 SubmittedBytes = 0;
+
         int32       RemainingBytesInCurrentSubmittedBuffer;
         int32       BytesPerSubmittedBuffer;
-#endif //#if PLATFORM_IOS || PLATFORM_TVOS
         
 		bool PerformCallback(AudioBufferList* OutputBufferData);
 		void HandleError(const TCHAR* InLogOutput, bool bTeardown = true);
@@ -75,4 +74,3 @@ namespace Audio
 	};
 
 }
-
diff --git a/Engine/Source/Runtime/AudioMixer/Private/AudioMixer.cpp b/Engine/Source/Runtime/AudioMixer/Private/AudioMixer.cpp
--- a/Engine/Source/Runtime/AudioMixer/Private/AudioMixer.cpp
+++ b/Engine/Source/Runtime/AudioMixer/Private/AudioMixer.cpp
@@ -187,13 +187,15 @@ namespace Audio
 		, CurrentBufferReadIndex(INDEX_NONE)
 		, CurrentBufferWriteIndex(INDEX_NONE)
 		, NumOutputBuffers(0)
+		, TargetMasterVolume(1.0f)
+		, FadeVolume(0.0f)
 		, LastError(TEXT("None"))
 		, bAudioDeviceChanging(false)
-		, bFadingIn(true)
-		, bFadingOut(false)
+		, bPerformingFade(true)
 		, bFadedOut(false)
-		, FadeEnvelopeValue(0.0f)
+		, bIsDeviceInitialized(false)
 	{
+		FadeParam.SetValue(0.0f);
 	}
 
 	IAudioMixerPlatformInterface::~IAudioMixerPlatformInterface()
@@ -203,76 +205,82 @@ namespace Audio
 
 	void IAudioMixerPlatformInterface::FadeIn()
 	{
-		bFadingIn = true;
-		bFadingOut = false;
+		bPerformingFade = true;
 		bFadedOut = false;
+		FadeVolume = 1.0f;
 	}
 
 	void IAudioMixerPlatformInterface::FadeOut()
 	{
-		if (FadeEnvelopeValue == 0.0f)
+		if (bFadedOut || FadeVolume == 0.0f)
 		{
 			return;
 		}
 
-		bFadingOut = true;
+		bPerformingFade = true;
 		AudioFadeEvent->Wait();
+		FadeVolume = 0.0f;
 	}
 
-	void IAudioMixerPlatformInterface::PerformFades()
+	void IAudioMixerPlatformInterface::PostInitializeHardware()
+	{
+		bIsDeviceInitialized = true;
+	}
+
+	void IAudioMixerPlatformInterface::SetMasterVolume(const float InVolume)
+	{
+		if (TargetMasterVolume != InVolume)
+		{
+			TargetMasterVolume = InVolume;
+			bUpdateMasterVolume = true;
+		}
+	}
+
+	void IAudioMixerPlatformInterface::ApplyMasterAttenuation()
 	{
 		const int32 NextReadIndex = (CurrentBufferReadIndex + 1) % NumOutputBuffers;
 		FOutputBuffer& CurrentReadBuffer = OutputBuffers[NextReadIndex];
 
-		// Perform fade in and fade out global attenuation to avoid clicks/pops on startup/shutdown
-		if (bFadingIn)
+		AlignedFloatBuffer& Buffer = CurrentReadBuffer.GetBuffer();
+		float* BufferDataPtr = Buffer.GetData();
+		const int32 NumFrames = Buffer.Num();
+
+		// Apply master volume changes
+		if (bUpdateMasterVolume)
 		{
-			if (FadeEnvelopeValue >= 1.0f)
-			{
-				bFadingIn = false;
-				FadeEnvelopeValue = 1.0f;
-			}
-			else
-			{
-				AlignedFloatBuffer& Buffer = CurrentReadBuffer.GetBuffer();
-				const float Slope = 1.0f / Buffer.Num();
-				for (int32 i = 0; i < Buffer.Num(); ++i)
-				{
-					Buffer[i] *= FadeEnvelopeValue;
-					FadeEnvelopeValue += Slope;
-				}
-			}
+			bUpdateMasterVolume = false;
+			MasterVolumeParam.SetValue(TargetMasterVolume, NumFrames);
 		}
 
-		if (bFadingOut)
+		// Perform fade in and fade out global attenuation to avoid clicks/pops on startup/shutdown
+		if (bPerformingFade)
 		{
-			if (FadeEnvelopeValue <= 0.0f)
-			{
-				bFadingOut = false;
-				bFadedOut = true;
-				FadeEnvelopeValue = 0.0f;
-			}
-			else
-			{
-				AlignedFloatBuffer& Buffer = CurrentReadBuffer.GetBuffer();
-				const float Slope = 1.0f / Buffer.Num();
-				for (int32 i = 0; i < Buffer.Num(); ++i)
-				{
-					Buffer[i] *= FadeEnvelopeValue;
-					FadeEnvelopeValue -= Slope;
-				}
+			FadeParam.SetValue(FadeVolume, NumFrames);
 
+			for (int32 i = 0; i < NumFrames; ++i)
+			{
+				BufferDataPtr[i] *= MasterVolumeParam.Update() * FadeParam.Update();
 			}
 
+			bPerformingFade = false;
 			AudioFadeEvent->Trigger();
 		}
-
-		if (bFadedOut)
+		else if (bFadedOut)
 		{
 			// If we're faded out, then just zero the data.
-			AlignedFloatBuffer& Buffer = CurrentReadBuffer.GetBuffer();
-			FPlatformMemory::Memzero(Buffer.GetData(), sizeof(float)*Buffer.Num());
+			FPlatformMemory::Memzero((void*)Buffer.GetData(), sizeof(float)*NumFrames);
 		}
+		else
+		{
+			// Only apply the master volume
+			for (int32 i = 0; i < NumFrames; ++i)
+			{
+				BufferDataPtr[i] *= MasterVolumeParam.Update();
+			}
+		}
+
+		MasterVolumeParam.Reset();
+		FadeParam.Reset();
 	}
 
 	void IAudioMixerPlatformInterface::ReadNextBuffer()
@@ -305,7 +313,7 @@ namespace Audio
 			
 			if (!bWarnedBufferUnderrun)
 			{						
-				UE_LOG(LogAudioMixer, Error, TEXT("Audio Buffer Underrun detected."));
+				UE_LOG(LogAudioMixer, Log, TEXT("Audio Buffer Underrun detected."));
 				bWarnedBufferUnderrun = true;
 			}
 		
@@ -313,7 +321,7 @@ namespace Audio
 		}
 		else
 		{
-			PerformFades();
+			ApplyMasterAttenuation();
 
 			// As soon as a valid buffer goes through, allow more warning
 			if (bWarnedBufferUnderrun)
@@ -367,7 +375,7 @@ namespace Audio
 		check(AudioFadeEvent != nullptr);
 
 		check(AudioRenderThread == nullptr);
-		AudioRenderThread = FRunnableThread::Create(this, *FString::Printf(TEXT("AudioMixerRenderThread(%d)"), AudioMixerTaskCounter.Increment()), 0, TPri_Highest);
+		AudioRenderThread = FRunnableThread::Create(this, *FString::Printf(TEXT("AudioMixerRenderThread(%d)"), AudioMixerTaskCounter.Increment()), 0, TPri_Highest, FPlatformAffinity::GetAudioThreadMask());
 		check(AudioRenderThread != nullptr);
 	}
 
@@ -408,6 +416,24 @@ namespace Audio
 		}
 	}
 
+	void IAudioMixerPlatformInterface::Tick()
+	{
+		// In single-threaded mode, we simply render buffers until we run out of space
+		// The single-thread audio backend will consume these rendered buffers when they need to
+		if (AudioStreamInfo.StreamState == EAudioOutputStreamState::Running && bIsDeviceInitialized)
+		{
+			// Render mixed buffers till our queued buffers are filled up
+			while (CurrentBufferReadIndex != CurrentBufferWriteIndex)
+			{
+				RenderTimeAnalysis.Start();
+				OutputBuffers[CurrentBufferWriteIndex].MixNextBuffer();
+				RenderTimeAnalysis.End();
+
+				CurrentBufferWriteIndex = (CurrentBufferWriteIndex + 1) % NumOutputBuffers;
+			}
+		}
+	}
+
 	uint32 IAudioMixerPlatformInterface::MainAudioDeviceRun()
 	{
 		return RunInternal();
@@ -430,7 +456,7 @@ namespace Audio
 			RenderTimeAnalysis.Start();
 
 			// Render mixed buffers till our queued buffers are filled up
-			while (CurrentBufferReadIndex != CurrentBufferWriteIndex)
+			while (CurrentBufferReadIndex != CurrentBufferWriteIndex && bIsDeviceInitialized)
 			{
 				OutputBuffers[CurrentBufferWriteIndex].MixNextBuffer();
 
diff --git a/Engine/Source/Runtime/AudioMixer/Private/AudioMixerBuffer.cpp b/Engine/Source/Runtime/AudioMixer/Private/AudioMixerBuffer.cpp
--- a/Engine/Source/Runtime/AudioMixer/Private/AudioMixerBuffer.cpp
+++ b/Engine/Source/Runtime/AudioMixer/Private/AudioMixerBuffer.cpp
@@ -235,7 +235,7 @@ namespace Audio
 
 			case DTYPE_Procedural:
 			{
-				// Always create a new buffer for procedural buffers
+				// Always create a new buffer for procedural or bus buffers
 				Buffer = FMixerBuffer::CreateProceduralBuffer(Mixer, InWave);
 			}
 			break;
@@ -248,7 +248,6 @@ namespace Audio
 			break;
 
 			case DTYPE_Native:
-			case DTYPE_Xenon:
 			{
 				if (InWave->ResourceID)
 				{
diff --git a/Engine/Source/Runtime/AudioMixer/Private/AudioMixerBuffer.h b/Engine/Source/Runtime/AudioMixer/Private/AudioMixerBuffer.h
--- a/Engine/Source/Runtime/AudioMixer/Private/AudioMixerBuffer.h
+++ b/Engine/Source/Runtime/AudioMixer/Private/AudioMixerBuffer.h
@@ -59,6 +59,7 @@ namespace Audio
 		void EnsureHeaderParseTaskFinished();
 
 		float GetSampleRate() const { return SampleRate; }
+		int32 GetNumChannels() const { return NumChannels; }
 		void InitSampleRate(const float InSampleRate) { SampleRate = InSampleRate; }
 
 	private:
diff --git a/Engine/Source/Runtime/AudioMixer/Private/AudioMixerBus.cpp b/Engine/Source/Runtime/AudioMixer/Private/AudioMixerBus.cpp
new file mode 100644
--- /dev/null
+++ b/Engine/Source/Runtime/AudioMixer/Private/AudioMixerBus.cpp
@@ -0,0 +1,152 @@
+// Copyright 1998-2017 Epic Games, Inc. All Rights Reserved.
+
+#include "AudioMixerBus.h"
+#include "AudioMixerSourceManager.h"
+
+namespace Audio
+{
+	FMixerBus::FMixerBus(FMixerSourceManager* InSourceManager, const int32 InNumChannels, const int32 InNumFrames)
+		: CurrentBufferIndex(1)
+		, NumChannels(InNumChannels)
+		, NumFrames(InNumFrames)
+		, SourceManager(InSourceManager)
+	{
+		// Prepare the previous buffer with zero'd data
+		const int32 NumSamples = InNumChannels * InNumFrames;
+		for (int32 i = 0; i < 2; ++i)
+		{
+			MixedSourceData[i].Reset();
+			MixedSourceData[i].AddZeroed(NumSamples);
+		}
+	}
+
+	void FMixerBus::Update()
+	{
+		CurrentBufferIndex = !CurrentBufferIndex;
+	}
+
+	void FMixerBus::AddInstanceId(const int32 InSourceId)
+	{
+		InstanceIds.Add(InSourceId);
+	}
+
+	bool FMixerBus::RemoveInstanceId(const int32 InSourceId)
+	{
+		InstanceIds.Remove(InSourceId);
+
+		// Return true if there is no more instances or sends
+		return !InstanceIds.Num() && !BusSends.Num();
+	}
+
+	void FMixerBus::AddBusSend(const FBusSend& InBusSend)
+	{
+		BusSends.Add(InBusSend);
+	}
+
+	bool FMixerBus::RemoveBusSend(const int32 InSourceId)
+	{
+		for (int32 i = BusSends.Num() - 1; i >= 0; --i)
+		{
+			// Remove this source id's send
+			if (BusSends[i].SourceId == InSourceId)
+			{
+				BusSends.RemoveAtSwap(i, 1, false);
+
+				// There will only be one entry
+				break;
+			}
+		}
+
+		// Return true if there is no more instances or sends
+		return !InstanceIds.Num() && !BusSends.Num();
+	}
+
+	void FMixerBus::MixBuffer()
+	{
+		// Reset and zero the mixed source data buffer for this bus
+		const int32 NumSamples = NumFrames * NumChannels;
+
+		MixedSourceData[CurrentBufferIndex].Reset();
+		MixedSourceData[CurrentBufferIndex].AddZeroed(NumSamples);
+
+		float* BusDataBufferPtr = MixedSourceData[CurrentBufferIndex].GetData();
+
+		// Loop through the send list for this bus
+		for (FBusSend& BusSend : BusSends)
+		{
+			const float* SourceBufferPtr = nullptr;
+
+			// If the source is a bus, we need to use the previous renderer buffer 
+			if (SourceManager->IsBus(BusSend.SourceId))
+			{
+				SourceBufferPtr = SourceManager->GetPreviousBusBuffer(BusSend.SourceId);
+			}
+			// If the source mixing into this is not itself a bus, then simply mix the pre-attenuation audio of the source into the bus
+			// The source will have already computed its buffers for this frame
+			else
+			{
+				SourceBufferPtr = SourceManager->GetPreDistanceAttenuationBuffer(BusSend.SourceId);
+			}
+
+			const int32 NumSourceChannels = SourceManager->GetNumChannels(BusSend.SourceId);
+			const int32 NumSourceSamples = NumSourceChannels * SourceManager->GetNumOutputFrames();
+
+			// If source channels are 1 but the bus is 2 channels, we need to up-mix
+			if (NumSourceChannels == 1 && NumChannels == 2)
+			{
+				int32 BusSampleIndex = 0;
+				for (int32 SourceSampleIndex = 0; SourceSampleIndex < NumSourceSamples; ++SourceSampleIndex)
+				{
+					// Take half the source sample to up-mix to stereo
+					const float SourceSample = 0.5f * BusSend.SendLevel * SourceBufferPtr[SourceSampleIndex];
+
+					// Mix in the source sample
+					for (int32 BusChannel = 0; BusChannel < NumChannels; ++BusChannel)
+					{
+						BusDataBufferPtr[BusSampleIndex++] += SourceSample;
+					}
+				}
+			}
+			// If the source channels are 2 but the bus is 1 channel, we need to down-mix
+			else if (NumSourceChannels == 2 && NumChannels == 1)
+			{
+				int32 SourceSampleIndex = 0;
+				for (int32 BusSampleIndex = 0; BusSampleIndex < NumSourceSamples; ++BusSampleIndex)
+				{
+					// Downmix the stereo source to mono before summing to bus
+					float SourceSample = 0.0f;
+					for (int32 SourceChannel = 0; SourceChannel < 2; ++SourceChannel)
+					{
+						SourceSample += SourceBufferPtr[SourceSampleIndex++];
+					}
+
+					SourceSample *= 0.5f;
+					SourceSample *= BusSend.SendLevel;
+
+					BusDataBufferPtr[BusSampleIndex] += SourceSample;
+				}
+			}
+			// If they're the same channels, just mix it in
+			else
+			{
+				for (int32 SampleIndex = 0; SampleIndex < NumSourceSamples; ++SampleIndex)
+				{
+					BusDataBufferPtr[SampleIndex] += (BusSend.SendLevel * SourceBufferPtr[SampleIndex]);
+				}
+			}
+		}
+	}
+
+
+	const float* FMixerBus::GetCurrentBusBuffer() const
+	{
+		return MixedSourceData[CurrentBufferIndex].GetData();
+	}
+
+	const float* FMixerBus::GetPreviousBusBuffer() const
+	{
+		return MixedSourceData[!CurrentBufferIndex].GetData();
+	}
+
+}
+
diff --git a/Engine/Source/Runtime/AudioMixer/Private/AudioMixerDevice.cpp b/Engine/Source/Runtime/AudioMixer/Private/AudioMixerDevice.cpp
--- a/Engine/Source/Runtime/AudioMixer/Private/AudioMixerDevice.cpp
+++ b/Engine/Source/Runtime/AudioMixer/Private/AudioMixerDevice.cpp
@@ -13,6 +13,7 @@
 #include "DSP/SinOsc.h"
 #include "UObject/UObjectIterator.h"
 #include "Runtime/HeadMountedDisplay/Public/IHeadMountedDisplayModule.h"
+#include "Misc/App.h"
 
 #if WITH_EDITOR
 #include "AudioEditorModule.h"
@@ -172,26 +173,33 @@ namespace Audio
 				AudioClock = 0.0;
 				AudioClockDelta = (double)OpenStreamParams.NumFrames / OpenStreamParams.SampleRate;
 
+				FAudioPluginInitializationParams PluginInitializationParams;
+				PluginInitializationParams.NumSources = MaxChannels;
+				PluginInitializationParams.SampleRate = SampleRate;
+				PluginInitializationParams.BufferLength = OpenStreamParams.NumFrames;
+				PluginInitializationParams.AudioDevicePtr = this;
 
 				// Initialize any plugins if they exist
 				if (SpatializationPluginInterface.IsValid())
 				{
-					SpatializationPluginInterface->Initialize(SampleRate, MaxChannels, OpenStreamParams.NumFrames);
+					SpatializationPluginInterface->Initialize(PluginInitializationParams);
 				}
 
 				if (OcclusionInterface.IsValid())
 				{
-					OcclusionInterface->Initialize(SampleRate, MaxChannels, OpenStreamParams.NumFrames);
+					OcclusionInterface->Initialize(PluginInitializationParams);
 				}
 
 				if (ReverbPluginInterface.IsValid())
 				{
-					ReverbPluginInterface->Initialize(SampleRate, MaxChannels, OpenStreamParams.NumFrames);
+					ReverbPluginInterface->Initialize(PluginInitializationParams);
 				}
 
 				// Need to set these up before we start the audio stream.
 				InitSoundSubmixes();
 
+				AudioMixerPlatform->PostInitializeHardware();
+
 				// Start streaming audio
 				return AudioMixerPlatform->StartAudioStream();
 			}
@@ -221,7 +229,7 @@ namespace Audio
 		{
 			SourceManager.Update();
 
-			AudioMixerPlatform->UnRegisterDeviceChangedListener();
+			AudioMixerPlatform->UnregisterDeviceChangedListener();
 			AudioMixerPlatform->StopAudioStream();
 			AudioMixerPlatform->CloseAudioStream();
 			AudioMixerPlatform->TeardownHardware();
@@ -245,6 +253,18 @@ namespace Audio
 			// Audio rendering was suspended in CheckAudioDeviceChange if it changed.
 			AudioMixerPlatform->ResumePlaybackOnNewDevice();
 		}
+
+		// Update the master volume
+		if (IsAudioDeviceMuted())
+		{
+			AudioMixerPlatform->SetMasterVolume(0.0f);
+		}
+		else
+		{
+			float MasterVolume = GetPlatformAudioHeadroom();
+			MasterVolume *= GetTransientMasterVolume();
+			AudioMixerPlatform->SetMasterVolume(MasterVolume);
+		}
 	}
 
 	double FMixerDevice::GetAudioTime() const
@@ -314,10 +334,12 @@ namespace Audio
 
 	void FMixerDevice::ResumeContext()
 	{
+        AudioMixerPlatform->ResumeContext();
 	}
 
 	void FMixerDevice::SuspendContext()
 	{
+        AudioMixerPlatform->SuspendContext();
 	}
 
 	void FMixerDevice::EnableDebugAudioOutput()
@@ -327,9 +349,10 @@ namespace Audio
 
 	bool FMixerDevice::OnProcessAudioStream(AlignedFloatBuffer& Output)
 	{
-// Turn on to only hear PIE audio
-#if 0
-		if (IsMainAudioDevice())
+#if WITH_EDITOR
+		// Turn on to only hear PIE audio
+		bool bBypassMainAudioDevice = FParse::Param(FCommandLine::Get(), TEXT("AudioPIEOnly"));
+		if (bBypassMainAudioDevice && IsMainAudioDevice())
 		{
 			return true;
 		}
@@ -469,7 +492,8 @@ namespace Audio
 
 		// Reset existing submixes if they exist
 		Submixes.Reset();
-		// Make sure all submixs are registered but not initialized
+
+		// Make sure all submixes are registered but not initialized
 		for (TObjectIterator<USoundSubmix> It; It; ++It)
 		{
 			RegisterSoundSubmix(*It, false);
@@ -505,7 +529,7 @@ namespace Audio
 			SubmixInstance->Init(SoundSubmix);
 		}
 	}
-
+	
  	FAudioPlatformSettings FMixerDevice::GetPlatformSettings() const
  	{
 		FAudioPlatformSettings Settings = AudioMixerPlatform->GetPlatformSettings();
@@ -717,12 +741,21 @@ namespace Audio
 		return *MixerSubmix;		
 	}
 
-	FMixerSourceVoice* FMixerDevice::GetMixerSourceVoice(const FWaveInstance* InWaveInstance, ISourceBufferQueueListener* InBufferQueueListener, bool bUseHRTFSpatialization)
+	FMixerSourceVoice* FMixerDevice::GetMixerSourceVoice()
 	{
-		// Create a new mixer source voice using our source manager
-		FMixerSourceVoice* NewMixerSourceVoice = new FMixerSourceVoice(this, &SourceManager);
+		FMixerSourceVoice* Voice = nullptr;
+		if (!SourceVoices.Dequeue(Voice))
+		{
+			Voice = new FMixerSourceVoice();
+		}
+
+		Voice->Reset(this);
+		return Voice;
+	}
 
-		return NewMixerSourceVoice;
+	void FMixerDevice::ReleaseMixerSourceVoice(FMixerSourceVoice* InSourceVoice)
+	{
+		SourceVoices.Enqueue(InSourceVoice);
 	}
 
 	int32 FMixerDevice::GetNumSources() const
diff --git a/Engine/Source/Runtime/AudioMixer/Private/AudioMixerSource.cpp b/Engine/Source/Runtime/AudioMixer/Private/AudioMixerSource.cpp
--- a/Engine/Source/Runtime/AudioMixer/Private/AudioMixerSource.cpp
+++ b/Engine/Source/Runtime/AudioMixer/Private/AudioMixerSource.cpp
@@ -18,7 +18,6 @@ FAutoConsoleVariableRef CVarDisableHRTF(
 
 namespace Audio
 {
-
 	FMixerSource::FMixerSource(FAudioDevice* InAudioDevice)
 		: FSoundSource(InAudioDevice)
 		, MixerDevice((FMixerDevice*)InAudioDevice)
@@ -61,7 +60,7 @@ namespace Audio
 		// Get the number of frames before creating the buffer
 		int32 NumFrames = INDEX_NONE;
 
-		if (InWaveInstance->WaveData->DecompressionType != EDecompressionType::DTYPE_Procedural)
+		if (InWaveInstance->WaveData->DecompressionType != DTYPE_Procedural)
 		{
 			const int32 NumBytes = InWaveInstance->WaveData->RawPCMDataSize;
 			NumFrames = NumBytes / (InWaveInstance->WaveData->NumChannels * sizeof(int16));
@@ -76,7 +75,7 @@ namespace Audio
 			SCOPE_CYCLE_COUNTER(STAT_AudioSourceInitTime);
 
 			AUDIO_MIXER_CHECK(MixerDevice);
-			MixerSourceVoice = MixerDevice->GetMixerSourceVoice(InWaveInstance, this, UseObjectBasedSpatialization());
+			MixerSourceVoice = MixerDevice->GetMixerSourceVoice();
 			if (!MixerSourceVoice)
 			{
 				return false;
@@ -93,60 +92,102 @@ namespace Audio
 
 			InitParams.SourceEffectChainId = 0;
 
-			if (InitParams.NumInputChannels <= 2 && InWaveInstance->SourceEffectChain)
+			if (InitParams.NumInputChannels <= 2)
 			{
-				InitParams.SourceEffectChainId = InWaveInstance->SourceEffectChain->GetUniqueID();
-
-				for (int32 i = 0; i < InWaveInstance->SourceEffectChain->Chain.Num(); ++i)
+				if (InWaveInstance->SourceEffectChain)
 				{
-					InitParams.SourceEffectChain.Add(InWaveInstance->SourceEffectChain->Chain[i]);
-					InitParams.bPlayEffectChainTails = InWaveInstance->SourceEffectChain->bPlayEffectChainTails;
+					InitParams.SourceEffectChainId = InWaveInstance->SourceEffectChain->GetUniqueID();
+
+					for (int32 i = 0; i < InWaveInstance->SourceEffectChain->Chain.Num(); ++i)
+					{
+						InitParams.SourceEffectChain.Add(InWaveInstance->SourceEffectChain->Chain[i]);
+						InitParams.bPlayEffectChainTails = InWaveInstance->SourceEffectChain->bPlayEffectChainTails;
+					}
 				}
-			}
 
-			// If we've overridden which submix we're sending the sound, then add that as the first send
-			if (InWaveInstance->SoundSubmix != nullptr)
-			{
-				FMixerSourceSubmixSend SubmixSend;
-				SubmixSend.Submix = MixerDevice->GetSubmixInstance(InWaveInstance->SoundSubmix);
-				SubmixSend.SendLevel = 1.0f;
-				SubmixSend.bIsMainSend = true;
-				InitParams.SubmixSends.Add(SubmixSend);
-			}
-			else
-			{
-				// Send the voice to the EQ submix if it's enabled
-				const bool bIsEQDisabled = GetDefault<UAudioSettings>()->bDisableMasterEQ;
-				if (!bIsEQDisabled && IsEQFilterApplied())
+				// Setup the bus Id if this source is a bus
+				if (InWaveInstance->WaveData->bIsBus)
 				{
-					// Default the submix to use to use the master submix if none are set
-					FMixerSourceSubmixSend SubmixSend;
-					SubmixSend.Submix = MixerDevice->GetMasterEQSubmix();
-					SubmixSend.SendLevel = 1.0f;
-					SubmixSend.bIsMainSend = true;
-					InitParams.SubmixSends.Add(SubmixSend);
+					InitParams.BusId = InWaveInstance->WaveData->GetUniqueID();
+					if (!InWaveInstance->WaveData->IsLooping())
+					{
+						InitParams.BusDuration = InWaveInstance->WaveData->GetDuration();
+					}
 				}
-				else
+
+				// Toggle muting the source if sending only to output bus. 
+				// This can get set even if the source doesn't have bus sends since bus sends can be dynamically enabled.
+				InitParams.bOutputToBusOnly = InWaveInstance->bOutputToBusOnly;
+
+				// If this source is sending its audio to a bus
+				if (InWaveInstance->SoundSourceBusSends.Num() > 0)
 				{
-					// Default the submix to use to use the master submix if none are set
-					FMixerSourceSubmixSend SubmixSend;
-					SubmixSend.Submix = MixerDevice->GetMasterSubmix();
-					SubmixSend.SendLevel = 1.0f;
-					SubmixSend.bIsMainSend = true;
-					InitParams.SubmixSends.Add(SubmixSend);
+					// And add all the source bus sends
+					for (FSoundSourceBusSendInfo& SendInfo : InWaveInstance->SoundSourceBusSends)
+					{
+						if (SendInfo.SoundSourceBus != nullptr)
+						{
+							FMixerBusSend BusSend;
+							BusSend.BusId = SendInfo.SoundSourceBus->GetUniqueID();
+							BusSend.SendLevel = SendInfo.SendLevel;
+							InitParams.BusSends.Add(BusSend);
+						}
+					}
 				}
 			}
 
-			// Now add any addition submix sends for this source
-			for (FSoundSubmixSendInfo& SendInfo : InWaveInstance->SoundSubmixSends)
+			// Don't set up any submixing if we're set to output to bus only
+			if (!InitParams.bOutputToBusOnly)
 			{
-				if(SendInfo.SoundSubmix != nullptr)
+				// If we're spatializing using HRTF and its an external send, don't need to setup a default/base submix send to master or EQ submix
+				// We'll only be using non-default submix sends (e.g. reverb).
+				if (!(InWaveInstance->SpatializationMethod == ESoundSpatializationAlgorithm::SPATIALIZATION_HRTF && MixerDevice->bSpatializationIsExternalSend))
+				{
+					// If we've overridden which submix we're sending the sound, then add that as the first send
+					if (InWaveInstance->SoundSubmix != nullptr)
+					{
+						FMixerSourceSubmixSend SubmixSend;
+						SubmixSend.Submix = MixerDevice->GetSubmixInstance(InWaveInstance->SoundSubmix);
+						SubmixSend.SendLevel = 1.0f;
+						SubmixSend.bIsMainSend = true;
+						InitParams.SubmixSends.Add(SubmixSend);
+					}
+					else
+					{
+						// Send the voice to the EQ submix if it's enabled
+						const bool bIsEQDisabled = GetDefault<UAudioSettings>()->bDisableMasterEQ;
+						if (!bIsEQDisabled && IsEQFilterApplied())
+						{
+							// Default the submix to use to use the master submix if none are set
+							FMixerSourceSubmixSend SubmixSend;
+							SubmixSend.Submix = MixerDevice->GetMasterEQSubmix();
+							SubmixSend.SendLevel = 1.0f;
+							SubmixSend.bIsMainSend = true;
+							InitParams.SubmixSends.Add(SubmixSend);
+						}
+						else
+						{
+							// Default the submix to use to use the master submix if none are set
+							FMixerSourceSubmixSend SubmixSend;
+							SubmixSend.Submix = MixerDevice->GetMasterSubmix();
+							SubmixSend.SendLevel = 1.0f;
+							SubmixSend.bIsMainSend = true;
+							InitParams.SubmixSends.Add(SubmixSend);
+						}
+					}
+				}
+
+				// Now add any addition submix sends for this source
+				for (FSoundSubmixSendInfo& SendInfo : InWaveInstance->SoundSubmixSends)
 				{
-					FMixerSourceSubmixSend SubmixSend;
-					SubmixSend.Submix = MixerDevice->GetSubmixInstance(SendInfo.SoundSubmix);
-					SubmixSend.SendLevel = SendInfo.SendLevel;
-					SubmixSend.bIsMainSend = false;
-					InitParams.SubmixSends.Add(SubmixSend);
+					if (SendInfo.SoundSubmix != nullptr)
+					{
+						FMixerSourceSubmixSend SubmixSend;
+						SubmixSend.Submix = MixerDevice->GetSubmixInstance(SendInfo.SoundSubmix);
+						SubmixSend.SendLevel = SendInfo.SendLevel;
+						SubmixSend.bIsMainSend = false;
+						InitParams.SubmixSends.Add(SubmixSend);
+					}
 				}
 			}
 
@@ -276,10 +317,10 @@ namespace Audio
 			// Set up buffer areas to decompress audio into
 			for (int32 BufferIndex = 0; BufferIndex < Audio::MAX_BUFFERS_QUEUED; ++BufferIndex)
 			{
-				const uint32 BufferSize = MONO_PCM_BUFFER_SIZE * MixerBuffer->NumChannels;
+				const uint32 TotalSamples = MONO_PCM_BUFFER_SAMPLES * MixerBuffer->NumChannels;
 				SourceVoiceBuffers[BufferIndex]->AudioData.Reset();
-				SourceVoiceBuffers[BufferIndex]->AudioData.AddZeroed(BufferSize);
-				SourceVoiceBuffers[BufferIndex]->AudioBytes = BufferSize;
+				SourceVoiceBuffers[BufferIndex]->AudioData.AddZeroed(TotalSamples);
+				SourceVoiceBuffers[BufferIndex]->Samples = TotalSamples;
 				SourceVoiceBuffers[BufferIndex]->bRealTimeBuffer = true;
 				SourceVoiceBuffers[BufferIndex]->LoopCount = 0;
 			}
@@ -303,6 +344,11 @@ namespace Audio
 				// not ready
 				return AsyncRealtimeAudioTask->IsDone();
 			}
+			else if (WaveInstance->WaveData->bIsBus)
+			{
+				// Buses don't need to do anything to play audio
+				return true;
+			}
 			else
 			{
 				// Now check to see if we need to kick off a decode the first chunk of audio
@@ -310,10 +356,10 @@ namespace Audio
 				if ((BufferType == EBufferType::PCMRealTime || BufferType == EBufferType::Streaming) && WaveInstance->WaveData)
 				{
 					// If any of these conditions meet, we need to do an initial async decode before we're ready to start playing the sound
-					if (WaveInstance->StartTime > 0.0f || WaveInstance->WaveData->bProcedural || !WaveInstance->WaveData->CachedRealtimeFirstBuffer)
+					if (WaveInstance->StartTime > 0.0f || WaveInstance->WaveData->bProcedural || WaveInstance->WaveData->bIsBus || !WaveInstance->WaveData->CachedRealtimeFirstBuffer)
 					{
 						// Before reading more PCMRT data, we first need to seek the buffer
-						if (WaveInstance->StartTime > 0.0f)
+						if (WaveInstance->StartTime > 0.0f && !WaveInstance->WaveData->bIsBus && !WaveInstance->WaveData->bProcedural)
 						{
 							MixerBuffer->Seek(WaveInstance->StartTime);
 						}
@@ -405,6 +451,16 @@ namespace Audio
 				WaveInstance->NotifyFinished();
 				return true;
 			}
+			// Buses don't do buffer end callbacks, so we need to directly query buses doneness
+			else if (WaveInstance->WaveData->bIsBus)
+			{
+				if (MixerSourceVoice->IsSourceEffectTailsDone() && MixerSourceVoice->IsDone())
+				{
+					bIsFinished = true;
+					WaveInstance->NotifyFinished();
+					return true;
+				}
+			}
 
 			if (bLoopCallback && WaveInstance->LoopingMode == LOOP_WithNotification)
 			{
@@ -454,45 +510,55 @@ namespace Audio
 
 		CurrentBuffer = 0;
 
-		uint8* Data = nullptr;
-		uint32 DataSize = 0;
-		MixerBuffer->GetPCMData(&Data, &DataSize);
+		RawPCMDataBuffer.Data = nullptr;
+		RawPCMDataBuffer.DataSize = 0;
+		MixerBuffer->GetPCMData(&RawPCMDataBuffer.Data, &RawPCMDataBuffer.DataSize);
+
+		RawPCMDataBuffer.NumSamples = RawPCMDataBuffer.DataSize / sizeof(int16);
+		RawPCMDataBuffer.CurrentSample = 0;
 
 		// Only submit data if we've successfully loaded it
-		if (!Data || !DataSize)
+		if (!RawPCMDataBuffer.Data || !RawPCMDataBuffer.DataSize)
 		{
 			UE_LOG(LogAudioMixer, Error, TEXT("Failed to load PCM data from sound source %s"), *WaveInstance->GetName());
 			return;
 		}
 
-		// Reset the data, copy it over
-		SourceVoiceBuffers[0]->AudioData.Reset();
-		SourceVoiceBuffers[0]->AudioData.AddDefaulted(DataSize);
-		FMemory::Memcpy(SourceVoiceBuffers[0]->AudioData.GetData(), Data, DataSize);
-
-		AUDIO_MIXER_CHECK(SourceVoiceBuffers[0]->AudioData.Num() == DataSize);
+		RawPCMDataBuffer.LoopCount = (WaveInstance->LoopingMode != LOOP_Never) ? Audio::LOOP_FOREVER : 0;
 
-		// Set the size of the data. In general AudioBytes may not be the same as AudioData.Num() 
-		SourceVoiceBuffers[0]->AudioBytes = DataSize;
-		SourceVoiceBuffers[0]->bRealTimeBuffer = false;
-		SourceVoiceBuffers[0]->LoopCount = (WaveInstance->LoopingMode != LOOP_Never) ? Audio::LOOP_FOREVER : 0;
+		// Submit the first two format-converted chunks to the source voice
+		const uint32 NumSamplesPerBuffer = MONO_PCM_BUFFER_SAMPLES * MixerBuffer->NumChannels;
+		int16* RawPCMBufferDataPtr = (int16*)RawPCMDataBuffer.Data;
 
+		RawPCMDataBuffer.GetNextBuffer(SourceVoiceBuffers[0], NumSamplesPerBuffer);
 		MixerSourceVoice->SubmitBuffer(SourceVoiceBuffers[0], false);
+
+		CurrentBuffer = 1;
 	}
 
 	void FMixerSource::SubmitPCMRTBuffers()
 	{
 		CurrentBuffer = 0;
 
-		const uint32 BufferSize = MONO_PCM_BUFFER_SIZE * MixerBuffer->NumChannels;
-
 		bPlayedCachedBuffer = false;
 		bool bIsSeeking = (WaveInstance->StartTime > 0.0f);
 		if (!bIsSeeking && WaveInstance->WaveData && WaveInstance->WaveData->CachedRealtimeFirstBuffer)
 		{
 			bPlayedCachedBuffer = true;
-			FMemory::Memcpy(SourceVoiceBuffers[0]->AudioData.GetData(), WaveInstance->WaveData->CachedRealtimeFirstBuffer, BufferSize);
-			FMemory::Memcpy(SourceVoiceBuffers[1]->AudioData.GetData(), WaveInstance->WaveData->CachedRealtimeFirstBuffer + BufferSize, BufferSize);
+
+			// Format convert the first cached buffers
+			const uint32 NumSamples = MONO_PCM_BUFFER_SAMPLES * MixerBuffer->NumChannels;
+			const uint32 BufferSize = MONO_PCM_BUFFER_SIZE * MixerBuffer->NumChannels;
+
+			int16* CachedBufferPtr0 = (int16*)WaveInstance->WaveData->CachedRealtimeFirstBuffer;
+			int16* CachedBufferPtr1 = (int16*)(WaveInstance->WaveData->CachedRealtimeFirstBuffer + BufferSize);
+			float* AudioData0 = SourceVoiceBuffers[0]->AudioData.GetData();
+			float* AudioData1 = SourceVoiceBuffers[1]->AudioData.GetData();
+			for (uint32 Sample = 0; Sample < NumSamples; ++Sample)
+			{
+				AudioData0[Sample] = CachedBufferPtr0[Sample] / 32768.0f;
+				AudioData1[Sample] = CachedBufferPtr1[Sample] / 32768.0f;
+			}
 
 			// Submit the already decoded and cached audio buffers
 			MixerSourceVoice->SubmitBuffer(SourceVoiceBuffers[0], false);
@@ -500,7 +566,7 @@ namespace Audio
 
 			CurrentBuffer = 2;
 		}
-		else
+		else if (WaveInstance->WaveData && !WaveInstance->WaveData->bIsBus)
 		{
 			// We should have already kicked off and finished a task. 
 			check(AsyncRealtimeAudioTask != nullptr);
@@ -515,48 +581,48 @@ namespace Audio
 	{
 		USoundWave* WaveData = WaveInstance->WaveData;
 
-		if (WaveData && WaveData->bProcedural)
+		if (WaveData)
 		{
-			const int32 MaxSamples = (MONO_PCM_BUFFER_SIZE * Buffer->NumChannels) / sizeof(int16);
+			if (WaveData->bProcedural)
+			{
+				const int32 MaxSamples = MONO_PCM_BUFFER_SAMPLES * Buffer->NumChannels;
 
-			if (BufferReadMode == EBufferReadMode::Synchronous || WaveData->bCanProcessAsync == false)
+				FProceduralAudioTaskData NewTaskData;
+				NewTaskData.ProceduralSoundWave = WaveData;
+				NewTaskData.AudioData = SourceVoiceBuffers[BufferIndex]->AudioData.GetData();
+				NewTaskData.NumSamples = MaxSamples;
+				NewTaskData.NumChannels = Buffer->NumChannels;
+				check(!AsyncRealtimeAudioTask);
+				AsyncRealtimeAudioTask = CreateAudioTask(NewTaskData);
+
+				// Procedural sound waves never loop
+				return false;
+			}
+			else if (!MixerBuffer->IsRealTimeBuffer())
 			{
-				const int32 BytesWritten = WaveData->GeneratePCMData(SourceVoiceBuffers[BufferIndex]->AudioData.GetData(), MaxSamples);
-				SourceVoiceBuffers[BufferIndex]->AudioBytes = BytesWritten;
+				check(RawPCMDataBuffer.Data != nullptr);
+
+				// Read the next raw PCM buffer into the source buffer index. This converts raw PCM to float.
+				const uint32 NumSamplesPerBuffer = MONO_PCM_BUFFER_SAMPLES * MixerBuffer->NumChannels;
+				return RawPCMDataBuffer.GetNextBuffer(SourceVoiceBuffers[BufferIndex], NumSamplesPerBuffer);
 			}
 			else
 			{
-				FProceduralAudioTaskData NewTaskData;
-				NewTaskData.ProceduralSoundWave = Cast<USoundWaveProcedural>(WaveData);
+				FDecodeAudioTaskData NewTaskData;
+				NewTaskData.MixerBuffer = MixerBuffer;
 				NewTaskData.AudioData = SourceVoiceBuffers[BufferIndex]->AudioData.GetData();
-				NewTaskData.MaxAudioDataSamples = MaxSamples;
+				NewTaskData.bLoopingMode = WaveInstance->LoopingMode != LOOP_Never;
+				NewTaskData.bSkipFirstBuffer = (BufferReadMode == EBufferReadMode::AsynchronousSkipFirstFrame);
+				NewTaskData.NumFramesToDecode = MONO_PCM_BUFFER_SAMPLES;
 
 				check(!AsyncRealtimeAudioTask);
 				AsyncRealtimeAudioTask = CreateAudioTask(NewTaskData);
-			}
-
-			// Not looping
-			return false;
-		}
-		else if (BufferReadMode == EBufferReadMode::Synchronous)
-		{
-			return MixerBuffer->ReadCompressedData(SourceVoiceBuffers[BufferIndex]->AudioData.GetData(), WaveInstance->LoopingMode != LOOP_Never);
-		}
-		else
-		{
-			FDecodeAudioTaskData NewTaskData;
-			NewTaskData.MixerBuffer = MixerBuffer;
-			NewTaskData.AudioData = SourceVoiceBuffers[BufferIndex]->AudioData.GetData();
-			NewTaskData.bLoopingMode = WaveInstance->LoopingMode != LOOP_Never;
-			NewTaskData.bSkipFirstBuffer = (BufferReadMode == EBufferReadMode::AsynchronousSkipFirstFrame);
-			NewTaskData.NumFramesToDecode = MONO_PCM_BUFFER_SAMPLES;
-
-			check(!AsyncRealtimeAudioTask);
-			AsyncRealtimeAudioTask = CreateAudioTask(NewTaskData);
 
-			// Not looping
-			return false;
+				// Not looping
+				return false;
+			}
 		}
+		return true;
 	}
 
 	void FMixerSource::SubmitRealTimeSourceData(const bool bLooped, const bool bSubmitSynchronously)
@@ -611,8 +677,7 @@ namespace Audio
 						FDecodeAudioTaskResults TaskResult;
 						AsyncRealtimeAudioTask->GetResult(TaskResult);
 
-						const uint32 BufferSize = MONO_PCM_BUFFER_SIZE * MixerBuffer->NumChannels;
-						SourceVoiceBuffers[CurrentBuffer]->AudioBytes = BufferSize;
+						SourceVoiceBuffers[CurrentBuffer]->Samples = MONO_PCM_BUFFER_SAMPLES * MixerBuffer->NumChannels;
 						bLooped = TaskResult.bLooped;
 					}
 					break;
@@ -622,7 +687,7 @@ namespace Audio
 						FProceduralAudioTaskResults TaskResult;
 						AsyncRealtimeAudioTask->GetResult(TaskResult);
 
-						SourceVoiceBuffers[CurrentBuffer]->AudioBytes = TaskResult.NumBytesWritten;
+						SourceVoiceBuffers[CurrentBuffer]->Samples = TaskResult.NumSamplesWritten;
 					}
 					break;
 				}
@@ -670,13 +735,11 @@ namespace Audio
 		{
 			int32 BuffersQueued = MixerSourceVoice->GetNumBuffersQueued();
 
-			const bool bIsRealTimeSource = MixerBuffer->IsRealTimeBuffer();
-
-			if (BuffersQueued == 0 && (bBuffersToFlush || !bIsRealTimeSource))
+			if (BuffersQueued == 0 && bBuffersToFlush)
 			{
 				bIsFinished = true;
 			}
-			else if (bIsRealTimeSource && !bBuffersToFlush && BuffersQueued <= (Audio::MAX_BUFFERS_QUEUED - 1))
+			else if (!bBuffersToFlush && BuffersQueued <= (Audio::MAX_BUFFERS_QUEUED - 1))
 			{
 				// OnSourceBufferEnd is always called from render thread and the source needs to be 
 				// processed and any decoded buffers submitted to render thread synchronously			
@@ -782,6 +845,9 @@ namespace Audio
 			PendingReleases.Enqueue(PendingDecodeTask);
 		}
 
+		// Reset the raw PCM buffer data
+		RawPCMDataBuffer = FRawPCMDataBuffer();
+
 		MixerBuffer = nullptr;
 		AsyncRealtimeAudioTask = nullptr;
 		Buffer = nullptr;
@@ -819,18 +885,11 @@ namespace Audio
 
 	void FMixerSource::UpdateVolume()
 	{
-		float CurrentVolume = 0.0f;
-
-		if (!AudioDevice->IsAudioDeviceMuted())
-		{
-			CurrentVolume = WaveInstance->GetActualVolume();
-		}
-
-		CurrentVolume = GetDebugVolume(CurrentVolume);
-
-		CurrentVolume = FMath::Clamp<float>(CurrentVolume * AudioDevice->GetPlatformAudioHeadroom(), 0.0f, MAX_VOLUME);
+		float CurrentVolume = WaveInstance->GetVolume();
+		CurrentVolume = FMath::Clamp<float>(GetDebugVolume(CurrentVolume), 0.0f, MAX_VOLUME);
 
 		MixerSourceVoice->SetVolume(CurrentVolume);
+		MixerSourceVoice->SetDistanceAttenuation(WaveInstance->GetDistanceAttenuation());
 	}
 
 	void FMixerSource::UpdateSpatialization()
@@ -1021,7 +1080,6 @@ namespace Audio
 	bool FMixerSource::UseObjectBasedSpatialization() const
 	{
 		return (Buffer->NumChannels == 1 &&
-				AudioDevice->AudioPlugin != nullptr &&
 				AudioDevice->IsSpatializationPluginEnabled() &&
 				DisableHRTFCvar == 0 &&
 				WaveInstance->SpatializationMethod == ESoundSpatializationAlgorithm::SPATIALIZATION_HRTF);
@@ -1047,4 +1105,53 @@ namespace Audio
 				AudioDevice->IsReverbPluginEnabled() &&
 				WaveInstance->ReverbPluginSettings != nullptr;
 	}
+
+	bool FRawPCMDataBuffer::GetNextBuffer(FMixerSourceBufferPtr OutSourceBufferPtr, const uint32 NumSampleToGet)
+	{
+		// TODO: support loop counts
+		float* OutBufferPtr = OutSourceBufferPtr->AudioData.GetData();
+		int16* DataPtr = (int16*)Data;
+
+		if (LoopCount == Audio::LOOP_FOREVER)
+		{
+			bool bLooped = false;
+			for (uint32 Sample = 0; Sample < NumSampleToGet; ++Sample)
+			{
+				OutBufferPtr[Sample] = DataPtr[CurrentSample++] / 32768.0f;
+
+				// Loop around if we're looping
+				if (CurrentSample >= NumSamples)
+				{
+					CurrentSample = 0;
+					bLooped = true;
+				}
+			}
+			return bLooped;
+		}
+		else if (CurrentSample < NumSamples)
+		{
+			uint32 Sample = 0;
+			while (Sample < NumSampleToGet && CurrentSample < NumSamples)
+			{
+				OutBufferPtr[Sample++] = (float)DataPtr[CurrentSample++] / 32768.0f;
+			}
+
+			// Zero out the rest of the buffer
+			while (Sample < NumSampleToGet)
+			{
+				OutBufferPtr[Sample++] = 0.0f;
+			}
+		}
+		else
+		{
+			for (uint32 Sample = 0; Sample < NumSampleToGet; ++Sample)
+			{
+				OutBufferPtr[Sample] = 0.0f;
+			}
+		}
+
+		// If the current sample is greater or equal to num samples we hit the end of the buffer
+		return CurrentSample >= NumSamples;
+	}
+
 }
\ No newline at end of file
diff --git a/Engine/Source/Runtime/AudioMixer/Private/AudioMixerSource.h b/Engine/Source/Runtime/AudioMixer/Private/AudioMixerSource.h
--- a/Engine/Source/Runtime/AudioMixer/Private/AudioMixerSource.h
+++ b/Engine/Source/Runtime/AudioMixer/Private/AudioMixerSource.h
@@ -2,11 +2,7 @@
 
 #pragma once
 
-/* Public dependencies
-*****************************************************************************/
-
 #include "CoreMinimal.h"
-
 #include "AudioMixerBuffer.h"
 #include "AudioMixerSourceManager.h"
 
@@ -21,26 +17,24 @@ namespace Audio
 	static const int32 MAX_BUFFERS_QUEUED = 3;
 	static const int32 LOOP_FOREVER = -1;
 
-
-	/** Struct defining a source voice buffer. */
-	struct FMixerSourceVoiceBuffer
+	struct FRawPCMDataBuffer
 	{
-		/** Raw audio PCM data. */
-		TArray<uint8> AudioData;
-
-		/** The amount of real audio data in the byte array (may not be the size of the array in the case of procedural audio. */
-		uint32 AudioBytes;
-
-		/** How many times this buffer will loop. */
+		uint8* Data;
+		uint32 DataSize;
 		int32 LoopCount;
+		uint32 CurrentSample;
+		uint32 NumSamples;
 
-		/** If this buffer is from real-time decoding and needs to make callbacks for more data. */
-		uint32 bRealTimeBuffer : 1;
+		bool GetNextBuffer(FMixerSourceBufferPtr OutSourceBufferPtr, const uint32 NumSampleToGet);
 
-		FMixerSourceVoiceBuffer()
-		{
-			FMemory::Memzero(this, sizeof(*this));
-		}
+
+		FRawPCMDataBuffer()
+			: Data(nullptr)
+			, DataSize(0)
+			, LoopCount(0)
+			, CurrentSample(0)
+			, NumSamples(0)
+		{}
 	};
 
 	/** 
@@ -182,6 +176,9 @@ namespace Audio
 		// flushed/destroyed sound can destroy the buffer before the sound finishes using the buffer.
 		TArray<FMixerSourceBufferPtr> SourceVoiceBuffers;
 
+		// Raw uncompressed, non-float PCM data (int16)
+		FRawPCMDataBuffer RawPCMDataBuffer;
+
 		FSpatializationParams SpatializationParams;
 
 		FThreadSafeBool bPlayedCachedBuffer;
diff --git a/Engine/Source/Runtime/AudioMixer/Private/AudioMixerSourceDecode.cpp b/Engine/Source/Runtime/AudioMixer/Private/AudioMixerSourceDecode.cpp
--- a/Engine/Source/Runtime/AudioMixer/Private/AudioMixerSourceDecode.cpp
+++ b/Engine/Source/Runtime/AudioMixer/Private/AudioMixerSourceDecode.cpp
@@ -9,147 +9,9 @@
 #include "AudioMixerBuffer.h"
 #include "Async/Async.h"
 
-// Toggles using futures for audio tasks
-#define AUDIO_DECODE_USE_FUTURES 0
-
 namespace Audio
 {
 
-#if AUDIO_DECODE_USE_FUTURES
-
-	class FDecodeHandleBase : public IAudioTask
-	{
-	public:
-		FDecodeHandleBase()
-			: bFinished(false)
-		{}
-
-		virtual ~FDecodeHandleBase()
-		{
-			EnsureCompletion();
-		}
-
-		virtual bool IsDone() const override
-		{
-			return bFinished;
-		}
-
-		virtual void EnsureCompletion() override
-		{
-			Future.Get();
-		}
-
-	protected:
-		FThreadSafeBool bFinished;
-		TFuture<void> Future;
-	};
-
-	class FHeaderDecodeHandle : public FDecodeHandleBase
-	{
-	public:
-		virtual EAudioTaskType GetType() const override { return EAudioTaskType::Header; }
-
-		void Decode(const FHeaderParseAudioTaskData& InJobData)
-		{
-			TFunction<void()> Func = [this, InJobData]()
-			{
-				InJobData.MixerBuffer->ReadCompressedInfo(InJobData.SoundWave);
-				bFinished = true;
-			};
-
-			Future = Async(EAsyncExecution::TaskGraph, Func);
-		}
-	};
-
-	class FProceduralDecodeHandle : public FDecodeHandleBase
-	{
-	public:
-		virtual EAudioTaskType GetType() const override { return EAudioTaskType::Procedural; }
-
-		void Decode(const FProceduralAudioTaskData& InJobData)
-		{
-			TFunction<void()> Func = [this, InJobData]()
-			{
-				Result.NumBytesWritten = InJobData.ProceduralSoundWave->GeneratePCMData(InJobData.AudioData, InJobData.MaxAudioDataSamples);
-				bFinished = true;
-			};
-
-			Future = Async(EAsyncExecution::TaskGraph, Func);
-		}
-
-		virtual void GetResult(FProceduralAudioTaskResults& OutResult) override
-		{
-			EnsureCompletion();
-			OutResult = Result;
-		}
-
-	private:
-		FProceduralAudioTaskResults Result;
-	};
-
-	class FDecodeHandle : public FDecodeHandleBase
-	{
-	public:
-		virtual EAudioTaskType GetType() const override { return EAudioTaskType::Decode; }
-
-		void Decode(const FDecodeAudioTaskData& InJobData)
-		{
-
-			TFunction<void()> Func = [this, InJobData]()
-			{
-				// skip the first buffer if we've already decoded them
-				if (InJobData.bSkipFirstBuffer)
-				{
-#if PLATFORM_ANDROID
-					// Only skip one buffer on Android
-					InJobData.MixerBuffer->ReadCompressedData(InJobData.AudioData, InJobData.NumFramesToDecode, InJobData.bLoopingMode);
-#else
-					// If we're using cached data we need to skip the first two reads from the data
-					InJobData.MixerBuffer->ReadCompressedData(InJobData.AudioData, InJobData.NumFramesToDecode, InJobData.bLoopingMode);
-					InJobData.MixerBuffer->ReadCompressedData(InJobData.AudioData, InJobData.NumFramesToDecode, InJobData.bLoopingMode);
-#endif
-				}
-
-				Result.bLooped = InJobData.MixerBuffer->ReadCompressedData(InJobData.AudioData, InJobData.NumFramesToDecode, InJobData.bLoopingMode);
-				bFinished = true;
-			};
-
-			Future = Async(EAsyncExecution::TaskGraph, Func);
-		}
-
-		virtual void GetResult(FDecodeAudioTaskResults& OutResult) override
-		{
-			EnsureCompletion();
-			OutResult = Result;
-		}
-
-	private:
-		FDecodeAudioTaskResults Result;
-	};
-
-	IAudioTask* CreateAudioTask(const FProceduralAudioTaskData& InJobData)
-	{
-		FProceduralDecodeHandle* NewTask = new FProceduralDecodeHandle();
-		NewTask->Decode(InJobData);
-		return NewTask;
-	}
-
-	IAudioTask* CreateAudioTask(const FHeaderParseAudioTaskData& InJobData)
-	{
-		FHeaderDecodeHandle* NewTask = new FHeaderDecodeHandle();
-		NewTask->Decode(InJobData);
-		return NewTask;
-	}
-
-	IAudioTask* CreateAudioTask(const FDecodeAudioTaskData& InJobData)
-	{
-		FDecodeHandle* NewTask = new FDecodeHandle();
-		NewTask->Decode(InJobData);
-		return NewTask;
-	}
-
-#else // #if AUDIO_DECODE_USE_FUTURES
-
 class FAsyncDecodeWorker : public FNonAbandonableTask
 {
 public:
@@ -174,13 +36,46 @@ public:
 	{
 	}
 
+	~FAsyncDecodeWorker()
+	{
+	}
+
 	void DoWork()
 	{
 		switch (TaskType)
 		{
 			case EAudioTaskType::Procedural:
 			{
-				ProceduralResult.NumBytesWritten = ProceduralTaskData.ProceduralSoundWave->GeneratePCMData(ProceduralTaskData.AudioData, ProceduralTaskData.MaxAudioDataSamples);
+				// If we're not a float format, we need to convert the format to float
+				const EAudioMixerStreamDataFormat::Type FormatType = ProceduralTaskData.ProceduralSoundWave->GetGeneratedPCMDataFormat();
+				if (FormatType != EAudioMixerStreamDataFormat::Float)
+				{
+					check(FormatType == EAudioMixerStreamDataFormat::Int16);
+
+					int32 NumChannels = ProceduralTaskData.NumChannels;
+					int32 ByteSize = NumChannels * ProceduralTaskData.NumSamples * sizeof(int16);
+
+					TArray<uint8> DecodeBuffer;
+					DecodeBuffer.AddUninitialized(ByteSize);
+
+					const int32 NumBytesWritten = ProceduralTaskData.ProceduralSoundWave->GeneratePCMData(DecodeBuffer.GetData(), ProceduralTaskData.NumSamples);
+
+					check(NumBytesWritten <= ByteSize);
+
+					ProceduralResult.NumSamplesWritten = NumBytesWritten / sizeof(int16);
+
+					// Convert the buffer to float
+					int16* DecodedBufferPtr = (int16*)DecodeBuffer.GetData();
+					for (int32 SampleIndex = 0; SampleIndex < ProceduralResult.NumSamplesWritten; ++SampleIndex)
+					{
+						ProceduralTaskData.AudioData[SampleIndex] = (float)(DecodedBufferPtr[SampleIndex]) / 32768.0f;
+					}
+				}
+				else
+				{
+					const int32 NumBytesWritten = ProceduralTaskData.ProceduralSoundWave->GeneratePCMData((uint8*)ProceduralTaskData.AudioData, ProceduralTaskData.NumSamples);
+					ProceduralResult.NumSamplesWritten = NumBytesWritten / sizeof(float);
+				}
 			}
 			break;
 
@@ -192,20 +87,38 @@ public:
 
 			case EAudioTaskType::Decode:
 			{
+				int32 NumChannels = DecodeTaskData.MixerBuffer->GetNumChannels();
+				int32 ByteSize = NumChannels * DecodeTaskData.NumFramesToDecode * sizeof(int16);
+
+				// Create a buffer to decode into that's of the appropriate size
+				TArray<uint8> DecodeBuffer;
+				DecodeBuffer.AddUninitialized(ByteSize);
+
 				// skip the first buffer if we've already decoded them
 				if (DecodeTaskData.bSkipFirstBuffer)
 				{
 #if PLATFORM_ANDROID
 					// Only skip one buffer on Android
-					DecodeTaskData.MixerBuffer->ReadCompressedData(DecodeTaskData.AudioData, DecodeTaskData.NumFramesToDecode, DecodeTaskData.bLoopingMode);
+					DecodeTaskData.MixerBuffer->ReadCompressedData(DecodeBuffer.GetData(), DecodeTaskData.NumFramesToDecode, DecodeTaskData.bLoopingMode);
 #else // #if PLATFORM_ANDROID
 					// If we're using cached data we need to skip the first two reads from the data
-					DecodeTaskData.MixerBuffer->ReadCompressedData(DecodeTaskData.AudioData, DecodeTaskData.NumFramesToDecode, DecodeTaskData.bLoopingMode);
-					DecodeTaskData.MixerBuffer->ReadCompressedData(DecodeTaskData.AudioData, DecodeTaskData.NumFramesToDecode, DecodeTaskData.bLoopingMode);
+					DecodeTaskData.MixerBuffer->ReadCompressedData(DecodeBuffer.GetData(), DecodeTaskData.NumFramesToDecode, DecodeTaskData.bLoopingMode);
+					DecodeTaskData.MixerBuffer->ReadCompressedData(DecodeBuffer.GetData(), DecodeTaskData.NumFramesToDecode, DecodeTaskData.bLoopingMode);
 #endif // #else // #if PLATFORM_ANDROID
 				}
 
-				DecodeResult.bLooped = DecodeTaskData.MixerBuffer->ReadCompressedData(DecodeTaskData.AudioData, DecodeTaskData.NumFramesToDecode, DecodeTaskData.bLoopingMode);
+				DecodeResult.bLooped = DecodeTaskData.MixerBuffer->ReadCompressedData(DecodeBuffer.GetData(), DecodeTaskData.NumFramesToDecode, DecodeTaskData.bLoopingMode);
+
+				// Convert the decoded PCM data into a float buffer while still in the async task
+				int32 SampleIndex = 0;
+				int16* DecodedBufferPtr = (int16*)DecodeBuffer.GetData();
+				for (int32 Frame = 0; Frame < DecodeTaskData.NumFramesToDecode; ++Frame)
+				{
+					for (int32 Channel = 0; Channel < NumChannels; ++Channel, ++SampleIndex)
+					{
+						DecodeTaskData.AudioData[SampleIndex] = (float)(DecodedBufferPtr[SampleIndex]) / 32768.0f;
+					}
+				}
 			}
 			break;
 		}
@@ -338,6 +251,4 @@ IAudioTask* CreateAudioTask(const FDecodeAudioTaskData& InJobData)
 	return new FDecodeHandle(InJobData);
 }
 
-#endif // #else // #if AUDIO_DECODE_USE_FUTURES
-
 }
diff --git a/Engine/Source/Runtime/AudioMixer/Private/AudioMixerSourceDecode.h b/Engine/Source/Runtime/AudioMixer/Private/AudioMixerSourceDecode.h
--- a/Engine/Source/Runtime/AudioMixer/Private/AudioMixerSourceDecode.h
+++ b/Engine/Source/Runtime/AudioMixer/Private/AudioMixerSourceDecode.h
@@ -14,13 +14,16 @@ namespace Audio
 	struct FProceduralAudioTaskData
 	{
 		// The procedural sound wave ptr to use to generate audio with
-		USoundWaveProcedural* ProceduralSoundWave;
+		USoundWave* ProceduralSoundWave;
 
 		// The audio buffer to fill from the results of the generation
-		uint8* AudioData;
+		float* AudioData;
 
 		// The size of the audio buffer
-		int32 MaxAudioDataSamples;
+		int32 NumSamples;
+
+		// The number of channels of the procedural buffer
+		int32 NumChannels;
 	};
 
 	// Data needed for a decode audio task
@@ -29,8 +32,8 @@ namespace Audio
 		// The mixer buffer buffer object which holds state to use for the decode operation
 		FMixerBuffer* MixerBuffer;
 
-		// A pointer to a buffer of audio which will be written to
-		uint8* AudioData;
+		// A pointer to a buffer of audio which will be decoded to
+		float* AudioData;
 
 		// The number of frames to decode
 		int32 NumFramesToDecode;
@@ -55,16 +58,17 @@ namespace Audio
 	// Results from procedural audio task
 	struct FProceduralAudioTaskResults
 	{
-		int32 NumBytesWritten;
+		int32 NumSamplesWritten;
 
 		FProceduralAudioTaskResults()
-			: NumBytesWritten(0)
+			: NumSamplesWritten(0)
 		{}
 	};
 
 	// Results from decode audio task
 	struct FDecodeAudioTaskResults
 	{
+		// Whether or not the audio buffer looped
 		bool bLooped;
 
 		FDecodeAudioTaskResults()
