Copying //UE4/Dev-AnimPhys to //UE4/Dev-Main (Source: //UE4/Dev-AnimPhys @ 3624379)

diff --git a/Engine/Source/Runtime/AudioMixer/Private/AudioMixerSourceManager.cpp b/Engine/Source/Runtime/AudioMixer/Private/AudioMixerSourceManager.cpp
--- a/Engine/Source/Runtime/AudioMixer/Private/AudioMixerSourceManager.cpp
+++ b/Engine/Source/Runtime/AudioMixer/Private/AudioMixerSourceManager.cpp
@@ -6,6 +6,7 @@
 #include "AudioMixerSourceVoice.h"
 #include "AudioMixerSubmix.h"
 #include "IAudioExtensionPlugin.h"
+#include "AudioMixer.h"
 
 DEFINE_STAT(STAT_AudioMixerHRTF);
 
@@ -23,6 +24,23 @@ FAutoConsoleVariableRef CVarDisableParallelSourceProcessing(
 
 #define VALIDATE_SOURCE_MIXER_STATE 1
 
+#if AUDIO_MIXER_ENABLE_DEBUG_MODE
+
+// Macro which checks if the source id is in debug mode, avoids having a bunch of #ifdefs in code
+#define AUDIO_MIXER_DEBUG_LOG(SourceId, Format, ...)																							\
+	if (SourceInfos[SourceId].bIsDebugMode)																													\
+	{																																			\
+		FString CustomMessage = FString::Printf(Format, ##__VA_ARGS__);																			\
+		FString LogMessage = FString::Printf(TEXT("<Debug Sound Log> [Id=%d][Name=%s]: %s"), SourceId, *SourceInfos[SourceId].DebugName, *CustomMessage);	\
+		UE_LOG(LogAudioMixer, Log, TEXT("%s"), *LogMessage);																								\
+	}
+
+#else
+
+#define AUDIO_MIXER_DEBUG_LOG(SourceId, Message)
+
+#endif
+
 namespace Audio
 {
 	/*************************************************************************
@@ -37,6 +55,7 @@ namespace Audio
 		, NumOutputSamples(0)
 		, NumSourceWorkers(4)
 		, bInitialized(false)
+		, bUsingSpatializationPlugin(false)
 	{
 	}
 
@@ -73,6 +92,7 @@ namespace Audio
 		MixerSources.Init(nullptr, NumTotalSources);
 
 		SourceInfos.AddDefaulted(NumTotalSources);
+
 		for (int32 i = 0; i < NumTotalSources; ++i)
 		{
 			FSourceInfo& SourceInfo = SourceInfos[i];
@@ -83,6 +103,10 @@ namespace Audio
 			SourceInfo.CurrentFrameAlpha = 0.0f;
 			SourceInfo.CurrentFrameIndex = 0;
 			SourceInfo.NumFramesPlayed = 0;
+			SourceInfo.SubmixSends.Reset();
+			SourceInfo.BusId = INDEX_NONE;
+			SourceInfo.BusDurationFrames = INDEX_NONE;
+			SourceInfo.BusSends.Reset();
 
 			SourceInfo.SourceEffectChainId = INDEX_NONE;
 
@@ -97,12 +121,14 @@ namespace Audio
 			SourceInfo.bIsPlaying = false;
 			SourceInfo.bIsPaused = false;
 			SourceInfo.bIsDone = false;
+			SourceInfo.bIsLastBuffer = false;
 			SourceInfo.bIsBusy = false;
 			SourceInfo.bUseHRTFSpatializer = false;
 			SourceInfo.bUseOcclusionPlugin = false;
 			SourceInfo.bUseReverbPlugin = false;
 			SourceInfo.bHasStarted = false;
 			SourceInfo.bIsDebugMode = false;
+			SourceInfo.bOutputToBusOnly = false;
 
 			SourceInfo.NumInputChannels = 0;
 			SourceInfo.NumPostEffectChannels = 0;
@@ -126,6 +152,7 @@ namespace Audio
 			FSourceInfo& SourceInfo = SourceInfos[SourceId];
 
 			SourceInfo.SourceBuffer.Reset(NumOutputFrames * 8);
+			SourceInfo.PreDistanceAttenuationBuffer.Reset(NumOutputFrames * 8);
 			SourceInfo.AudioPluginOutputData.AudioBuffer.Reset(NumOutputFrames * 2);
 		}
 
@@ -145,6 +172,13 @@ namespace Audio
 		}
 		NumSourceWorkers = SourceWorkers.Num();
 
+		// Cache the spatialization plugin
+		SpatializationPlugin = MixerDevice->SpatializationPluginInterface;
+		if (SpatializationPlugin.IsValid())
+		{
+			bUsingSpatializationPlugin = true;
+		}
+
 		bInitialized = true;
 		bPumpQueue = false;
 	}
@@ -178,7 +212,7 @@ namespace Audio
 		AUDIO_MIXER_CHECK(MixerSources[SourceId] != nullptr);
 
 		AUDIO_MIXER_DEBUG_LOG(SourceId, TEXT("Is releasing"));
-
+		
 		FSourceInfo& SourceInfo = SourceInfos[SourceId];
 
 #if AUDIO_MIXER_ENABLE_DEBUG_MODE
@@ -187,6 +221,36 @@ namespace Audio
 			DebugSoloSources.Remove(SourceId);
 		}
 #endif
+		// Remove from list of active bus or source ids depending on what type of source this is
+		if (SourceInfo.BusId != INDEX_NONE)
+		{
+			// Remove this bus from the registry of bus instances
+			FMixerBus* Bus = Buses.Find(SourceInfo.BusId);
+			AUDIO_MIXER_CHECK(Bus);
+
+			// Remove this source from the list of bus instances.
+			if (Bus->RemoveInstanceId(SourceId))
+			{
+				Buses.Remove(SourceInfo.BusId);
+			}
+		}
+		else
+		{
+			ActiveSourceIds.Remove(SourceId);
+		}
+
+		// Remove this source's send list from the bus data registry
+		for (uint32 BusId : SourceInfo.BusSends)
+		{
+			// we should have a bus registration entry still since the send hasn't been cleaned up yet
+			FMixerBus* Bus = Buses.Find(BusId);
+			AUDIO_MIXER_CHECK(Bus);
+
+			if (Bus->RemoveBusSend(SourceId))
+			{
+				Buses.Remove(BusId);
+			}
+		}
 
 		// Call OnRelease on the BufferQueueListener to give it a chance 
 		// to release any resources it owns on the audio render thread
@@ -197,16 +261,17 @@ namespace Audio
 		}
 
 		// Remove the mixer source from its submix sends
-		TMap<uint32, FMixerSourceSubmixSend>& SubmixSends = MixerSources[SourceId]->GetSubmixSends();
-		for (auto SubmixSendItem : SubmixSends)
+		for (FMixerSourceSubmixSend& SubmixSendItem : SourceInfo.SubmixSends)
 		{
-			SubmixSendItem.Value.Submix->RemoveSourceVoice(MixerSources[SourceId]);
+			SubmixSendItem.Submix->RemoveSourceVoice(MixerSources[SourceId]);
 		}
+		SourceInfo.SubmixSends.Reset();
 
 		// Notify plugin effects
 		if (SourceInfo.bUseHRTFSpatializer)
 		{
-			MixerDevice->SpatializationPluginInterface->OnReleaseSource(SourceId);
+			AUDIO_MIXER_CHECK(bUsingSpatializationPlugin);
+			SpatializationPlugin->OnReleaseSource(SourceId);
 		}
 
 		if (SourceInfo.bUseOcclusionPlugin)
@@ -223,16 +288,21 @@ namespace Audio
 		SourceInfo.SourceEffectChainId = INDEX_NONE;
 		ResetSourceEffectChain(SourceId);
 
+		SourceInfo.BusId = INDEX_NONE;
+		SourceInfo.BusDurationFrames = INDEX_NONE;
+		SourceInfo.BusSends.Reset();
+
 		SourceInfo.SourceEnvelopeFollower.Reset();
 		SourceInfo.bEffectTailsDone = true;
 
-		// Delete the mixer source and null the slot
-		delete MixerSources[SourceId];
+		// Release the source voice back to the mixer device. This is pooled.
+		MixerDevice->ReleaseMixerSourceVoice(MixerSources[SourceId]);
 		MixerSources[SourceId] = nullptr;
 
 		// Reset all state and data
 		SourceInfo.PitchSourceParam.Reset();
 		SourceInfo.VolumeSourceParam.Reset();
+		SourceInfo.DistanceAttenuationSourceParam.Reset();
 		SourceInfo.LPFCutoffFrequencyParam.Reset();
 
 		SourceInfo.LowPassFilter.Reset();
@@ -242,6 +312,7 @@ namespace Audio
 		SourceInfo.CurrentPCMBuffer = nullptr;
 		SourceInfo.CurrentAudioChunkNumFrames = 0;
 		SourceInfo.SourceBuffer.Reset();
+		SourceInfo.PreDistanceAttenuationBuffer.Reset();
 		SourceInfo.AudioPluginOutputData.AudioBuffer.Reset();
 		SourceInfo.CurrentFrameValues.Reset();
 		SourceInfo.NextFrameValues.Reset();
@@ -255,6 +326,7 @@ namespace Audio
 		SourceInfo.bIsActive = false;
 		SourceInfo.bIsPlaying = false;
 		SourceInfo.bIsDone = true;
+		SourceInfo.bIsLastBuffer = false;
 		SourceInfo.bIsPaused = false;
 		SourceInfo.bIsBusy = false;
 		SourceInfo.bUseHRTFSpatializer = false;
@@ -262,6 +334,7 @@ namespace Audio
 		SourceInfo.bUseReverbPlugin = false;
 		SourceInfo.bHasStarted = false;
 		SourceInfo.bIsDebugMode = false;
+		SourceInfo.bOutputToBusOnly = false;
 		SourceInfo.DebugName = FString();
 		SourceInfo.NumInputChannels = 0;
 		SourceInfo.NumPostEffectChannels = 0;
@@ -346,6 +419,11 @@ namespace Audio
 		return NumActiveSources;
 	}
 
+	int32 FMixerSourceManager::GetNumActiveBuses() const
+	{
+		return Buses.Num();
+	}
+
 	void FMixerSourceManager::InitSource(const int32 SourceId, const FMixerSourceVoiceInitParams& InitParams)
 	{
 		AUDIO_MIXER_CHECK(SourceId < NumTotalSources);
@@ -370,6 +448,7 @@ namespace Audio
 			SourceInfo.bIsActive = true;
 			SourceInfo.bIsBusy = true;
 			SourceInfo.bIsDone = false;
+			SourceInfo.bIsLastBuffer = false;
 			SourceInfo.bUseHRTFSpatializer = InitParams.bUseHRTFSpatialization;
 
 			SourceInfo.BufferQueueListener = InitParams.BufferQueueListener;
@@ -386,8 +465,8 @@ namespace Audio
 			// Create the spatialization plugin source effect
 			if (InitParams.bUseHRTFSpatialization)
 			{
-				MixerDevice->SpatializationPluginInterface->OnInitSource(SourceId, InitParams.AudioComponentUserID, 
-					InitParams.SpatializationPluginSettings);
+				AUDIO_MIXER_CHECK(bUsingSpatializationPlugin);
+				SpatializationPlugin->OnInitSource(SourceId, InitParams.AudioComponentUserID, InitParams.SpatializationPluginSettings);
 			}
 
 			// Create the occlusion plugin source effect
@@ -431,6 +510,72 @@ namespace Audio
 
 				SourceInfo.SourceEffectChainId = InitParams.SourceEffectChainId;
 				BuildSourceEffectChain(SourceId, InitData, InitParams.SourceEffectChain);
+
+				// Whether or not to output to bus only
+				SourceInfo.bOutputToBusOnly = InitParams.bOutputToBusOnly;
+
+				// If this is a bus, add this source id to the list of active bus ids
+				if (InitParams.BusId != INDEX_NONE)
+				{
+					// Setting this BusId will flag this source as a bus. It doesn't try to generate 
+					// audio in the normal way but instead will render in a second stage, after normal source rendering.
+					SourceInfo.BusId = InitParams.BusId;
+
+					// Bus duration allows us to stop a bus after a given time
+					if (InitParams.BusDuration != 0.0f)
+					{
+						SourceInfo.BusDurationFrames = InitParams.BusDuration * MixerDevice->GetSampleRate();
+					}
+
+					// Register this bus as an instance
+					FMixerBus* Bus = Buses.Find(InitParams.BusId);
+					if (Bus)
+					{
+						// If this bus is already registered, add this as a source id
+						Bus->AddInstanceId(SourceId);
+					}
+					else
+					{
+						// If the bus is not registered, make a new entry
+						FMixerBus NewBusData(this, InitParams.NumInputChannels, NumOutputFrames);
+
+						NewBusData.AddInstanceId(SourceId);
+
+						Buses.Add(InitParams.BusId, NewBusData);
+					}
+				}
+
+				// Iterate through source's bus sends and add this source to the bus send list
+				// Note: buses can also send their audio to other buses.
+				for (const FMixerBusSend& BusSend : InitParams.BusSends)
+				{
+					// New struct to map which source (SourceId) is sending to the bus
+					FBusSend NewBusSend;
+					NewBusSend.SourceId = SourceId;
+					NewBusSend.SendLevel = BusSend.SendLevel;
+
+					// Get existing BusId and add the send, or create new bus registration
+					FMixerBus* Bus = Buses.Find(BusSend.BusId);
+					if (Bus)
+					{
+						Bus->AddBusSend(NewBusSend);
+					}
+					else
+					{
+						// If the bus is not registered, make a new entry
+						FMixerBus NewBusData(this, InitParams.NumInputChannels, NumOutputFrames);
+
+						// Add a send to it. This will not have a bus instance id (i.e. won't output audio), but 
+						// we register the send anyway in the event that this bus does play, we'll know to send this
+						// source's audio to it.
+						NewBusData.AddBusSend(NewBusSend);
+
+						Buses.Add(BusSend.BusId, NewBusData);
+					}
+
+					// Store on this source, which buses its sending its audio to
+					SourceInfo.BusSends.Add(BusSend.BusId);
+				}
 			}
 
 			SourceInfo.CurrentFrameValues.Init(0.0f, InitParams.NumInputChannels);
@@ -440,10 +585,14 @@ namespace Audio
 			MixerSources[SourceId] = InitParams.SourceVoice;
 
 			// Loop through the source's sends and add this source to those submixes with the send info
+
+			AUDIO_MIXER_CHECK(SourceInfo.SubmixSends.Num() == 0);
+
 			for (int32 i = 0; i < InitParams.SubmixSends.Num(); ++i)
 			{
-				const FMixerSourceSubmixSend& MixerSourceSend = InitParams.SubmixSends[i];
-				MixerSourceSend.Submix->AddOrSetSourceVoice(InitParams.SourceVoice, MixerSourceSend.SendLevel);
+				const FMixerSourceSubmixSend& MixerSubmixSend = InitParams.SubmixSends[i];
+				SourceInfo.SubmixSends.Add(MixerSubmixSend);
+				MixerSubmixSend.Submix->AddOrSetSourceVoice(InitParams.SourceVoice, MixerSubmixSend.SendLevel);
 			}
 
 #if AUDIO_MIXER_ENABLE_DEBUG_MODE
@@ -570,6 +719,21 @@ namespace Audio
 		});
 	}
 
+	void FMixerSourceManager::SetDistanceAttenuation(const int32 SourceId, const float DistanceAttenuation)
+	{
+		AUDIO_MIXER_CHECK(SourceId < NumTotalSources);
+		AUDIO_MIXER_CHECK(GameThreadInfo.bIsBusy[SourceId]);
+		AUDIO_MIXER_CHECK_GAME_THREAD(MixerDevice);
+
+		AudioMixerThreadCommand([this, SourceId, DistanceAttenuation]()
+		{
+			AUDIO_MIXER_CHECK_AUDIO_PLAT_THREAD(MixerDevice);
+			check(NumOutputFrames > 0);
+
+			SourceInfos[SourceId].DistanceAttenuationSourceParam.SetValue(DistanceAttenuation, NumOutputFrames);
+		});
+	}
+
 	void FMixerSourceManager::SetSpatializationParams(const int32 SourceId, const FSpatializationParams& InParams)
 	{
 		AUDIO_MIXER_CHECK(SourceId < NumTotalSources);
@@ -644,9 +808,6 @@ namespace Audio
 		{
 			AUDIO_MIXER_CHECK_AUDIO_PLAT_THREAD(MixerDevice);
 
-			float SampleRate = MixerDevice->GetSampleRate();
-			AUDIO_MIXER_CHECK(SampleRate > 0.0f);
-
 			SourceInfos[SourceId].LPFCutoffFrequencyParam.SetValue(InLPFFrequency, NumOutputFrames);
 		});
 	}
@@ -661,9 +822,6 @@ namespace Audio
 		{
 			AUDIO_MIXER_CHECK_AUDIO_PLAT_THREAD(MixerDevice);
 
-			float SampleRate = MixerDevice->GetSampleRate();
-			AUDIO_MIXER_CHECK(SampleRate > 0.0f);
-
 			SourceInfos[SourceId].HPFCutoffFrequencyParam.SetValue(InHPFFrequency, NumOutputFrames);
 		});
 	}
@@ -671,7 +829,7 @@ namespace Audio
 	void FMixerSourceManager::SubmitBuffer(const int32 SourceId, FMixerSourceBufferPtr InSourceVoiceBuffer, const bool bSubmitSynchronously)
 	{
 		AUDIO_MIXER_CHECK(SourceId < NumTotalSources);
-		AUDIO_MIXER_CHECK(InSourceVoiceBuffer->AudioBytes <= (uint32)InSourceVoiceBuffer->AudioData.Num());
+		AUDIO_MIXER_CHECK(InSourceVoiceBuffer->Samples <= (uint32)InSourceVoiceBuffer->AudioData.Num());
 
 		// If we're submitting synchronously then don't use AudioMixerThreadCommand but submit directly
 		if (bSubmitSynchronously)
@@ -680,7 +838,6 @@ namespace Audio
 		}
 		else
 		{
-			// make sure we're on the game/audio thread
 			AUDIO_MIXER_CHECK_GAME_THREAD(MixerDevice);
 
 			// Use a thread command
@@ -693,17 +850,33 @@ namespace Audio
 		}
 	}
 
-
-	void FMixerSourceManager::SetSubmixSendInfo(const int32 SourceId, FMixerSubmixPtr Submix, const float SendLevel)
+	void FMixerSourceManager::SetSubmixSendInfo(const int32 SourceId, const FMixerSourceSubmixSend& InSubmixSend)
 	{
 		AUDIO_MIXER_CHECK(SourceId < NumTotalSources);
 		AUDIO_MIXER_CHECK(GameThreadInfo.bIsBusy[SourceId]);
 		AUDIO_MIXER_CHECK_GAME_THREAD(MixerDevice);
 
-		AudioMixerThreadCommand([this, SourceId, Submix, SendLevel]()
+		AudioMixerThreadCommand([this, SourceId, InSubmixSend]()
 		{
-			AUDIO_MIXER_CHECK_AUDIO_PLAT_THREAD(MixerDevice);
-			Submix->AddOrSetSourceVoice(MixerSources[SourceId], SendLevel);
+			FSourceInfo& SourceInfo = SourceInfos[SourceId];
+
+			bool bIsNew = true;
+			for (FMixerSourceSubmixSend& SubmixSend : SourceInfo.SubmixSends)
+			{
+				if (SubmixSend.Submix->GetId() == InSubmixSend.Submix->GetId())
+				{
+					SubmixSend.SendLevel = InSubmixSend.SendLevel;
+					bIsNew = false;
+					break;
+				}
+			}
+
+			if (bIsNew)
+			{
+				SourceInfos[SourceId].SubmixSends.Add(InSubmixSend);
+			}
+
+			InSubmixSend.Submix->AddOrSetSourceVoice(MixerSources[SourceId], InSubmixSend.SendLevel);
 		});
 	}
 
@@ -753,12 +926,12 @@ namespace Audio
 				bReadCurrentFrame = false;
 
 				AUDIO_MIXER_CHECK(SourceInfo.CurrentPCMBuffer.IsValid());
-				const int16* AudioData = (const int16*)SourceInfo.CurrentPCMBuffer->AudioData.GetData();
+				const float* AudioData = SourceInfo.CurrentPCMBuffer->AudioData.GetData();
 				const int32 CurrentSampleIndex = SourceInfo.CurrentFrameIndex * NumChannels;
 
 				for (int32 Channel = 0; Channel < NumChannels; ++Channel)
 				{
-					SourceInfo.CurrentFrameValues[Channel] = (float)AudioData[CurrentSampleIndex + Channel] * ONEOVERSHORTMAX;
+					SourceInfo.CurrentFrameValues[Channel] = AudioData[CurrentSampleIndex + Channel];
 				}
 			}
 
@@ -787,7 +960,7 @@ namespace Audio
 
 				MixerSources[SourceId]->NumBuffersQueued.Decrement();
 
-				SourceInfo.CurrentAudioChunkNumFrames = SourceInfo.CurrentPCMBuffer->AudioBytes / (NUM_BYTES_PER_SAMPLE * NumChannels);
+				SourceInfo.CurrentAudioChunkNumFrames = SourceInfo.CurrentPCMBuffer->Samples / NumChannels;
 
 				// Subtract the number of frames in the current buffer from our frame index.
 				// Note: if this is the first time we're playing, CurrentFrameIndex will be 0
@@ -804,7 +977,7 @@ namespace Audio
 			}
 			else
 			{
-				SourceInfo.bIsDone = true;
+				SourceInfo.bIsLastBuffer = true;
 				return;
 			}
 
@@ -814,8 +987,8 @@ namespace Audio
 
 		if (SourceInfo.CurrentPCMBuffer.IsValid())
 		{
-			// Grab the 16-bit PCM audio data (which could be a new audio chunk from previous ReadSourceFrame call)
-			const int16* AudioData = (const int16*)SourceInfo.CurrentPCMBuffer->AudioData.GetData();
+			// Grab the float PCM audio data (which could be a new audio chunk from previous ReadSourceFrame call)
+			const float* AudioData = SourceInfo.CurrentPCMBuffer->AudioData.GetData();
 			const int32 NextSampleIndex = (SourceInfo.CurrentFrameIndex + 1)  * NumChannels;
 
 			if (bReadCurrentFrame)
@@ -823,111 +996,333 @@ namespace Audio
 				const int32 CurrentSampleIndex = SourceInfo.CurrentFrameIndex * NumChannels;
 				for (int32 Channel = 0; Channel < NumChannels; ++Channel)
 				{
-					SourceInfo.CurrentFrameValues[Channel] = (float)AudioData[CurrentSampleIndex + Channel] * ONEOVERSHORTMAX;
-					SourceInfo.NextFrameValues[Channel] = (float)AudioData[NextSampleIndex + Channel] * ONEOVERSHORTMAX;
+					SourceInfo.CurrentFrameValues[Channel] = AudioData[CurrentSampleIndex + Channel];
+					SourceInfo.NextFrameValues[Channel] = AudioData[NextSampleIndex + Channel];
 				}
 			}
 			else
 			{
 				for (int32 Channel = 0; Channel < NumChannels; ++Channel)
 				{
-					SourceInfo.NextFrameValues[Channel] = (float)AudioData[NextSampleIndex + Channel] * ONEOVERSHORTMAX;
+					SourceInfo.NextFrameValues[Channel] = AudioData[NextSampleIndex + Channel];
 				}
 			}
 		}
 	}
 
-	void FMixerSourceManager::ComputeSourceBuffersForIdRange(const int32 SourceIdStart, const int32 SourceIdEnd)
+	void FMixerSourceManager::ComputeSourceBuffersForIdRange(const bool bGenerateBuses, const int32 SourceIdStart, const int32 SourceIdEnd)
 	{
 		SCOPE_CYCLE_COUNTER(STAT_AudioMixerSourceBuffers);
 
-		// Local variable used for sample values
-		float SampleValue = 0.0f;
-
 		for (int32 SourceId = SourceIdStart; SourceId < SourceIdEnd; ++SourceId)
 		{
 			FSourceInfo& SourceInfo = SourceInfos[SourceId];
 
-			SourceInfo.SourceBuffer.Reset();
-
 			if (!SourceInfo.bIsBusy || !SourceInfo.bIsPlaying || SourceInfo.bIsPaused)
 			{
 				continue;
 			}
 
+			// If this source is still playing at this point but technically done, zero the buffers. We haven't yet been removed by the FMixerSource owner.
+			// This should be rare but could happen due to thread timing since done-ness is queried on audio thread.
+			if (SourceInfo.bIsDone)
+			{
+				const int32 NumSamples = NumOutputFrames * SourceInfo.NumInputChannels;
+
+				SourceInfo.PreDistanceAttenuationBuffer.Reset();
+				SourceInfo.PreDistanceAttenuationBuffer.AddZeroed(NumSamples);
+
+				SourceInfo.SourceBuffer.Reset();
+				SourceInfo.SourceBuffer.AddZeroed(NumSamples);
+
+				continue;
+			}
+
+			const bool bIsBus = SourceInfo.BusId != INDEX_NONE;
+			if ((bGenerateBuses && !bIsBus) || (!bGenerateBuses && bIsBus))
+			{
+				continue;
+			}
+
 			// Fill array with elements all at once to avoid sequential Add() operation overhead.
-			const int32 NumSourceChannels = SourceInfo.NumInputChannels;
-			SourceInfo.SourceBuffer.AddUninitialized(NumOutputFrames * NumSourceChannels);
-			int32 SampleIndex = 0;
+			const int32 NumSamples = NumOutputFrames * SourceInfo.NumInputChannels;
+			
+			// Initialize both the pre-distance attenuation buffer and the source buffer
+			SourceInfo.PreDistanceAttenuationBuffer.Reset();
+			SourceInfo.PreDistanceAttenuationBuffer.AddZeroed(NumSamples);
 
-			for (int32 Frame = 0; Frame < NumOutputFrames; ++Frame)
+			SourceInfo.SourceBuffer.Reset();
+			SourceInfo.SourceBuffer.AddZeroed(NumSamples);
+
+			float* PreDistanceAttenBufferPtr = SourceInfo.PreDistanceAttenuationBuffer.GetData();
+
+			// if this is a bus, we just want to copy the bus audio to this source's output audio
+			// Note we need to copy this since bus instances may have different audio via dynamic source effects, etc.
+			if (bIsBus)
 			{
-				// If the source is done, then we'll just write out 0s
-				if (!SourceInfo.bIsDone)
+				// Get the source's rendered bus data
+				const FMixerBus* Bus = Buses.Find(SourceInfo.BusId);
+				const float* BusBufferPtr = Bus->GetCurrentBusBuffer();
+
+				int32 NumFramesPlayed = NumOutputFrames;
+				if (SourceInfo.BusDurationFrames != INDEX_NONE)
 				{
-					// Whether or not we need to read another sample from the source buffers
-					// If we haven't yet played any frames, then we will need to read the first source samples no matter what
-					bool bReadNextSample = !SourceInfo.bHasStarted;
+					// If we're now finishing, only copy over the real data
+					if ((SourceInfo.NumFramesPlayed + NumOutputFrames) >= SourceInfo.BusDurationFrames)
+					{
+						NumFramesPlayed = SourceInfo.BusDurationFrames - SourceInfo.NumFramesPlayed;
+						SourceInfo.bIsLastBuffer = true;				
+					}
+				}
 
-					// Reset that we've started generating audio
-					SourceInfo.bHasStarted = true;
+				SourceInfo.NumFramesPlayed += NumFramesPlayed;
+
+				// Simply copy into the pre distance attenuation buffer ptr
+				FMemory::Memcpy(PreDistanceAttenBufferPtr, BusBufferPtr, sizeof(float) * NumFramesPlayed * SourceInfo.NumInputChannels);
+			}
+			// Otherwise we need to generate source audio using buffer queues and perform SRC
+			else
+			{
+				int32 SampleIndex = 0;
 
-					// Update the PrevFrameIndex value for the source based on alpha value
-					while (SourceInfo.CurrentFrameAlpha >= 1.0f)
+				for (int32 Frame = 0; Frame < NumOutputFrames; ++Frame)
+				{
+					// If the source is done, then we'll just write out 0s
+					if (!SourceInfo.bIsLastBuffer)
 					{
-						// Our inter-frame alpha lerping value is causing us to read new source frames
-						bReadNextSample = true;
+						// Whether or not we need to read another sample from the source buffers
+						// If we haven't yet played any frames, then we will need to read the first source samples no matter what
+						bool bReadNextSample = !SourceInfo.bHasStarted;
 
-						// Bump up the current frame index
-						SourceInfo.CurrentFrameIndex++;
+						// Reset that we've started generating audio
+						SourceInfo.bHasStarted = true;
 
-						// Bump up the frames played -- this is tracking the total frames in source file played
-						// CurrentFrameIndex can wrap for looping sounds so won't be accurate in that case
-						SourceInfo.NumFramesPlayed++;
+						// Update the PrevFrameIndex value for the source based on alpha value
+						while (SourceInfo.CurrentFrameAlpha >= 1.0f)
+						{
+							// Our inter-frame alpha lerping value is causing us to read new source frames
+							bReadNextSample = true;
 
-						SourceInfo.CurrentFrameAlpha -= 1.0f;
-					}
+							// Bump up the current frame index
+							SourceInfo.CurrentFrameIndex++;
 
-					// If our alpha parameter caused us to jump to a new source frame, we need
-					// read new samples into our prev and next frame sample data
-					if (bReadNextSample)
-					{
-						ReadSourceFrame(SourceId);
+							// Bump up the frames played -- this is tracking the total frames in source file played
+							// CurrentFrameIndex can wrap for looping sounds so won't be accurate in that case
+							SourceInfo.NumFramesPlayed++;
+
+							SourceInfo.CurrentFrameAlpha -= 1.0f;
+						}
+
+						// If our alpha parameter caused us to jump to a new source frame, we need
+						// read new samples into our prev and next frame sample data
+						if (bReadNextSample)
+						{
+							ReadSourceFrame(SourceId);
+						}
 					}
-				}
 
-				// If we've finished reading all buffer data, then just write out 0s
-				if (SourceInfo.bIsDone)
-				{
-					for (int32 Channel = 0; Channel < NumSourceChannels; ++Channel)
+					// If we've finished reading all buffer data, then just write out 0s
+					if (SourceInfo.bIsLastBuffer)
 					{
-						SourceInfo.SourceBuffer[SampleIndex++] = 0.0f;
+						for (int32 Channel = 0; Channel < SourceInfo.NumInputChannels; ++Channel)
+						{
+							PreDistanceAttenBufferPtr[SampleIndex++] = 0.0f;
+						}
 					}
-				}
-				else
-				{
-					// Get the volume value of the source at this frame index
-					for (int32 Channel = 0; Channel < NumSourceChannels; ++Channel)
+					else
 					{
-						const float CurrFrameValue = SourceInfo.CurrentFrameValues[Channel];
-						const float NextFrameValue = SourceInfo.NextFrameValues[Channel];
-						const float CurrentAlpha = SourceInfo.CurrentFrameAlpha;
+						// perform linear SRC to get the next sample value from the decoded buffer
+						for (int32 Channel = 0; Channel < SourceInfo.NumInputChannels; ++Channel)
+						{
+							const float CurrFrameValue = SourceInfo.CurrentFrameValues[Channel];
+							const float NextFrameValue = SourceInfo.NextFrameValues[Channel];
+							const float CurrentAlpha = SourceInfo.CurrentFrameAlpha;
+
+							PreDistanceAttenBufferPtr[SampleIndex++] = FMath::Lerp(CurrFrameValue, NextFrameValue, CurrentAlpha);
+						}
+						const float CurrentPitchScale = SourceInfo.PitchSourceParam.Update();
 
-						SampleValue = FMath::Lerp(CurrFrameValue, NextFrameValue, CurrentAlpha);
-						SourceInfo.SourceBuffer[SampleIndex++] = SampleValue;
+						SourceInfo.CurrentFrameAlpha += CurrentPitchScale;
 					}
-					const float CurrentPitchScale = SourceInfo.PitchSourceParam.Update();
+				}
+
+				// After processing the frames, reset the pitch param
+				SourceInfo.PitchSourceParam.Reset();
+			}
+		}
+	}
 
-					SourceInfo.CurrentFrameAlpha += CurrentPitchScale;
+	void FMixerSourceManager::ComputeBuses()
+	{
+		// Loop through the bus registry and mix source audio
+		for (auto& Entry : Buses)
+		{
+			FMixerBus& Bus = Entry.Value;
+			Bus.MixBuffer();
+		}
+	}
+
+	void FMixerSourceManager::UpdateBuses()
+	{
+		// Update the bus states post mixing. This flips the current/previous buffer indices.
+		for (auto& Entry : Buses)
+		{
+			FMixerBus& Bus = Entry.Value;
+			Bus.Update();
+		}
+	}
+
+	void FMixerSourceManager::ApplyDistanceAttenuation(FSourceInfo& SourceInfo, int32 NumSamples)
+	{
+		float* PreDistanceAttenBufferPtr = SourceInfo.PreDistanceAttenuationBuffer.GetData();
+		float* PostDistanceAttenBufferPtr = SourceInfo.SourceBuffer.GetData();
+
+		// Interpolate the distance attenuation value to avoid discontinuities
+		float CurrentDistanceAttenuationValue = SourceInfo.DistanceAttenuationSourceParam.GetValue();
+		for (int32 Sample = 0; Sample < NumSamples; Sample += SourceInfo.NumInputChannels)
+		{
+			for (int32 Chan = 0; Chan < SourceInfo.NumInputChannels; ++Chan)
+			{
+				PostDistanceAttenBufferPtr[Sample + Chan] = PreDistanceAttenBufferPtr[Sample + Chan] * CurrentDistanceAttenuationValue;
+			}
+
+			CurrentDistanceAttenuationValue = SourceInfo.DistanceAttenuationSourceParam.Update();
+		}
+		SourceInfo.DistanceAttenuationSourceParam.Reset();
+	}
+
+	void FMixerSourceManager::ComputePluginAudio(FSourceInfo& SourceInfo, int32 SourceId, int32 NumSamples)
+	{
+		// Apply the distance attenuation before sending to plugins
+		float* PreDistanceAttenBufferPtr = SourceInfo.PreDistanceAttenuationBuffer.GetData();
+		float* PostDistanceAttenBufferPtr = SourceInfo.SourceBuffer.GetData();
+
+		bool bShouldMixInReverb = false;
+		if (SourceInfo.bUseReverbPlugin)
+		{
+			const FSpatializationParams* SourceSpatParams = &SourceInfo.SpatParams;
+
+			// Move the audio buffer to the reverb plugin buffer
+			FAudioPluginSourceInputData AudioPluginInputData;
+			AudioPluginInputData.SourceId = SourceId;
+			AudioPluginInputData.AudioBuffer = &SourceInfo.SourceBuffer;
+			AudioPluginInputData.SpatializationParams = SourceSpatParams;
+			AudioPluginInputData.NumChannels = SourceInfo.NumInputChannels;
+			SourceInfo.AudioPluginOutputData.AudioBuffer.Reset();
+			SourceInfo.AudioPluginOutputData.AudioBuffer.AddZeroed(AudioPluginInputData.AudioBuffer->Num());
+
+			MixerDevice->ReverbPluginInterface->ProcessSourceAudio(AudioPluginInputData, SourceInfo.AudioPluginOutputData);
+
+			// Make sure the buffer counts didn't change and are still the same size
+			AUDIO_MIXER_CHECK(SourceInfo.AudioPluginOutputData.AudioBuffer.Num() == NumSamples);
+
+			//If the reverb effect doesn't send it's audio to an external device, mix the output data back in.
+			if (!MixerDevice->bReverbIsExternalSend)
+			{
+				// Copy the reverb-processed data back to the source buffer
+				SourceInfo.ReverbPluginOutputBuffer.Reset();
+				SourceInfo.ReverbPluginOutputBuffer.Append(SourceInfo.AudioPluginOutputData.AudioBuffer);
+				bShouldMixInReverb = true;
+			}
+		}
+
+		if (SourceInfo.bUseOcclusionPlugin)
+		{
+			const FSpatializationParams* SourceSpatParams = &SourceInfo.SpatParams;
+
+			// Move the audio buffer to the occlusion plugin buffer
+			FAudioPluginSourceInputData AudioPluginInputData;
+			AudioPluginInputData.SourceId = SourceId;
+			AudioPluginInputData.AudioBuffer = &SourceInfo.SourceBuffer;
+			AudioPluginInputData.SpatializationParams = SourceSpatParams;
+			AudioPluginInputData.NumChannels = SourceInfo.NumInputChannels;
+
+			SourceInfo.AudioPluginOutputData.AudioBuffer.Reset();
+			SourceInfo.AudioPluginOutputData.AudioBuffer.AddZeroed(AudioPluginInputData.AudioBuffer->Num());
+
+			MixerDevice->OcclusionInterface->ProcessAudio(AudioPluginInputData, SourceInfo.AudioPluginOutputData);
+
+			// Make sure the buffer counts didn't change and are still the same size
+			AUDIO_MIXER_CHECK(SourceInfo.AudioPluginOutputData.AudioBuffer.Num() == NumSamples);
+
+			// Copy the occlusion-processed data back to the source buffer and mix with the reverb plugin output buffer
+			if (bShouldMixInReverb)
+			{
+				const float* ReverbPluginOutputBufferPtr = SourceInfo.ReverbPluginOutputBuffer.GetData();
+				const float* AudioPluginOutputDataPtr = SourceInfo.AudioPluginOutputData.AudioBuffer.GetData();
+
+				for (int32 i = 0; i < NumSamples; ++i)
+				{
+					PostDistanceAttenBufferPtr[i] = ReverbPluginOutputBufferPtr[i] + AudioPluginOutputDataPtr[i];
 				}
 			}
+			else
+			{
+				FMemory::Memcpy(PostDistanceAttenBufferPtr, SourceInfo.AudioPluginOutputData.AudioBuffer.GetData(), sizeof(float) * NumSamples);
+			}
+		}
+		else if (bShouldMixInReverb)
+		{
+			const float* ReverbPluginOutputBufferPtr = SourceInfo.ReverbPluginOutputBuffer.GetData();
+			for (int32 i = 0; i < NumSamples; ++i)
+			{
+				PostDistanceAttenBufferPtr[i] += ReverbPluginOutputBufferPtr[i];
+			}
+		}
+
+		// If the source has HRTF processing enabled, run it through the spatializer
+		if (SourceInfo.bUseHRTFSpatializer)
+		{
+			SCOPE_CYCLE_COUNTER(STAT_AudioMixerHRTF);
+
+			AUDIO_MIXER_CHECK(SpatializationPlugin.IsValid());
+			AUDIO_MIXER_CHECK(SourceInfo.NumInputChannels == 1);
+
+			FAudioPluginSourceInputData AudioPluginInputData;
+			AudioPluginInputData.AudioBuffer = &SourceInfo.SourceBuffer;
+			AudioPluginInputData.NumChannels = SourceInfo.NumInputChannels;
+			AudioPluginInputData.SourceId = SourceId;
+			AudioPluginInputData.SpatializationParams = &SourceInfo.SpatParams;
+
+			if (!MixerDevice->bSpatializationIsExternalSend)
+			{
+				SourceInfo.AudioPluginOutputData.AudioBuffer.Reset();
+				SourceInfo.AudioPluginOutputData.AudioBuffer.AddZeroed(2 * NumOutputFrames);
+			}
+
+			SpatializationPlugin->ProcessAudio(AudioPluginInputData, SourceInfo.AudioPluginOutputData);
+
+			// If this is an external send, we treat this source audio as if it was still a mono source
+			// This will allow it to traditionally pan in the ComputeOutputBuffers function and be
+			// sent to submixes (e.g. reverb) panned and mixed down. Certain submixes will want this spatial 
+			// information in addition to the external send. We've already bypassed adding this source
+			// to a base submix (e.g. master/eq, etc)
+			if (MixerDevice->bSpatializationIsExternalSend)
+			{
+				// Otherwise our pre- and post-effect channels are the same as the input channels
+				SourceInfo.NumPostEffectChannels = SourceInfo.NumInputChannels;
+
+				// Set the ptr to use for post-effect buffers rather than the plugin output data (since the plugin won't have output audio data)
+				SourceInfo.PostEffectBuffers = &SourceInfo.SourceBuffer;
+			}
+			else
+			{
+				// Otherwise, we are now a 2-channel file and should not be spatialized using normal 3d spatialization
+				SourceInfo.NumPostEffectChannels = 2;
+				SourceInfo.PostEffectBuffers = &SourceInfo.AudioPluginOutputData.AudioBuffer;
+			}
+		}
+		else
+		{
+			// Otherwise our pre- and post-effect channels are the same as the input channels
+			SourceInfo.NumPostEffectChannels = SourceInfo.NumInputChannels;
 
-			// After processing the frames, reset the pitch param
-			SourceInfo.PitchSourceParam.Reset();
+			// Set the ptr to use for post-effect buffers
+			SourceInfo.PostEffectBuffers = &SourceInfo.SourceBuffer;
 		}
 	}
 
-	void FMixerSourceManager::ComputePostSourceEffectBufferForIdRange(const int32 SourceIdStart, const int32 SourceIdEnd)
+	void FMixerSourceManager::ComputePostSourceEffectBufferForIdRange(bool bGenerateBuses, const int32 SourceIdStart, const int32 SourceIdEnd)
 	{
 		SCOPE_CYCLE_COUNTER(STAT_AudioMixerSourceEffectBuffers);
 
@@ -937,54 +1332,49 @@ namespace Audio
 		{
 			FSourceInfo& SourceInfo = SourceInfos[SourceId];
 
-			if (!SourceInfo.bIsBusy || !SourceInfo.bIsPlaying || SourceInfo.bIsPaused)
+			if (!SourceInfo.bIsBusy || !SourceInfo.bIsPlaying || SourceInfo.bIsPaused || SourceInfo.bIsDone)
+			{
+				continue;
+			}
+
+			const bool bIsBus = SourceInfo.BusId != INDEX_NONE;
+			if ((bGenerateBuses && !bIsBus) || (!bGenerateBuses && bIsBus))
 			{
 				continue;
 			}
 
-			// Get the source buffer
-			float* BufferPtr = SourceInfo.SourceBuffer.GetData();
-			const int32 NumSamples = SourceInfo.SourceBuffer.Num();
+			float* PreDistanceAttenBufferPtr = SourceInfo.PreDistanceAttenuationBuffer.GetData();
+			const int32 NumSamples = SourceInfo.PreDistanceAttenuationBuffer.Num();
 
-			const int32 NumInputChans = SourceInfo.NumInputChannels;
 			float CurrentVolumeValue = SourceInfo.VolumeSourceParam.GetValue();
 
-			// Process the per-source LPF if the source hasn't actually been finished
-			if (!SourceInfo.bIsDone)
+			for (int32 Frame = 0; Frame < NumOutputFrames; ++Frame)
 			{
-				for (int32 Frame = 0; Frame < NumOutputFrames; ++Frame)
-				{
-					FSourceParam& LPFFrequencyParam = SourceInfo.LPFCutoffFrequencyParam;
-					FSourceParam& HPFFrequencyParam = SourceInfo.HPFCutoffFrequencyParam;
-
-					// Update the frequency at the block rate
-					const float LPFFreq = LPFFrequencyParam.Update();
-					const float HPFFreq = HPFFrequencyParam.Update();
+				const float LPFFreq = SourceInfo.LPFCutoffFrequencyParam.Update();
+				const float HPFFreq = SourceInfo.HPFCutoffFrequencyParam.Update();
 
-					// Update the volume
 #if AUDIO_MIXER_ENABLE_DEBUG_MODE				
-					CurrentVolumeValue = (bIsDebugModeEnabled && !SourceInfo.bIsDebugMode) ? 0.0f : SourceInfo.VolumeSourceParam.Update();
+				CurrentVolumeValue = (bIsDebugModeEnabled && !SourceInfo.bIsDebugMode) ? 0.0f : SourceInfo.VolumeSourceParam.Update();
 
 #else
-					CurrentVolumeValue = SourceInfo.VolumeSourceParam.Update();
+				CurrentVolumeValue = SourceInfo.VolumeSourceParam.Update();
 #endif
-					SourceInfo.LowPassFilter.SetFrequency(LPFFreq);
-					SourceInfo.LowPassFilter.Update();
 
-					SourceInfo.HighPassFilter.SetFrequency(HPFFreq);
-					SourceInfo.HighPassFilter.Update();
+				SourceInfo.LowPassFilter.SetFrequency(LPFFreq);
+				SourceInfo.LowPassFilter.Update();
 
-					// feed audio frame through filters
-					int32 SampleIndex = NumInputChans * Frame;
+				SourceInfo.HighPassFilter.SetFrequency(HPFFreq);
+				SourceInfo.HighPassFilter.Update();
 
-					SourceInfo.LowPassFilter.ProcessAudio(&BufferPtr[SampleIndex], &BufferPtr[SampleIndex]);
-					SourceInfo.HighPassFilter.ProcessAudio(&BufferPtr[SampleIndex], &BufferPtr[SampleIndex]);
+				int32 SampleIndex = SourceInfo.NumInputChannels * Frame;
 
-					// Scale by current volume value. TODO: do this as a SIMD operation in its own loop
-					for (int32 Channel = 0; Channel < NumInputChans; ++Channel)
-					{
-						BufferPtr[SampleIndex + Channel] *= CurrentVolumeValue;
-					}
+				SourceInfo.LowPassFilter.ProcessAudio(&PreDistanceAttenBufferPtr[SampleIndex], &PreDistanceAttenBufferPtr[SampleIndex]);
+				SourceInfo.HighPassFilter.ProcessAudio(&PreDistanceAttenBufferPtr[SampleIndex], &PreDistanceAttenBufferPtr[SampleIndex]);
+
+				// Scale by current volume value (note: not including distance attenuation). TODO: do this as a SIMD operation in its own loop
+				for (int32 Channel = 0; Channel < SourceInfo.NumInputChannels; ++Channel)
+				{
+					PreDistanceAttenBufferPtr[SampleIndex + Channel] *= CurrentVolumeValue;
 				}
 			}
 
@@ -994,26 +1384,26 @@ namespace Audio
 			SourceInfo.VolumeSourceParam.Reset();
 
 			// Now process the effect chain if it exists
-			if (SourceInfo.SourceEffects.Num() > 0 && (!SourceInfo.bIsDone || !SourceInfo.bEffectTailsDone))
+			if (SourceInfo.SourceEffects.Num() > 0)
 			{
 				SourceInfo.SourceEffectInputData.CurrentVolume = CurrentVolumeValue;
 
 				SourceInfo.SourceEffectOutputData.AudioFrame.Reset();
-				SourceInfo.SourceEffectOutputData.AudioFrame.AddZeroed(NumInputChans);
+				SourceInfo.SourceEffectOutputData.AudioFrame.AddZeroed(SourceInfo.NumInputChannels);
 
 				SourceInfo.SourceEffectInputData.AudioFrame.Reset();
-				SourceInfo.SourceEffectInputData.AudioFrame.AddZeroed(NumInputChans);
+				SourceInfo.SourceEffectInputData.AudioFrame.AddZeroed(SourceInfo.NumInputChannels);
 
 				float* SourceEffectInputDataFramePtr = SourceInfo.SourceEffectInputData.AudioFrame.GetData();
 				float* SourceEffectOutputDataFramePtr = SourceInfo.SourceEffectOutputData.AudioFrame.GetData();
 
 				// Process the effect chain for this buffer per frame
-				for (int32 Sample = 0; Sample < NumSamples; Sample += NumInputChans)
+				for (int32 Sample = 0; Sample < NumSamples; Sample += SourceInfo.NumInputChannels)
 				{
 					// Get the buffer input sample
-					for (int32 Chan = 0; Chan < NumInputChans; ++Chan)
+					for (int32 Chan = 0; Chan < SourceInfo.NumInputChannels; ++Chan)
 					{
-						SourceEffectInputDataFramePtr[Chan] = BufferPtr[Sample + Chan];
+						SourceEffectInputDataFramePtr[Chan] = PreDistanceAttenBufferPtr[Sample + Chan];
 					}
 
 					for (int32 SourceEffectIndex = 0; SourceEffectIndex < SourceInfo.SourceEffects.Num(); ++SourceEffectIndex)
@@ -1025,7 +1415,7 @@ namespace Audio
 							SourceInfo.SourceEffects[SourceEffectIndex]->ProcessAudio(SourceInfo.SourceEffectInputData, SourceInfo.SourceEffectOutputData);
 
 							// Copy the output of the effect into the input so the next effect will get the outputs audio
-							for (int32 Chan = 0; Chan < NumInputChans; ++Chan)
+							for (int32 Chan = 0; Chan < SourceInfo.NumInputChannels; ++Chan)
 							{
 								SourceEffectInputDataFramePtr[Chan] = SourceEffectOutputDataFramePtr[Chan];
 							}
@@ -1033,93 +1423,23 @@ namespace Audio
 					}
 
 					// Copy audio frame back to the buffer
-					for (int32 Chan = 0; Chan < NumInputChans; ++Chan)
-					{
-						BufferPtr[Sample + Chan] = SourceEffectInputDataFramePtr[Chan];
-					}
-				}
-			}
-
-			bool bReverbPluginUsed = false;
-			if (SourceInfo.bUseReverbPlugin)
-			{
-				const FSpatializationParams* SourceSpatParams = &SourceInfo.SpatParams;
-
-				// Move the audio buffer to the reverb plugin buffer
-				FAudioPluginSourceInputData AudioPluginInputData;
-				AudioPluginInputData.SourceId = SourceId;
-				AudioPluginInputData.AudioBuffer = &SourceInfo.SourceBuffer;
-				AudioPluginInputData.SpatializationParams = SourceSpatParams;
-				AudioPluginInputData.NumChannels = NumInputChans;
-				SourceInfo.AudioPluginOutputData.AudioBuffer.Reset();
-				SourceInfo.AudioPluginOutputData.AudioBuffer.AddZeroed(AudioPluginInputData.AudioBuffer->Num());
-
-				MixerDevice->ReverbPluginInterface->ProcessSourceAudio(AudioPluginInputData, SourceInfo.AudioPluginOutputData);
-
-				// Make sure the buffer counts didn't change and are still the same size
-				AUDIO_MIXER_CHECK(SourceInfo.AudioPluginOutputData.AudioBuffer.Num() == NumSamples);
-
-				// Copy the reverb-processed data back to the source buffer
-				SourceInfo.ReverbPluginOutputBuffer.Reset();
-				SourceInfo.ReverbPluginOutputBuffer.Append(SourceInfo.AudioPluginOutputData.AudioBuffer);
-
-				bReverbPluginUsed = true;
-			}
-
-			if (SourceInfo.bUseOcclusionPlugin)
-			{
-				const FSpatializationParams* SourceSpatParams = &SourceInfo.SpatParams;
-
-				// Move the audio buffer to the occlusion plugin buffer
-				FAudioPluginSourceInputData AudioPluginInputData;
-				AudioPluginInputData.SourceId = SourceId;
-				AudioPluginInputData.AudioBuffer = &SourceInfo.SourceBuffer;
-				AudioPluginInputData.SpatializationParams = SourceSpatParams;
-				AudioPluginInputData.NumChannels = NumInputChans;
-
-				SourceInfo.AudioPluginOutputData.AudioBuffer.Reset();
-				SourceInfo.AudioPluginOutputData.AudioBuffer.AddZeroed(AudioPluginInputData.AudioBuffer->Num());
-
-				MixerDevice->OcclusionInterface->ProcessAudio(AudioPluginInputData, SourceInfo.AudioPluginOutputData);
-
-				// Make sure the buffer counts didn't change and are still the same size
-				AUDIO_MIXER_CHECK(SourceInfo.AudioPluginOutputData.AudioBuffer.Num() == NumSamples);
-
-				// Copy the occlusion-processed data back to the source buffer and mix with the reverb plugin output buffer
-				if (bReverbPluginUsed)
-				{
-					const float* ReverbPluginOutputBufferPtr = SourceInfo.ReverbPluginOutputBuffer.GetData();
-					const float* AudioPluginOutputDataPtr = SourceInfo.AudioPluginOutputData.AudioBuffer.GetData();
-
-					for (int32 i = 0; i < NumSamples; ++i)
+					for (int32 Chan = 0; Chan < SourceInfo.NumInputChannels; ++Chan)
 					{
-						BufferPtr[i] = ReverbPluginOutputBufferPtr[i] + AudioPluginOutputDataPtr[i];
+						PreDistanceAttenBufferPtr[Sample + Chan] = SourceEffectInputDataFramePtr[Chan];
 					}
 				}
-				else
-				{
-					FMemory::Memcpy(BufferPtr, SourceInfo.AudioPluginOutputData.AudioBuffer.GetData(), sizeof(float) * NumSamples);
-				}
-			}
-			else if (bReverbPluginUsed)
-			{
-				const float* ReverbPluginOutputBufferPtr = SourceInfo.ReverbPluginOutputBuffer.GetData();
-				for (int32 i = 0; i < NumSamples; ++i)
-				{
-					BufferPtr[i] += ReverbPluginOutputBufferPtr[i];
-				}
 			}
 
-			// Compute the source envelope
+			// Compute the source envelope using pre-distance attenuation buffer
 			float AverageSampleValue;
 			for (int32 Sample = 0; Sample < NumSamples;)
 			{
 				AverageSampleValue = 0.0f;
-				for (int32 Chan = 0; Chan < NumInputChans; ++Chan)
+				for (int32 Chan = 0; Chan < SourceInfo.NumInputChannels; ++Chan)
 				{
-					AverageSampleValue += BufferPtr[Sample++];
+					AverageSampleValue += PreDistanceAttenBufferPtr[Sample++];
 				}
-				AverageSampleValue /= NumInputChans;
+				AverageSampleValue /= SourceInfo.NumInputChannels;
 				SourceInfo.SourceEnvelopeFollower.ProcessAudio(AverageSampleValue);
 			}
 
@@ -1128,8 +1448,18 @@ namespace Audio
 
 			SourceInfo.bEffectTailsDone = SourceInfo.bEffectTailsDone || SourceInfo.SourceEnvelopeValue < ENVELOPE_TAIL_THRESHOLD;
 
+			// Only scale with distance attenuation and send to source audio to plugins if we're not in output-to-bus only mode
+			if (!SourceInfo.bOutputToBusOnly)
+			{
+				// Apply distance attenuation
+				ApplyDistanceAttenuation(SourceInfo, NumSamples);
+
+				// Send source audio to plugins
+				ComputePluginAudio(SourceInfo, SourceId, NumSamples);
+			}
+
 			// Check the source effect tails condition
-			if (SourceInfo.bIsDone && SourceInfo.bEffectTailsDone)
+			if (SourceInfo.bIsLastBuffer && SourceInfo.bEffectTailsDone)
 			{
 				// If we're done and our tails our done, clear everything out
 				SourceInfo.BufferQueue.Empty();
@@ -1137,61 +1467,15 @@ namespace Audio
 				SourceInfo.NextFrameValues.Reset();
 				SourceInfo.CurrentPCMBuffer = nullptr;
 			}
-
-
-			// If the source has HRTF processing enabled, run it through the spatializer
-			if (SourceInfo.bUseHRTFSpatializer)
-			{
-				SCOPE_CYCLE_COUNTER(STAT_AudioMixerHRTF);
-
-				TAudioSpatializationPtr SpatializationPlugin = MixerDevice->SpatializationPluginInterface;
-
-				AUDIO_MIXER_CHECK(SpatializationPlugin.IsValid());
-				AUDIO_MIXER_CHECK(NumInputChans == 1);
-
-				FAudioPluginSourceInputData AudioPluginInputData;
-				AudioPluginInputData.AudioBuffer = &SourceInfo.SourceBuffer;
-				AudioPluginInputData.NumChannels = NumInputChans;
-				AudioPluginInputData.SourceId = SourceId;
-				AudioPluginInputData.SpatializationParams = &SourceInfo.SpatParams;
-
-				SourceInfo.AudioPluginOutputData.AudioBuffer.Reset();
-				SourceInfo.AudioPluginOutputData.AudioBuffer.AddZeroed(2 * NumOutputFrames);
-
-				SpatializationPlugin->ProcessAudio(AudioPluginInputData, SourceInfo.AudioPluginOutputData);
-
-				// We are now a 2-channel file and should not be spatialized using normal 3d spatialization
-				SourceInfo.NumPostEffectChannels = 2;
-
-				SourceInfo.PostEffectBuffers = &SourceInfo.AudioPluginOutputData.AudioBuffer;
-			}
-			else
-			{
-				// Otherwise our pre- and post-effect channels are the same as the input channels
-				SourceInfo.NumPostEffectChannels = SourceInfo.NumInputChannels;
-
-				// Set the ptr to use for post-effect buffers
-				SourceInfo.PostEffectBuffers = &SourceInfo.SourceBuffer;
-			}
 		}
 	}
 
-	void FMixerSourceManager::ComputeOutputBuffersForIdRange(const int32 SourceIdStart, const int32 SourceIdEnd)
+	void FMixerSourceManager::ComputeOutputBuffersForIdRange(const bool bGenerateBuses, const int32 SourceIdStart, const int32 SourceIdEnd)
 	{
 		SCOPE_CYCLE_COUNTER(STAT_AudioMixerSourceOutputBuffers);
 
 		const int32 NumOutputChannels = MixerDevice->GetNumDeviceChannels();
 
-		// Reset the dry/wet buffers for all the sources
-
-		for (int32 SourceId = SourceIdStart; SourceId < SourceIdEnd; ++SourceId)
-		{
-			FSourceInfo& SourceInfo = SourceInfos[SourceId];
-
-			SourceInfo.OutputBuffer.Reset();
-			SourceInfo.OutputBuffer.AddZeroed(NumOutputSamples);
-		}
-
 		for (int32 SourceId = SourceIdStart; SourceId < SourceIdEnd; ++SourceId)
 		{
 			FSourceInfo& SourceInfo = SourceInfos[SourceId];
@@ -1205,7 +1489,25 @@ namespace Audio
 			// Don't need to compute anything if the source is not playing or paused (it will remain at 0.0 volume)
 			// Note that effect chains will still be able to continue to compute audio output. The source output 
 			// will simply stop being read from.
-			if (!SourceInfo.bIsBusy || !SourceInfo.bIsPlaying || SourceInfo.bIsPaused)
+			if (!SourceInfo.bIsBusy || !SourceInfo.bIsPlaying || SourceInfo.bIsDone)
+			{
+				continue;
+			}
+
+			// If we're in generate buses mode and not a bus, or vice versa, or if we're set to only output audio to buses.
+			// If set to output buses, no need to do any panning for the source. The buses will do the panning.
+			const bool bIsBus = SourceInfo.BusId != INDEX_NONE;
+			if ((bGenerateBuses && !bIsBus) || (!bGenerateBuses && bIsBus) || SourceInfo.bOutputToBusOnly)
+			{
+				continue;
+			}
+
+			// Zero the buffers for all cases, this will catch the pause state. We want to zero buffers when paused.
+			SourceInfo.OutputBuffer.Reset();
+			SourceInfo.OutputBuffer.AddZeroed(NumOutputSamples);
+
+			// If we're paused, then early return now
+			if (SourceInfo.bIsPaused)
 			{
 				continue;
 			}
@@ -1253,18 +1555,66 @@ namespace Audio
 		}
 	}
 
+	void FMixerSourceManager::GenerateSourceAudio(const bool bGenerateBuses, const int32 SourceIdStart, const int32 SourceIdEnd)
+	{
+		// Buses generate their input buffers independently
+		// Get the next block of frames from the source buffers
+		ComputeSourceBuffersForIdRange(bGenerateBuses, SourceIdStart, SourceIdEnd);
+
+		// Compute the audio source buffers after their individual effect chain processing
+		ComputePostSourceEffectBufferForIdRange(bGenerateBuses, SourceIdStart, SourceIdEnd);
+
+		// Get the audio for the output buffers
+		ComputeOutputBuffersForIdRange(bGenerateBuses, SourceIdStart, SourceIdEnd);
+	}
+
+	void FMixerSourceManager::GenerateSourceAudio(const bool bGenerateBuses)
+	{
+		// If there are no buses, don't need to do anything here
+		if (bGenerateBuses && !Buses.Num())
+		{
+			return;
+		}
+
+		if (NumSourceWorkers > 0 && !DisableParallelSourceProcessingCvar)
+		{
+			AUDIO_MIXER_CHECK(SourceWorkers.Num() == NumSourceWorkers);
+			for (int32 i = 0; i < SourceWorkers.Num(); ++i)
+			{
+				FAudioMixerSourceWorker& Worker = SourceWorkers[i]->GetTask();
+				Worker.SetGenerateBuses(bGenerateBuses);
+
+				SourceWorkers[i]->StartBackgroundTask();
+			}
+
+			for (int32 i = 0; i < SourceWorkers.Num(); ++i)
+			{
+				SourceWorkers[i]->EnsureCompletion();
+			}
+		}
+		else
+		{
+			GenerateSourceAudio(bGenerateBuses, 0, NumTotalSources);
+		}
+	}
+
 	void FMixerSourceManager::MixOutputBuffers(const int32 SourceId, AlignedFloatBuffer& OutWetBuffer, const float SendLevel) const
 	{
 		if (SendLevel > 0.0f)
 		{
 			const FSourceInfo& SourceInfo = SourceInfos[SourceId];
-			const float* SourceOutputBufferPtr = SourceInfo.OutputBuffer.GetData();
-			const int32 OutWetBufferSize = OutWetBuffer.Num();
-			float* OutWetBufferPtr = OutWetBuffer.GetData();
 
-			for (int32 SampleIndex = 0; SampleIndex < OutWetBufferSize; ++SampleIndex)
+			// Don't need to mix into submixes if the source is paused
+			if (!SourceInfo.bIsPaused && !SourceInfo.bIsDone && SourceInfo.bIsPlaying)
 			{
-				OutWetBufferPtr[SampleIndex] += SourceOutputBufferPtr[SampleIndex] * SendLevel;
+				const float* SourceOutputBufferPtr = SourceInfo.OutputBuffer.GetData();
+				const int32 OutWetBufferSize = OutWetBuffer.Num();
+				float* OutWetBufferPtr = OutWetBuffer.GetData();
+
+				for (int32 SampleIndex = 0; SampleIndex < OutWetBufferSize; ++SampleIndex)
+				{
+					OutWetBufferPtr[SampleIndex] += SourceOutputBufferPtr[SampleIndex] * SendLevel;
+				}
 			}
 
 // TOODO: SIMD
@@ -1383,41 +1733,58 @@ namespace Audio
 		});
 	}
 
+	const float* FMixerSourceManager::GetPreDistanceAttenuationBuffer(const int32 SourceId) const
+	{
+		return SourceInfos[SourceId].PreDistanceAttenuationBuffer.GetData();
+	}
+
+	const float* FMixerSourceManager::GetPreviousBusBuffer(const int32 SourceId) const
+	{
+		const uint32 BusId = SourceInfos[SourceId].BusId;
+		const FMixerBus* MixerBus = Buses.Find(BusId);
+		return MixerBus->GetPreviousBusBuffer();
+	}
+
+	int32 FMixerSourceManager::GetNumChannels(const int32 SourceId) const
+	{
+		return SourceInfos[SourceId].NumInputChannels;
+	}
+
+	bool FMixerSourceManager::IsBus(const int32 SourceId) const
+	{
+		return SourceInfos[SourceId].BusId != INDEX_NONE;
+	}
+
 	void FMixerSourceManager::ComputeNextBlockOfSamples()
 	{
 		AUDIO_MIXER_CHECK_AUDIO_PLAT_THREAD(MixerDevice);
 
 		SCOPE_CYCLE_COUNTER(STAT_AudioMixerSourceManagerUpdate);
 
+		// Get the this blocks commands before rendering audio
 		if (bPumpQueue)
 		{
 			bPumpQueue = false;
 			PumpCommandQueue();
 		}
 
-		if (NumSourceWorkers > 0 && !DisableParallelSourceProcessingCvar)
-		{
-			AUDIO_MIXER_CHECK(SourceWorkers.Num() == NumSourceWorkers);
-			for (int32 i = 0; i < SourceWorkers.Num(); ++i)
-			{
-				SourceWorkers[i]->StartBackgroundTask();
-			}
+		// First generate non-bus audio (bGenerateBuses = false)
+		GenerateSourceAudio(false);
 
-			for (int32 i = 0; i < SourceWorkers.Num(); ++i)
-			{
-				SourceWorkers[i]->EnsureCompletion();
-			}
-		}
-		else
-		{
-			// Get the next block of frames from the source buffers
-			ComputeSourceBuffersForIdRange(0, NumTotalSources);
+		// Now mix in the non-bus audio into the buses
+		ComputeBuses();
+
+		// Now generate bus audio (bGenerateBuses = true)
+		GenerateSourceAudio(true);
 
-			// Compute the audio source buffers after their individual effect chain processing
-			ComputePostSourceEffectBufferForIdRange(0, NumTotalSources);
+		// Update the buses now
+		UpdateBuses();
 
-			// Get the audio for the output buffers
-			ComputeOutputBuffersForIdRange(0, NumTotalSources);
+		// Let the plugin know we finished processing all sources
+		if (bUsingSpatializationPlugin)
+		{
+			AUDIO_MIXER_CHECK(SpatializationPlugin.IsValid());
+			SpatializationPlugin->OnAllSourcesProcessed();
 		}
 
 		// Update the game thread copy of source doneness
@@ -1425,6 +1792,15 @@ namespace Audio
 		{		
 			FSourceInfo& SourceInfo = SourceInfos[SourceId];
 
+			if (SourceInfo.bIsLastBuffer)
+			{
+				SourceInfo.bIsDone = true;
+			}
+			else
+			{
+				SourceInfo.bIsDone = false;
+			}
+
 			GameThreadInfo.bIsDone[SourceId] = SourceInfo.bIsDone;
 			GameThreadInfo.bEffectTailsDone[SourceId] = SourceInfo.bEffectTailsDone;
 		}
diff --git a/Engine/Source/Runtime/AudioMixer/Private/AudioMixerSourceManager.h b/Engine/Source/Runtime/AudioMixer/Private/AudioMixerSourceManager.h
--- a/Engine/Source/Runtime/AudioMixer/Private/AudioMixerSourceManager.h
+++ b/Engine/Source/Runtime/AudioMixer/Private/AudioMixerSourceManager.h
@@ -8,31 +8,14 @@
 #include "CoreMinimal.h"
 #include "AudioMixerBuffer.h"
 #include "AudioMixerSubmix.h"
+#include "AudioMixerBus.h"
 #include "DSP/OnePole.h"
 #include "DSP/Filter.h"
 #include "DSP/EnvelopeFollower.h"
+#include "DSP/ParamInterpolator.h"
 #include "IAudioExtensionPlugin.h"
 #include "Containers/Queue.h"
 
-#define ENABLE_AUDIO_OUTPUT_DEBUGGING 0
-
-#if AUDIO_MIXER_ENABLE_DEBUG_MODE
-
-// Macro which checks if the source id is in debug mode, avoids having a bunch of #ifdefs in code
-#define AUDIO_MIXER_DEBUG_LOG(SourceId, Format, ...)																							\
-	if (SourceInfos[SourceId].bIsDebugMode)																													\
-	{																																			\
-		FString CustomMessage = FString::Printf(Format, ##__VA_ARGS__);																			\
-		FString LogMessage = FString::Printf(TEXT("<Debug Sound Log> [Id=%d][Name=%s]: %s"), SourceId, *SourceInfos[SourceId].DebugName, *CustomMessage);	\
-		UE_LOG(LogAudioMixer, Log, TEXT("%s"), *LogMessage);																								\
-	}
-
-#else
-
-#define AUDIO_MIXER_DEBUG_LOG(SourceId, Message)
-
-#endif
-
 
 namespace Audio
 {
@@ -40,7 +23,26 @@ namespace Audio
 	class FMixerDevice;
 	class FMixerSourceVoice;
 
-	struct FMixerSourceVoiceBuffer;
+	/** Struct defining a source voice buffer. */
+	struct FMixerSourceVoiceBuffer
+	{
+		/** PCM float data. */
+		TArray<float> AudioData;
+
+		/** The amount of samples of audio data in the float buffer array. */
+		uint32 Samples;
+
+		/** How many times this buffer will loop. */
+		int32 LoopCount;
+
+		/** If this buffer is from real-time decoding and needs to make callbacks for more data. */
+		uint32 bRealTimeBuffer : 1;
+
+		FMixerSourceVoiceBuffer()
+		{
+			FMemory::Memzero(this, sizeof(*this));
+		}
+	};
 
 	typedef TSharedPtr<FMixerSourceVoiceBuffer, ESPMode::ThreadSafe> FMixerSourceBufferPtr;
 	typedef TSharedPtr<FMixerSubmix, ESPMode::ThreadSafe> FMixerSubmixPtr;
@@ -70,10 +72,20 @@ namespace Audio
 		bool bIsMainSend;
 	};
 
+	// Struct holding mappings of bus ids (unique ids) to send level
+	struct FMixerBusSend
+	{
+		uint32 BusId;
+		float SendLevel;
+	};
+
 	struct FMixerSourceVoiceInitParams
 	{
 		ISourceBufferQueueListener* BufferQueueListener;
 		TArray<FMixerSourceSubmixSend> SubmixSends;
+		TArray<FMixerBusSend> BusSends;
+		uint32 BusId;
+		float BusDuration;
 		uint32 SourceEffectChainId;
 		TArray<FSourceEffectChainEntry> SourceEffectChain;
 		FMixerSourceVoice* SourceVoice;
@@ -87,9 +99,12 @@ namespace Audio
 		bool bPlayEffectChainTails;
 		bool bUseHRTFSpatialization;
 		bool bIsDebugMode;
+		bool bOutputToBusOnly;
 
 		FMixerSourceVoiceInitParams()
 			: BufferQueueListener(nullptr)
+			, BusId(INDEX_NONE)
+			, BusDuration(0.0f)
 			, SourceEffectChainId(INDEX_NONE)
 			, SourceVoice(nullptr)
 			, NumInputChannels(0)
@@ -100,55 +115,10 @@ namespace Audio
 			, bPlayEffectChainTails(true)
 			, bUseHRTFSpatialization(false)
 			, bIsDebugMode(false)
+			, bOutputToBusOnly(false)
 		{}
 	};
 
-	class FSourceParam
-	{
-	public:
-		FSourceParam()
-		 : CurrentValue(0.0f)
-		 , StartingValue(0.0f)
-		 , TargetValue(0.0f)
-		 , DeltaValue(0.0f)
-		{}
-
-		// Set the parameter value to the given target value over the given interpolation frames.
-		FORCEINLINE void SetValue(const float InValue, const int32 InNumInterpFrames)
-		{
-			TargetValue = InValue;
-			DeltaValue = (InValue - CurrentValue) / InNumInterpFrames;
-			StartingValue = CurrentValue;
-		}
-
-		// Resets the delta value back to 0.0. To be called at beginning of callback render.
-		FORCEINLINE void Reset()
-		{ 
-			DeltaValue = 0.0f; 
-			CurrentValue = TargetValue;
-		}
-
-		// Updates the parameter, assumes called in one of the frames.
-		FORCEINLINE float Update() 
-		{ 
-			CurrentValue += DeltaValue;
-			return CurrentValue; 
-		}
-
-		// Returns the current value, but does not update the value
-		FORCEINLINE float GetValue() const 
-		{ 
-			return CurrentValue; 
-		}
-
-	private:
-		float CurrentValue;
-		float StartingValue;
-		float TargetValue;
-		float DeltaValue;
-	};
-
-
 	class FSourceChannelMap
 	{
 	public:
@@ -166,7 +136,7 @@ namespace Audio
 				ChannelValues.Reset();
 				for (int32 i = 0; i < ChannelMap.Num(); ++i)
 				{
-					ChannelValues.Add(FSourceParam());
+					ChannelValues.Add(FParam());
 					ChannelValues[i].SetValue(ChannelMap[i], InNumInterpFrames);
 				}
 			}
@@ -207,13 +177,13 @@ namespace Audio
 			int32 CurrentSize = ChannelValues.Num();
 			for (int32 i = CurrentSize; i < ToSize; ++i)
 			{
-				ChannelValues.Add(FSourceParam());
+				ChannelValues.Add(FParam());
 				ChannelValues[i].SetValue(0.0f, InNumInterpFrames);
 			}
 		}
 
 	private:
-		TArray<FSourceParam> ChannelValues;
+		TArray<FParam> ChannelValues;
 	};
 
 	struct FSourceManagerInitParams
@@ -236,6 +206,7 @@ namespace Audio
 
 		bool GetFreeSourceId(int32& OutSourceId);
 		int32 GetNumActiveSources() const;
+		int32 GetNumActiveBuses() const;
 
 		void ReleaseSourceId(const int32 SourceId);
 		void InitSource(const int32 SourceId, const FMixerSourceVoiceInitParams& InitParams);
@@ -245,6 +216,7 @@ namespace Audio
 		void Pause(const int32 SourceId);
 		void SetPitch(const int32 SourceId, const float Pitch);
 		void SetVolume(const int32 SourceId, const float Volume);
+		void SetDistanceAttenuation(const int32 SourceId, const float DistanceAttenuation);
 		void SetSpatializationParams(const int32 SourceId, const FSpatializationParams& InParams);
 		void SetChannelMap(const int32 SourceId, const TArray<float>& InChannelMap, const bool bInIs3D, const bool bInIsCenterChannelOnly);
 		void SetLPFFrequency(const int32 SourceId, const float Frequency);
@@ -259,21 +231,34 @@ namespace Audio
 		void ComputeNextBlockOfSamples();
 		void MixOutputBuffers(const int32 SourceId, AlignedFloatBuffer& OutWetBuffer, const float SendLevel) const;
 
-		void SetSubmixSendInfo(const int32 SourceId, FMixerSubmixPtr Submix, const float SendLevel);
+		void SetSubmixSendInfo(const int32 SourceId, const FMixerSourceSubmixSend& SubmixSend);
 
 		void UpdateDeviceChannelCount(const int32 InNumOutputChannels);
 
 		void UpdateSourceEffectChain(const uint32 SourceEffectChainId, const TArray<FSourceEffectChainEntry>& SourceEffectChain, const bool bPlayEffectChainTails);
 
+		const float* GetPreDistanceAttenuationBuffer(const int32 SourceId) const;
+		const float* GetPreviousBusBuffer(const int32 SourceId) const;
+		int32 GetNumChannels(const int32 SourceId) const;
+		int32 GetNumOutputFrames() const { return NumOutputFrames; }
+		bool IsBus(const int32 SourceId) const;
+
 	private:
 
 		void ReleaseSource(const int32 SourceId);
 		void BuildSourceEffectChain(const int32 SourceId, FSoundEffectSourceInitData& InitData, const TArray<FSourceEffectChainEntry>& SourceEffectChain);
 		void ResetSourceEffectChain(const int32 SourceId);
 		void ReadSourceFrame(const int32 SourceId);
-		void ComputeSourceBuffersForIdRange(const int32 SourceIdStart, const int32 SourceIdEnd);
-		void ComputePostSourceEffectBufferForIdRange(const int32 SourceIdStart, const int32 SourceIdEnd);
-		void ComputeOutputBuffersForIdRange(const int32 SourceIdStart, const int32 SourceIdEnd);
+
+		void GenerateSourceAudio(const bool bGenerateBuses);
+		void GenerateSourceAudio(const bool bGenerateBuses, const int32 SourceIdStart, const int32 SourceIdEnd);
+
+		void ComputeSourceBuffersForIdRange(const bool bGenerateBuses, const int32 SourceIdStart, const int32 SourceIdEnd);
+		void ComputePostSourceEffectBufferForIdRange(const bool bGenerateBuses, const int32 SourceIdStart, const int32 SourceIdEnd);
+		void ComputeOutputBuffersForIdRange(const bool bGenerateBuses, const int32 SourceIdStart, const int32 SourceIdEnd);
+
+		void ComputeBuses();
+		void UpdateBuses();
 
 		void AudioMixerThreadCommand(TFunction<void()> InFunction);
 		void PumpCommandQueue();
@@ -286,25 +271,25 @@ namespace Audio
 			FMixerSourceManager* SourceManager;
 			int32 StartSourceId;
 			int32 EndSourceId;
+			bool bGenerateBuses;
 
 		public:
 			FAudioMixerSourceWorker(FMixerSourceManager* InSourceManager, const int32 InStartSourceId, const int32 InEndSourceId)
 				: SourceManager(InSourceManager)
 				, StartSourceId(InStartSourceId)
 				, EndSourceId(InEndSourceId)
+				, bGenerateBuses(false)
 			{
 			}
 
-			void DoWork()
+			void SetGenerateBuses(bool bInGenerateBuses)
 			{
-				// Get the next block of frames from the source buffers
-				SourceManager->ComputeSourceBuffersForIdRange(StartSourceId, EndSourceId);
-
-				// Compute the audio source buffers after their individual effect chain processing
-				SourceManager->ComputePostSourceEffectBufferForIdRange(StartSourceId, EndSourceId);
+				bGenerateBuses = bInGenerateBuses;
+			}
 
-				// Get the audio for the output buffers
-				SourceManager->ComputeOutputBuffersForIdRange(StartSourceId, EndSourceId);
+			void DoWork()
+			{
+				SourceManager->GenerateSourceAudio(bGenerateBuses, StartSourceId, EndSourceId);
 			}
 
 			FORCEINLINE TStatId GetStatId() const
@@ -315,6 +300,9 @@ namespace Audio
 
 		FMixerDevice* MixerDevice;
 
+		// Cached ptr to an optional spatialization plugin
+		TAudioSpatializationPtr SpatializationPlugin;
+
 		// Array of pointers to game thread audio source objects
 		TArray<FMixerSourceVoice*> MixerSources;
 
@@ -335,18 +323,34 @@ namespace Audio
 			// Data used for rendering sources
 			FMixerSourceBufferPtr CurrentPCMBuffer;
 			int32 CurrentAudioChunkNumFrames;
+
+			// The post-attenuation source buffer, used to send audio to submixes
 			TArray<float> SourceBuffer;
+			TArray<float> PreDistanceAttenuationBuffer;
+
 			TArray<float> CurrentFrameValues;
 			TArray<float> NextFrameValues;
 			float CurrentFrameAlpha;
 			int32 CurrentFrameIndex;
 			int64 NumFramesPlayed;
 
+			TArray<FMixerSourceSubmixSend> SubmixSends;
+
+			// What bus Id this source is, if it is a bus. This is INDEX_NONE for sources which are not buses.
+			uint32 BusId;
+
+			// Number of samples to count for bus
+			int64 BusDurationFrames;
+
+			// What buses this source is sending its audio to. Used to remove this source from the bus send list.
+			TArray<uint32> BusSends;
+
 			// Interpolated source params
-			FSourceParam PitchSourceParam;
-			FSourceParam VolumeSourceParam;
-			FSourceParam LPFCutoffFrequencyParam;
-			FSourceParam HPFCutoffFrequencyParam;
+			FParam PitchSourceParam;
+			FParam VolumeSourceParam;
+			FParam DistanceAttenuationSourceParam;
+			FParam LPFCutoffFrequencyParam;
+			FParam HPFCutoffFrequencyParam;
 
 			// One-Pole LPFs and HPFs per source
 			Audio::FOnePoleFilter LowPassFilter;
@@ -387,6 +391,8 @@ namespace Audio
 			bool bUseOcclusionPlugin;
 			bool bUseReverbPlugin;
 			bool bIsDone;
+			bool bIsLastBuffer;
+			bool bOutputToBusOnly;
 
 			bool bIsDebugMode;
 			FString DebugName;
@@ -397,9 +403,18 @@ namespace Audio
 			int32 NumInputFrames;
 		};
 
+		void ApplyDistanceAttenuation(FSourceInfo& InSourceInfo, int32 NumSamples);
+		void ComputePluginAudio(FSourceInfo& InSourceInfo, int32 SourceId, int32 NumSamples);
+
 		// Array of source infos.
 		TArray<FSourceInfo> SourceInfos;
 
+		// Array of active source ids
+		TArray<int32> ActiveSourceIds;
+
+		// Map of bus object Id's to bus data. 
+		TMap<uint32, FMixerBus> Buses;
+
 		// Async task workers for processing sources in parallel
 		TArray<FAsyncTask<FAudioMixerSourceWorker>*> SourceWorkers;
 
@@ -421,6 +436,7 @@ namespace Audio
 		int32 NumSourceWorkers;
 
 		uint8 bInitialized : 1;
+		uint8 bUsingSpatializationPlugin : 1;
 
 		friend class FMixerSourceVoice;
 		// Set to true when the audio source manager should pump the command queue
diff --git a/Engine/Source/Runtime/AudioMixer/Private/AudioMixerSourceVoice.cpp b/Engine/Source/Runtime/AudioMixer/Private/AudioMixerSourceVoice.cpp
--- a/Engine/Source/Runtime/AudioMixer/Private/AudioMixerSourceVoice.cpp
+++ b/Engine/Source/Runtime/AudioMixer/Private/AudioMixerSourceVoice.cpp
@@ -12,27 +12,44 @@ namespace Audio
 	* FMixerSourceVoice Implementation
 	*/
 
-	FMixerSourceVoice::FMixerSourceVoice(FMixerDevice* InMixerDevice, FMixerSourceManager* InSourceManager)
-		: SourceManager(InSourceManager)
-		, MixerDevice(InMixerDevice)
-		, NumBuffersQueued(0)
-		, Pitch(-1.0f)
-		, Volume(-1.0f)
-		, Distance(-1.0f)
-		, LPFFrequency(-1.0f)
-		, HPFFrequency(-1.0f)
-		, SourceId(INDEX_NONE)
-		, bIsPlaying(false)
-		, bIsPaused(false)
-		, bIsActive(false)
-	{
-		AUDIO_MIXER_CHECK(SourceManager != nullptr);
+	FMixerSourceVoice::FMixerSourceVoice()
+	{
+		Reset(nullptr);
 	}
 
 	FMixerSourceVoice::~FMixerSourceVoice()
 	{
 	}
 
+	void FMixerSourceVoice::Reset(FMixerDevice* InMixerDevice)
+	{
+		if (InMixerDevice)
+		{
+			MixerDevice = InMixerDevice;
+			SourceManager = MixerDevice->GetSourceManager();
+		}
+		else
+		{
+			MixerDevice = nullptr;
+			SourceManager = nullptr;
+		}
+
+		NumBuffersQueued.Reset();
+		Pitch = -1.0f;
+		Volume = -1.0f;
+		DistanceAttenuation = -1.0f;
+		Distance = -1.0f;
+		LPFFrequency = -1.0f;
+		HPFFrequency = -1.0f;
+		SourceId = INDEX_NONE;
+		bIsPlaying = false;
+		bIsPaused = false;
+		bIsActive = false;
+		bIsBus = false;
+		bOutputToBusOnly = false;
+		SubmixSends.Reset();
+	}
+
 	bool FMixerSourceVoice::Init(const FMixerSourceVoiceInitParams& InitParams)
 	{
 		AUDIO_MIXER_CHECK_GAME_THREAD(MixerDevice);
@@ -42,6 +59,9 @@ namespace Audio
 			AUDIO_MIXER_CHECK(InitParams.BufferQueueListener != nullptr);
 			AUDIO_MIXER_CHECK(InitParams.NumInputChannels > 0);
 
+			bOutputToBusOnly = InitParams.bOutputToBusOnly;
+			bIsBus = InitParams.BusId != INDEX_NONE;
+
 			for (int32 i = 0; i < InitParams.SubmixSends.Num(); ++i)
 			{
 				SubmixSends.Add(InitParams.SubmixSends[i].Submix->GetId(), InitParams.SubmixSends[i]);
@@ -93,6 +113,17 @@ namespace Audio
 		}
 	}
 
+	void FMixerSourceVoice::SetDistanceAttenuation(const float InDistanceAttenuation)
+	{
+		AUDIO_MIXER_CHECK_GAME_THREAD(MixerDevice);
+
+		if (DistanceAttenuation != InDistanceAttenuation)
+		{
+			DistanceAttenuation = InDistanceAttenuation;
+			SourceManager->SetDistanceAttenuation(SourceId, InDistanceAttenuation);
+		}
+	}
+
 	void FMixerSourceVoice::SetLPFFrequency(const float InLPFFrequency)
 	{
 		AUDIO_MIXER_CHECK_GAME_THREAD(MixerDevice);
@@ -212,6 +243,8 @@ namespace Audio
 	{
 		AUDIO_MIXER_CHECK_AUDIO_PLAT_THREAD(MixerDevice);
 
+		check(!bOutputToBusOnly);
+
 		return SourceManager->MixOutputBuffers(SourceId, OutWetBuffer, SendLevel);
 	}
 
@@ -219,31 +252,33 @@ namespace Audio
 	{
 		AUDIO_MIXER_CHECK_GAME_THREAD(MixerDevice);
 
-		FMixerSourceSubmixSend* SubmixSend = SubmixSends.Find(Submix->GetId());
-		bool bChanged = false;
-		if (!SubmixSend)
-		{
-			FMixerSourceSubmixSend NewSubmixSend;
-			NewSubmixSend.Submix = Submix;
-			NewSubmixSend.SendLevel = SendLevel;
-			NewSubmixSend.bIsMainSend = false;
-			SubmixSends.Add(Submix->GetId(), NewSubmixSend);
-			bChanged = true;
-		}
-		else
+		if (!bOutputToBusOnly)
 		{
-			if (!FMath::IsNearlyEqual(SubmixSend->SendLevel, SendLevel))
+			FMixerSourceSubmixSend* SubmixSend = SubmixSends.Find(Submix->GetId());
+			if (!SubmixSend)
+			{
+				FMixerSourceSubmixSend NewSubmixSend;
+				NewSubmixSend.Submix = Submix;
+				NewSubmixSend.SendLevel = SendLevel;
+				NewSubmixSend.bIsMainSend = false;
+				SubmixSends.Add(Submix->GetId(), NewSubmixSend);
+				SourceManager->SetSubmixSendInfo(SourceId, NewSubmixSend);
+			}
+			else if (!FMath::IsNearlyEqual(SubmixSend->SendLevel, SendLevel))
 			{
 				SubmixSend->SendLevel = SendLevel;
-				bChanged = true;
+				SourceManager->SetSubmixSendInfo(SourceId, *SubmixSend);
 			}
 		}
+	}
+
+	void FMixerSourceVoice::OnMixBus(FMixerSourceBufferPtr OutMixerSourceBuffer)
+	{
+		AUDIO_MIXER_CHECK_AUDIO_PLAT_THREAD(MixerDevice);
 
-		if (bChanged)
+		for (uint32 i = 0; i < OutMixerSourceBuffer->Samples; ++i)
 		{
-			SourceManager->SetSubmixSendInfo(SourceId, Submix, SendLevel);
+			OutMixerSourceBuffer->AudioData[i] = 0.0f;
 		}
 	}
-
-
-}
\ No newline at end of file
+}
diff --git a/Engine/Source/Runtime/AudioMixer/Private/AudioMixerSourceVoice.h b/Engine/Source/Runtime/AudioMixer/Private/AudioMixerSourceVoice.h
--- a/Engine/Source/Runtime/AudioMixer/Private/AudioMixerSourceVoice.h
+++ b/Engine/Source/Runtime/AudioMixer/Private/AudioMixerSourceVoice.h
@@ -20,9 +20,12 @@ namespace Audio
 	class FMixerSourceVoice
 	{
 	public:
-		FMixerSourceVoice(FMixerDevice* InMixerDevice, FMixerSourceManager* InSourceManager);
+		FMixerSourceVoice();
 		~FMixerSourceVoice();
 
+		// Resets the source voice state
+		void Reset(FMixerDevice* InMixerDevice);
+
 		// Initializes the mixer source voice
 		bool Init(const FMixerSourceVoiceInitParams& InFormat);
 
@@ -41,6 +44,9 @@ namespace Audio
 		// Sets the source voice volume value.
 		void SetVolume(const float InVolume);
 
+		// Sets the source voice distance attenuation.
+		void SetDistanceAttenuation(const float InDistanceAttenuation);
+		
 		// Sets the source voice's LPF filter frequency.
 		void SetLPFFrequency(const float InFrequency);
 
@@ -89,8 +95,8 @@ namespace Audio
 		// Sets the submix send levels
 		void SetSubmixSendInfo(FMixerSubmixPtr Submix, const float SendLevel);
 
-		// Returns the submix that owns this source voice.
-		TMap<uint32, FMixerSourceSubmixSend>& GetSubmixSends() { return SubmixSends; }
+		// Called when the source is a bus and needs to mix other sources together to generate output
+		void OnMixBus(FMixerSourceBufferPtr OutMixerSourceBuffer);
 
 	private:
 
@@ -103,6 +109,7 @@ namespace Audio
 		FThreadSafeCounter NumBuffersQueued;
 		float Pitch;
 		float Volume;
+		float DistanceAttenuation;
 		float Distance;
 		float LPFFrequency;
 		float HPFFrequency;
@@ -110,6 +117,8 @@ namespace Audio
 		uint16 bIsPlaying : 1;
 		uint16 bIsPaused : 1;
 		uint16 bIsActive : 1;
+		uint16 bOutputToBusOnly : 1;
+		uint16 bIsBus : 1;
 	};
 
 }
diff --git a/Engine/Source/Runtime/AudioMixer/Private/Components/SynthComponent.cpp b/Engine/Source/Runtime/AudioMixer/Private/Components/SynthComponent.cpp
--- a/Engine/Source/Runtime/AudioMixer/Private/Components/SynthComponent.cpp
+++ b/Engine/Source/Runtime/AudioMixer/Private/Components/SynthComponent.cpp
@@ -32,20 +32,45 @@ void USynthSound::Init(USynthComponent* InSynthComponent, int32 InNumChannels)
 	Duration = INDEFINITELY_LOOPING_DURATION;
 	bLooping = true;
 	SampleRate = InSynthComponent->GetAudioDevice()->SampleRate;
+	bAudioMixer = InSynthComponent->GetAudioDevice()->IsAudioMixerEnabled();
 }
 
 bool USynthSound::OnGeneratePCMAudio(TArray<uint8>& OutAudio, int32 NumSamples)
 {
-	GeneratedPCMData.Reset();
-	GeneratedPCMData.AddZeroed(NumSamples);
-
-	OwningSynthComponent->OnGeneratePCMAudio(GeneratedPCMData);
+	if (bAudioMixer)
+	{
+		// If running with audio mixer, the output audio buffer will be in floats already
+		OutAudio.Reset();
+		OutAudio.AddZeroed(NumSamples * sizeof(float));
 
-	OutAudio.Append((uint8*)GeneratedPCMData.GetData(), NumSamples * sizeof(int16));
+		OwningSynthComponent->OnGeneratePCMAudio((float*)OutAudio.GetData(), NumSamples);
+	}
+	else
+	{
+		// Use the float scratch buffer instead of the out buffer directly
+		FloatBuffer.Reset();
+		FloatBuffer.AddZeroed(NumSamples * sizeof(float));
+		
+		float* FloatBufferDataPtr = FloatBuffer.GetData();
+		OwningSynthComponent->OnGeneratePCMAudio(FloatBufferDataPtr, NumSamples);
+
+		// Convert the float buffer to int16 data
+		int16* OutAudioBuffer = (int16*)OutAudio.GetData();
+		for (int32 i = 0; i < NumSamples; ++i)
+		{
+			OutAudioBuffer[i] = (int16)(32767.0f * FloatBufferDataPtr[i]);
+		}
+	}
 
 	return true;
 }
 
+Audio::EAudioMixerStreamDataFormat::Type USynthSound::GetGeneratedPCMDataFormat() const 
+{ 
+	// Only audio mixer supports return float buffers
+	return bAudioMixer ? Audio::EAudioMixerStreamDataFormat::Float : Audio::EAudioMixerStreamDataFormat::Int16;
+}
+
 USynthComponent::USynthComponent(const FObjectInitializer& ObjectInitializer)
 	: Super(ObjectInitializer)
 {
@@ -219,52 +244,17 @@ void USynthComponent::PumpPendingMessages()
 	}
 }
 
-void USynthComponent::OnGeneratePCMAudio(TArray<int16>& GeneratedPCMData)
+void USynthComponent::OnGeneratePCMAudio(float* GeneratedPCMData, int32 NumSamples)
 {
 	PumpPendingMessages();
-	check(GeneratedPCMData.Num() > 0);
-
-	AudioFloatData.Reset();
-
-#if SYNTH_GENERATOR_TEST_TONE
-	if (NumChannels == 1)
-	{
-		for (int32 i = 0; i < SamplesNeeded; ++i)
-		{
-			const float SampleFloat = TestSineLeft.ProcessAudio();
-			const int16 SamplePCM = (int16)(32767.0f * SampleFloat);
-			AudioData.Add(SamplePCM);
-		}
-	}
-	else
-	{
-		for (int32 i = 0; i < SamplesNeeded; ++i)
-		{
-			float SampleFloat = TestSineLeft.ProcessAudio();
-			int16 SamplePCM = (int16)(32767.0f * SampleFloat);
-			AudioData.Add(SamplePCM);
 
-			SampleFloat = TestSineRight.ProcessAudio();
-			SamplePCM = (int16)(32767.0f * SampleFloat);
-			AudioData.Add(SamplePCM);
-		}
-	}
-#else
-	AudioFloatData.AddZeroed(GeneratedPCMData.Num());
+	check(NumSamples > 0);
 
 	// Only call into the synth if we're actually playing, otherwise, we'll write out zero's
 	if (bIsSynthPlaying)
 	{
-		this->OnGenerateAudio(AudioFloatData);
-
-		// Convert the float data to int16 data
-		const float* AudioFloatDataPtr = AudioFloatData.GetData();
-		for (int32 i = 0; i < GeneratedPCMData.Num(); ++i)
-		{
-			GeneratedPCMData[i] = (int16)(32767.0f * AudioFloatDataPtr[i]);
-		}
+		this->OnGenerateAudio(GeneratedPCMData, NumSamples);
 	}
-#endif
 }
 
 void USynthComponent::Start()
diff --git a/Engine/Source/Runtime/AudioMixer/Private/DSP/Granulator.cpp b/Engine/Source/Runtime/AudioMixer/Private/DSP/Granulator.cpp
--- a/Engine/Source/Runtime/AudioMixer/Private/DSP/Granulator.cpp
+++ b/Engine/Source/Runtime/AudioMixer/Private/DSP/Granulator.cpp
@@ -350,7 +350,7 @@ namespace Audio
 		if (Parent->Mode == EGranularSynthMode::Granulation)
 		{
 			// Generate stereo output into the scratch buffer independent of if the loaded sample is stereo or not
-			SampleBufferReader.Generate(FrameScratch, 1, 2, true);
+			SampleBufferReader.Generate(FrameScratch.GetData(), 1, 2, true);
 			const float EnvelopeValue = GetEnvelopeValue();
 
 			for (int32 Channel = 0; Channel < 2; ++Channel)
@@ -777,11 +777,8 @@ namespace Audio
 		return SampleBuffer.SampleDuration;
 	}
 
-	void FGranularSynth::Generate(TArray<float>& OutAudiobuffer, const int32 NumFrames)
+	void FGranularSynth::Generate(float* OutAudiobuffer, const int32 NumFrames)
 	{
-		OutAudiobuffer.Reset();
-		OutAudiobuffer.AddZeroed(2 * NumFrames);
-
 		if (SampleBuffer.GetData() == nullptr)
 		{
 			return;
diff --git a/Engine/Source/Runtime/AudioMixer/Private/DSP/SampleBufferReader.cpp b/Engine/Source/Runtime/AudioMixer/Private/DSP/SampleBufferReader.cpp
--- a/Engine/Source/Runtime/AudioMixer/Private/DSP/SampleBufferReader.cpp
+++ b/Engine/Source/Runtime/AudioMixer/Private/DSP/SampleBufferReader.cpp
@@ -187,12 +187,13 @@ namespace Audio
 		return (float)PCMSampleValue / 32767.0f;
 	}
 
-	bool FSampleBufferReader::Generate(TArray<float>& OutAudioBuffer, const int32 NumFrames, const int32 OutChannels, const bool bWrap)
+	bool FSampleBufferReader::Generate(float* OutAudioBuffer, const int32 NumFrames, const int32 OutChannels, const bool bWrap)
 	{
 		// Don't have a buffer yet, so fill in zeros, say we're not done
 		if (!HasBuffer() || bIsFinished)
 		{
-			for (int32 i = 0; i < OutAudioBuffer.Num(); ++i)
+			int32 NumSamples = NumFrames * OutChannels;
+			for (int32 i = 0; i < NumSamples; ++i)
 			{
 				OutAudioBuffer[i] = 0.0f;
 			}
diff --git a/Engine/Source/Runtime/AudioMixer/Public/AudioMixer.h b/Engine/Source/Runtime/AudioMixer/Public/AudioMixer.h
--- a/Engine/Source/Runtime/AudioMixer/Public/AudioMixer.h
+++ b/Engine/Source/Runtime/AudioMixer/Public/AudioMixer.h
@@ -10,6 +10,8 @@
 #include "HAL/Runnable.h"
 #include "Stats/Stats.h"
 #include "Classes/Sound/AudioSettings.h"
+#include "Misc/SingleThreadRunnable.h"
+#include "DSP/ParamInterpolator.h"
 
 // defines used for AudioMixer.h
 #define AUDIO_PLATFORM_ERROR(INFO)			(OnAudioMixerPlatformError(INFO, FString(__FILE__), __LINE__))
@@ -302,7 +304,7 @@ namespace Audio
 	{
 	public:
 		virtual void RegisterDeviceChangedListener() {}
-		virtual void UnRegisterDeviceChangedListener() {}
+		virtual void UnregisterDeviceChangedListener() {}
 		virtual void OnDefaultCaptureDeviceChanged(const EAudioDeviceRole InAudioDeviceRole, const FString& DeviceId) {}
 		virtual void OnDefaultRenderDeviceChanged(const EAudioDeviceRole InAudioDeviceRole, const FString& DeviceId) {}
 		virtual void OnDeviceAdded(const FString& DeviceId) {}
@@ -314,6 +316,7 @@ namespace Audio
 
 	/** Abstract interface for mixer platform. */
 	class AUDIOMIXER_API IAudioMixerPlatformInterface : public FRunnable,
+														public FSingleThreadRunnable,
 														public IAudioMixerDeviceChangedLister
 	{
 
@@ -408,6 +411,16 @@ namespace Audio
 		uint32 Run() override;
 		//~ End FRunnable
 
+		/**
+		*  FSingleThreadRunnable accessor for ticking this FRunnable when multi-threading is disabled.
+		*  @return FSingleThreadRunnable Interface for this FRunnable object.
+		*/
+		virtual class FSingleThreadRunnable* GetSingleThreadInterface() override { return this; }
+
+		//~ Begin FSingleThreadRunnable Interface
+		virtual void Tick() override;
+		//~ End FSingleThreadRunnable Interface
+
 		/** Constructor. */
 		IAudioMixerPlatformInterface();
 
@@ -420,9 +433,15 @@ namespace Audio
 		/** Start a fadeout. Prevents pops during shutdown. */
 		void FadeOut();
 
+		/** Sets the mater volume of the audio device. This attenuates all audio, used for muting, etc. */
+		void SetMasterVolume(const float InVolume);
+
 		/** Returns the last error generated. */
 		FString GetLastError() const { return LastError; }
 
+		/** This is called after InitializeHardware() is called. */
+		void PostInitializeHardware();
+
 	protected:
 		
 		// Run the "main" audio device
@@ -445,7 +464,7 @@ namespace Audio
 		void StopGeneratingAudio();
 
 		/** Performs buffer fades for shutdown/startup of audio mixer. */
-		void PerformFades();
+		void ApplyMasterAttenuation();
 
 	protected:
 
@@ -483,6 +502,18 @@ namespace Audio
 		/** The number of mixer buffers to queue on the output source voice. */
 		int32 NumOutputBuffers;
 
+		/** The target master volume. */
+		float TargetMasterVolume;
+
+		/** The fade value. Used for fading in/out master audio. */
+		float FadeVolume;
+
+		/** The master volume of the audio device. */
+		FParam MasterVolumeParam;
+
+		/** Source param used to fade in and out audio device. */
+		FParam FadeParam;
+
 		/** String containing the last generated error. */
 		FString LastError;
 
@@ -492,11 +523,10 @@ namespace Audio
 		/** Flag if the audio device is in the process of changing. Prevents more buffers from being submitted to platform. */
 		FThreadSafeBool bAudioDeviceChanging;
 
-		FThreadSafeBool bFadingIn;
-		FThreadSafeBool bFadingOut;
+		FThreadSafeBool bPerformingFade;
 		FThreadSafeBool bFadedOut;
-		float FadeEnvelopeValue;
-
+		FThreadSafeBool bUpdateMasterVolume;
+		FThreadSafeBool bIsDeviceInitialized;
 	};
 
 
diff --git a/Engine/Source/Runtime/AudioMixer/Public/AudioMixerDevice.h b/Engine/Source/Runtime/AudioMixer/Public/AudioMixerDevice.h
--- a/Engine/Source/Runtime/AudioMixer/Public/AudioMixerDevice.h
+++ b/Engine/Source/Runtime/AudioMixer/Public/AudioMixerDevice.h
@@ -94,7 +94,8 @@ namespace Audio
 		bool IsAudioRenderingThread();
 
 		// Public Functions
-		FMixerSourceVoice* GetMixerSourceVoice(const FWaveInstance* InWaveInstance, ISourceBufferQueueListener* InBufferQueueListener, bool bUseHRTFSpatialization);
+		FMixerSourceVoice* GetMixerSourceVoice();
+		void ReleaseMixerSourceVoice(FMixerSourceVoice* InSourceVoice);
 		int32 GetNumSources() const;
 
 		const FAudioPlatformDeviceInfo& GetPlatformDeviceInfo() const { return PlatformInfo; };
@@ -195,8 +196,8 @@ namespace Audio
 		/** Map of USoundSubmix static data objects to the dynamic audio mixer submix. */
 		TMap<USoundSubmix*, FMixerSubmixPtr> Submixes;
 
-		/** List of mixer source voices. */
-		TArray<FMixerSourceVoice*> SourceVoices;
+		/** Queue of mixer source voices. */
+		TQueue<FMixerSourceVoice*> SourceVoices;
 
 		TMap<uint32, TArray<FSourceEffectChainEntry>> SourceEffectChainOverrides;
 
diff --git a/Engine/Source/Runtime/AudioMixer/Public/AudioMixerTypes.h b/Engine/Source/Runtime/AudioMixer/Public/AudioMixerTypes.h
--- a/Engine/Source/Runtime/AudioMixer/Public/AudioMixerTypes.h
+++ b/Engine/Source/Runtime/AudioMixer/Public/AudioMixerTypes.h
@@ -25,10 +25,7 @@ namespace Audio {
 		{
 			Unknown,
 			Float,
-			Double,
 			Int16,
-			Int24,
-			Int32,
 			Unsupported
 		};
 	}
diff --git a/Engine/Source/Runtime/AudioMixer/Public/Components/SynthComponent.h b/Engine/Source/Runtime/AudioMixer/Public/Components/SynthComponent.h
--- a/Engine/Source/Runtime/AudioMixer/Public/Components/SynthComponent.h
+++ b/Engine/Source/Runtime/AudioMixer/Public/Components/SynthComponent.h
@@ -7,6 +7,7 @@
 #include "UObject/ObjectMacros.h"
 #include "Engine/EngineTypes.h"
 #include "Sound/SoundWaveProcedural.h"
+#include "AudioMixerTypes.h"
 
 #include "SynthComponent.generated.h"
 
@@ -25,11 +26,15 @@ class AUDIOMIXER_API USynthSound : public USoundWaveProcedural
 
 	void Init(USynthComponent* INSynthComponent, const int32 InNumChannels);
 
+	/** Begin USoundWave */
 	virtual bool OnGeneratePCMAudio(TArray<uint8>& OutAudio, int32 NumSamples) override;
+	virtual Audio::EAudioMixerStreamDataFormat::Type GetGeneratedPCMDataFormat() const override;
+	/** End USoundWave */
 
 protected:
 	USynthComponent* OwningSynthComponent;
-	TArray<int16> GeneratedPCMData;
+	TArray<float> FloatBuffer;
+	bool bAudioMixer;
 };
 
 UCLASS(ClassGroup = Synth, hidecategories = (Object, ActorComponent, Physics, Rendering, Mobility, LOD))
@@ -141,10 +146,10 @@ protected:
 	virtual void OnStop() {}
 
 	// Called when more audio is needed to be generated
-	virtual void OnGenerateAudio(TArray<float>& OutAudio) PURE_VIRTUAL(USynthComponent::OnGenerateAudio,);
+	virtual void OnGenerateAudio(float* OutAudio, int32 NumSamples) PURE_VIRTUAL(USynthComponent::OnGenerateAudio,);
 
 	// Called by procedural sound wave
-	void OnGeneratePCMAudio(TArray<int16>& GeneratedPCMData);
+	void OnGeneratePCMAudio(float* GeneratedPCMData, int32 NumSamples);
 
 	// Gets the audio device associated with this synth component
 	FAudioDevice* GetAudioDevice() { return AudioComponent->GetAudioDevice(); }
@@ -162,8 +167,6 @@ private:
 
 	void PumpPendingMessages();
 
-	TArray<float> AudioFloatData;
-
 #if SYNTH_GENERATOR_TEST_TONE
 	Audio::FSineOsc TestSineLeft;
 	Audio::FSineOsc TestSineRight;
diff --git a/Engine/Source/Runtime/AudioMixer/Public/DSP/Granulator.h b/Engine/Source/Runtime/AudioMixer/Public/DSP/Granulator.h
--- a/Engine/Source/Runtime/AudioMixer/Public/DSP/Granulator.h
+++ b/Engine/Source/Runtime/AudioMixer/Public/DSP/Granulator.h
@@ -257,7 +257,7 @@ namespace Audio
 		float GetSampleDuration() const;
 
 		// Generate the next audio buffer
-		void Generate(TArray<float>& OutAudiobuffer, const int32 NumFrames);
+		void Generate(float* OutAudiobuffer, const int32 NumFrames);
 
 	protected:
 		// Spawns grains
diff --git a/Engine/Source/Runtime/AudioMixer/Public/DSP/ParamInterpolator.h b/Engine/Source/Runtime/AudioMixer/Public/DSP/ParamInterpolator.h
--- a/Engine/Source/Runtime/AudioMixer/Public/DSP/ParamInterpolator.h
+++ b/Engine/Source/Runtime/AudioMixer/Public/DSP/ParamInterpolator.h
@@ -6,36 +6,58 @@
 
 namespace Audio
 {
-
-	/**
-	* FParamInterpolator
-	* Simple parameter interpolator which interpolates to target values
-	* in the given the required number of update ticks.
-	*/
-	class FParamInterpolator
+	class FParam
 	{
 	public:
-
-		FParamInterpolator();
-		virtual ~FParamInterpolator();
-
-		/** Initialize the value. */
-		void InitValue(const float InValue);
-
-		/** Set the target value to interpolate to in the given number of update ticks. */
-		void SetValue(const float InValue, const uint32 InNumInterpolationTicks);
-
-		/** Returns the current value without updating the tick. */
-		float GetValue() const;
-
-		/** Functor operator to generate new output value. */
-		virtual float operator()();
-
-	protected:
-		float StartValue;
-		float EndValue;
-		float CurrentTick;
-		float NumTicksToNewValue;
+		FParam()
+			: CurrentValue(0.0f)
+			, StartingValue(0.0f)
+			, TargetValue(0.0f)
+			, DeltaValue(0.0f)
+		{}
+
+		// Set the parameter value to the given target value over the given interpolation frames.
+		FORCEINLINE void SetValue(const float InValue, const int32 InNumInterpFrames = 0)
+		{
+			TargetValue = InValue;
+			if (InNumInterpFrames == 0)
+			{
+				StartingValue = TargetValue;
+				CurrentValue = TargetValue;
+				DeltaValue = 0.0f;
+			}
+			else
+			{
+				DeltaValue = (InValue - CurrentValue) / InNumInterpFrames;
+				StartingValue = CurrentValue;
+			}
+		}
+
+		// Resets the delta value back to 0.0. To be called at beginning of callback render.
+		FORCEINLINE void Reset()
+		{
+			DeltaValue = 0.0f;
+			CurrentValue = TargetValue;
+		}
+
+		// Updates the parameter, assumes called in one of the frames.
+		FORCEINLINE float Update()
+		{
+			CurrentValue += DeltaValue;
+			return CurrentValue;
+		}
+
+		// Returns the current value, but does not update the value
+		FORCEINLINE float GetValue() const
+		{
+			return CurrentValue;
+		}
+
+	private:
+		float CurrentValue;
+		float StartingValue;
+		float TargetValue;
+		float DeltaValue;
 	};
 
 }
diff --git a/Engine/Source/Runtime/AudioMixer/Public/DSP/SampleBufferReader.h b/Engine/Source/Runtime/AudioMixer/Public/DSP/SampleBufferReader.h
--- a/Engine/Source/Runtime/AudioMixer/Public/DSP/SampleBufferReader.h
+++ b/Engine/Source/Runtime/AudioMixer/Public/DSP/SampleBufferReader.h
@@ -54,7 +54,7 @@ namespace Audio
 		float GetPlaybackProgress() const { return PlaybackProgress; }
 
 		// Generates the next block of audio. Returns true if it's no longer playing (reached end of the buffer and not set to wrap)
-		bool Generate(TArray<float>& OutAudioBuffer, const int32 NumFrames, const int32 OutChannels, const bool bWrap = false);
+		bool Generate(float* OutAudioBuffer, const int32 NumFrames, const int32 OutChannels, const bool bWrap = false);
 
 		// Whether or not the buffer reader has a buffer
 		bool HasBuffer() const { return BufferPtr != nullptr; }
diff --git a/Engine/Source/Runtime/Core/Public/UObject/AnimPhysObjectVersion.h b/Engine/Source/Runtime/Core/Public/UObject/AnimPhysObjectVersion.h
--- a/Engine/Source/Runtime/Core/Public/UObject/AnimPhysObjectVersion.h
+++ b/Engine/Source/Runtime/Core/Public/UObject/AnimPhysObjectVersion.h
@@ -29,6 +29,12 @@ struct CORE_API FAnimPhysObjectVersion
 		FixInvalidClothParticleMasses,
 		// Moved influence count to cached data
 		CacheClothMeshInfluences,
+		// Remove GUID from Smart Names entirely + remove automatic name fixup
+		SmartNameRefactorForDeterministicCooking,
+		// rename the variable and allow individual curves to be set
+		RenameDisableAnimCurvesToAllowAnimCurveEvaluation,
+		// link curve to LOD, so curve metadata has to include LODIndex
+		AddLODToCurveMetaData,
 
 		// -----<new versions can be added above this line>-------------------------------------------------
 		VersionPlusOne,
diff --git a/Engine/Source/Runtime/CoreUObject/Private/UObject/Obj.cpp b/Engine/Source/Runtime/CoreUObject/Private/UObject/Obj.cpp
--- a/Engine/Source/Runtime/CoreUObject/Private/UObject/Obj.cpp
+++ b/Engine/Source/Runtime/CoreUObject/Private/UObject/Obj.cpp
@@ -1089,7 +1089,7 @@ void UObject::PreSave(const class ITargetPlatform* TargetPlatform)
 
 bool UObject::CanModify() const
 {
-	return (!HasAnyFlags(RF_NeedInitialization) && !IsGarbageCollecting());
+	return (!HasAnyFlags(RF_NeedInitialization) && !IsGarbageCollecting() && !GExitPurge && !IsUnreachable());
 }
 
 bool UObject::Modify( bool bAlwaysMarkDirty/*=true*/ )
diff --git a/Engine/Source/Runtime/Engine/Classes/Animation/AnimInstance.h b/Engine/Source/Runtime/Engine/Classes/Animation/AnimInstance.h
--- a/Engine/Source/Runtime/Engine/Classes/Animation/AnimInstance.h
+++ b/Engine/Source/Runtime/Engine/Classes/Animation/AnimInstance.h
@@ -1050,7 +1050,7 @@ public:
 	/**
 	* Recalculate Required Curves based on Required Bones [RequiredBones]
 	*/
-	void RecalcRequiredCurves(bool bDisableAnimCurves);
+	void RecalcRequiredCurves(const FCurveEvaluationOption& CurveEvalOption);
 
 	// @todo document
 	inline USkeletalMeshComponent* GetSkelMeshComponent() const { return CastChecked<USkeletalMeshComponent>(GetOuter()); }
@@ -1122,6 +1122,9 @@ private:
 	 */
 	TArray<FQueuedRootMotionBlend> RootMotionBlendQueue;
 
+	// Root motion read from proxy (where it is calculated) and stored here to avoid potential stalls by calling GetProxyOnGameThread
+	FRootMotionMovementParams ExtractedRootMotion;
+
 private:
 	// update montage
 	void UpdateMontage(float DeltaSeconds);
diff --git a/Engine/Source/Runtime/Engine/Classes/Animation/AnimNode_StateMachine.h b/Engine/Source/Runtime/Engine/Classes/Animation/AnimNode_StateMachine.h
--- a/Engine/Source/Runtime/Engine/Classes/Animation/AnimNode_StateMachine.h
+++ b/Engine/Source/Runtime/Engine/Classes/Animation/AnimNode_StateMachine.h
@@ -130,6 +130,11 @@ public:
 	UPROPERTY(EditAnywhere, Category = Settings)
 	bool bSkipFirstUpdateTransition;
 
+	// Reinitialize the state machine if we have become relevant to the graph
+	// after not being ticked on the previous frame(s)
+	UPROPERTY(EditAnywhere, Category = Settings)
+	bool bReinitializeOnBecomingRelevant;
+
 public:
 
 	int32 GetCurrentState() const
@@ -190,6 +195,7 @@ public:
 	FAnimNode_StateMachine()
 		: MaxTransitionsPerFrame(3)
 		, bSkipFirstUpdateTransition(true)
+		, bReinitializeOnBecomingRelevant(true)
 		, PRIVATE_MachineDescription(NULL)
 		, CurrentState(INDEX_NONE)
 		, bFirstUpdate(true)
diff --git a/Engine/Source/Runtime/Engine/Classes/Animation/AnimSequenceBase.h b/Engine/Source/Runtime/Engine/Classes/Animation/AnimSequenceBase.h
--- a/Engine/Source/Runtime/Engine/Classes/Animation/AnimSequenceBase.h
+++ b/Engine/Source/Runtime/Engine/Classes/Animation/AnimSequenceBase.h
@@ -216,18 +216,11 @@ public:
 
 protected:
 	template <typename DataType>
-	void VerifyCurveNames(USkeleton* Skeleton, const FName& NameContainer, TArray<DataType>& CurveList)
+	void VerifyCurveNames(USkeleton& Skeleton, const FName& NameContainer, TArray<DataType>& CurveList)
 	{
-		USkeleton* MySkeleton = GetSkeleton();
-
-		if (MySkeleton)
+		for (DataType& Curve : CurveList)
 		{
-			// since this is verify function that makes sure it exists after loaded
-			// we should add it if it doesn't exist
-			for (DataType& Curve : CurveList)
-			{
-				MySkeleton->VerifySmartName(NameContainer, Curve.Name);
-			}
+			Skeleton.VerifySmartName(NameContainer, Curve.Name);
 		}
 	}
 };
diff --git a/Engine/Source/Runtime/Engine/Classes/Animation/AnimSingleNodeInstance.h b/Engine/Source/Runtime/Engine/Classes/Animation/AnimSingleNodeInstance.h
--- a/Engine/Source/Runtime/Engine/Classes/Animation/AnimSingleNodeInstance.h
+++ b/Engine/Source/Runtime/Engine/Classes/Animation/AnimSingleNodeInstance.h
@@ -68,6 +68,9 @@ public:
 	/** Set New Asset - calls InitializeAnimation, for now we need MeshComponent **/
 	UFUNCTION(BlueprintCallable, Category="Animation")
 	virtual void SetAnimationAsset(UAnimationAsset* NewAsset, bool bIsLooping=true, float InPlayRate=1.f);
+	/** Get the currently used asset */
+	UFUNCTION(BlueprintCallable, Category = "Animation")
+	virtual UAnimationAsset* GetAnimationAsset() const;
 	/** Set pose value */
  	UFUNCTION(BlueprintCallable, Category = "Animation")
  	void SetPreviewCurveOverride(const FName& PoseName, float Value, bool bRemoveIfZero);
diff --git a/Engine/Source/Runtime/Engine/Classes/Animation/PoseAsset.h b/Engine/Source/Runtime/Engine/Classes/Animation/PoseAsset.h
--- a/Engine/Source/Runtime/Engine/Classes/Animation/PoseAsset.h
+++ b/Engine/Source/Runtime/Engine/Classes/Animation/PoseAsset.h
@@ -204,10 +204,16 @@ public:
 	virtual bool GetAllAnimationSequencesReferred(TArray<UAnimationAsset*>& AnimationAssets, bool bRecursive = true) override;
 	virtual void ReplaceReferredAnimations(const TMap<UAnimationAsset*, UAnimationAsset*>& ReplacementMap) override;
 	// End AnimationAsset interface
-#endif
 
 	ENGINE_API bool ModifyPoseName(FName OldPoseName, FName NewPoseName, const SmartName::UID_Type* NewUID);
 
+	// Rename the smart names used by this Pose Asset
+	ENGINE_API void RenameSmartName(const FName& InOriginalName, const FName& InNewName);
+
+	// Remove poses or curves using the smart names supplied
+	ENGINE_API void RemoveSmartNames(const TArray<FName>& InNamesToRemove);
+#endif
+
 	ENGINE_API int32 DeletePoses(TArray<FName> PoseNamesToDelete);
 	ENGINE_API int32 DeleteCurves(TArray<FName> CurveNamesToDelete);
 	ENGINE_API bool ConvertSpace(bool bNewAdditivePose, int32 NewBasePoseInde);
diff --git a/Engine/Source/Runtime/Engine/Classes/Animation/Skeleton.h b/Engine/Source/Runtime/Engine/Classes/Animation/Skeleton.h
--- a/Engine/Source/Runtime/Engine/Classes/Animation/Skeleton.h
+++ b/Engine/Source/Runtime/Engine/Classes/Animation/Skeleton.h
@@ -440,7 +440,7 @@ public:
 	ENGINE_API void RemoveSmartnameAndModify(FName ContainerName, SmartName::UID_Type Uid);
 
 	// Removes smartnames from the specified container and modifies the skeleton
-	ENGINE_API void RemoveSmartnamesAndModify(FName ContainerName, const TArray<SmartName::UID_Type>& Uids);
+	ENGINE_API void RemoveSmartnamesAndModify(FName ContainerName, const TArray<FName>& Names);
 #endif// WITH_EDITOR
 
 	// quick wrapper function for Find UID by name, if not found, it will return SmartName::MaxUID
@@ -448,10 +448,6 @@ public:
 	ENGINE_API bool GetSmartNameByUID(const FName& ContainerName, SmartName::UID_Type UID, FSmartName& OutSmartName);
 	ENGINE_API bool GetSmartNameByName(const FName& ContainerName, const FName& InName, FSmartName& OutSmartName);
 
-	// Adds a new name to the smart name container and modifies the skeleton so it can be saved
-	// return bool - Whether a name was added (false if already present)
-	ENGINE_API bool RenameSmartName(FName ContainerName, const SmartName::UID_Type& Uid, FName NewName);
-
 	// Get or add a smartname container with the given name
 	ENGINE_API const FSmartNameMapping* GetSmartNameContainer(const FName& ContainerName) const;
 
@@ -815,7 +811,7 @@ public:
 	static const FString AnimNotifyTagDelimiter;
 
 	// Asset registry information for animation curves
-	ENGINE_API static const FName CurveTag;
+	ENGINE_API static const FName CurveNameTag;
 	ENGINE_API static const FString CurveTagDelimiter;
 
 	// rig Configs
