[NO NON-PII COMMIT MESSAGE]

diff --git a/Engine/Source/Runtime/AudioMixer/Private/Components/SynthComponent.cpp b/Engine/Source/Runtime/AudioMixer/Private/Components/SynthComponent.cpp
--- a/Engine/Source/Runtime/AudioMixer/Private/Components/SynthComponent.cpp
+++ b/Engine/Source/Runtime/AudioMixer/Private/Components/SynthComponent.cpp
@@ -102,6 +102,8 @@ USynthComponent::USynthComponent(const FObjectInitializer& ObjectInitializer)
 	bIsInitialized = false;
 	bIsUISound = false;
 
+	Synth = nullptr;
+
 	// Set the default sound class
 	SoundClass = USoundBase::DefaultSoundClassObject;
 
@@ -341,6 +343,15 @@ void USynthComponent::Start()
 	// This will try to create the audio component if it hasn't yet been created
 	CreateAudioComponent();
 
+	// We will also ensure that this synth was initialized before attempting to play.
+	Initialize();
+
+	if (Synth == nullptr)
+	{
+		UE_LOG(LogAudio, Warning, TEXT("Warning: SynthComponent failed to start due to failiure in initialization."));
+		return;
+	}
+
 	if (AudioComponent)
 	{
 		// Copy the attenuation and concurrency data from the synth component to the audio component
