Merging CL# 4091120 To UE4/Main
#rnx

diff --git a/Engine/Source/Runtime/AudioMixer/Private/AudioMixerSource.cpp b/Engine/Source/Runtime/AudioMixer/Private/AudioMixerSource.cpp
--- a/Engine/Source/Runtime/AudioMixer/Private/AudioMixerSource.cpp
+++ b/Engine/Source/Runtime/AudioMixer/Private/AudioMixerSource.cpp
@@ -16,6 +16,14 @@ FAutoConsoleVariableRef CVarDisableHRTF(
 	TEXT("0: Not Disabled, 1: Disabled"),
 	ECVF_Default);
 
+static int32 DisableStoppingVoicesCvar = 0;
+FAutoConsoleVariableRef CVarDisableStoppingVoices(
+	TEXT("au.DisableStoppingVoices"),
+	DisableStoppingVoicesCvar,
+	TEXT("Forces all sounds to not fade out and instead stop immediately.\n")
+	TEXT("0: Not Disabled, 1: Disabled"),
+	ECVF_Default);
+
 namespace Audio
 {
 	FMixerSource::FMixerSource(FAudioDevice* InAudioDevice)
@@ -455,7 +463,7 @@ namespace Audio
 			const bool bIsProcedural = WaveInstance->WaveData && WaveInstance->WaveData->DecompressionType == DTYPE_Procedural;
 
 			// If we've finished naturally or we're paused or we're told to immediately stop
-			if (bIsFinished || bIsProcedural || bStopNow || (MixerSourceVoice && MixerSourceVoice->IsPaused()))
+			if (bIsFinished || bIsProcedural || bStopNow || DisableStoppingVoicesCvar|| (MixerSourceVoice && MixerSourceVoice->IsPaused()))
 			{
 				FScopeLock Lock(&RenderThreadCritSect);
 
