[Release-4.20] #jira UE-58004 fix AudioMixer checks in AudioMixerBlueprintLibary, which solves crash when calling recording BP functions. #rb Aaron.McLeran

diff --git a/Engine/Source/Runtime/AudioMixer/Private/AudioMixerBlueprintLibrary.cpp b/Engine/Source/Runtime/AudioMixer/Private/AudioMixerBlueprintLibrary.cpp
--- a/Engine/Source/Runtime/AudioMixer/Private/AudioMixerBlueprintLibrary.cpp
+++ b/Engine/Source/Runtime/AudioMixer/Private/AudioMixerBlueprintLibrary.cpp
@@ -26,7 +26,14 @@ static Audio::FMixerDevice* GetAudioMixerDeviceFromWorldContext(const UObject* W
 {
 	if (FAudioDevice* AudioDevice = GetAudioDeviceFromWorldContext(WorldContextObject))
 	{
-		return static_cast<Audio::FMixerDevice*>(AudioDevice);
+		if (!AudioDevice->IsAudioMixerEnabled())
+		{
+			return nullptr;
+		}
+		else
+		{
+			return static_cast<Audio::FMixerDevice*>(AudioDevice);
+		}
 	}
 	return nullptr;
 }
@@ -92,6 +99,10 @@ void UAudioMixerBlueprintLibrary::StartRecordingOutput(const UObject* WorldConte
 	{
 		MixerDevice->StartRecording(SubmixToRecord, ExpectedDuration);
 	}
+	else
+	{
+		UE_LOG(LogAudioMixer, Error, TEXT("Output recording is an audio mixer only feature. Please run the game with -audiomixer to enable this feature."));
+	}
 }
 
 USoundWave* UAudioMixerBlueprintLibrary::StopRecordingOutput(const UObject* WorldContextObject, EAudioRecordingExportType ExportType, const FString& Name, FString Path, USoundSubmix* SubmixToRecord, USoundWave* ExistingSoundWaveToOverwrite)
@@ -109,6 +120,12 @@ USoundWave* UAudioMixerBlueprintLibrary::StopRecordingOutput(const UObject* Worl
 		// call the thing here.
 		Audio::AlignedFloatBuffer& RecordedBuffer = MixerDevice->StopRecording(SubmixToRecord, ChannelCount, SampleRate);
 
+		if (RecordedBuffer.Num() == 0)
+		{
+			UE_LOG(LogAudioMixer, Warning, TEXT("No audio data. Did you call Start Recording Output?"));
+			return nullptr;
+		}
+
 		// Pack output data into a TSampleBuffer and record out:
 		RecordingData.Reset(new Audio::FAudioRecordingData());
 		RecordingData->InputBuffer = Audio::TSampleBuffer<int16>(RecordedBuffer, ChannelCount, SampleRate);
@@ -140,6 +157,10 @@ USoundWave* UAudioMixerBlueprintLibrary::StopRecordingOutput(const UObject* Worl
 			break;
 		}	
 	}
+	else
+	{
+		UE_LOG(LogAudioMixer, Error, TEXT("Output recording is an audio mixer only feature. Please run the game with -audiomixer to enable this feature."));
+	}
 
 	return nullptr;
 }
