[Release-4.20] #jira UE-60471 Bring Aaron's scrubbing fixes over for Synth Sample Player. #rb none #fyi aaron.mcleran dan.reynolds

diff --git a/Engine/Plugins/Runtime/Synthesis/Source/Synthesis/Classes/SynthComponents/SynthComponentWaveTable.h b/Engine/Plugins/Runtime/Synthesis/Source/Synthesis/Classes/SynthComponents/SynthComponentWaveTable.h
--- a/Engine/Plugins/Runtime/Synthesis/Source/Synthesis/Classes/SynthComponents/SynthComponentWaveTable.h
+++ b/Engine/Plugins/Runtime/Synthesis/Source/Synthesis/Classes/SynthComponents/SynthComponentWaveTable.h
@@ -65,7 +65,7 @@ public:
 	void SetPitch(float InPitch, float TimeSec);
 
 	UFUNCTION(BlueprintCallable, Category = "Synth|Components|Audio")
-	void SeekToTime(float TimeSec, ESamplePlayerSeekType SeekType);
+	void SeekToTime(float TimeSec, ESamplePlayerSeekType SeekType, bool bWrap = true);
 
 	UFUNCTION(BlueprintCallable, Category = "Synth|Components|Audio")
 	void SetScrubMode(bool bScrubMode);
diff --git a/Engine/Plugins/Runtime/Synthesis/Source/Synthesis/Private/SynthComponentWaveTable.cpp b/Engine/Plugins/Runtime/Synthesis/Source/Synthesis/Private/SynthComponentWaveTable.cpp
--- a/Engine/Plugins/Runtime/Synthesis/Source/Synthesis/Private/SynthComponentWaveTable.cpp
+++ b/Engine/Plugins/Runtime/Synthesis/Source/Synthesis/Private/SynthComponentWaveTable.cpp
@@ -36,7 +36,7 @@ void USynthSamplePlayer::SetPitch(float InPitch, float InTimeSec)
 	});
 }
 
-void USynthSamplePlayer::SeekToTime(float InTimeSecs, ESamplePlayerSeekType InSeekType)
+void USynthSamplePlayer::SeekToTime(float InTimeSecs, ESamplePlayerSeekType InSeekType, bool bWrap)
 {
 	Audio::ESeekType::Type SeekType;
 	switch (InSeekType)
@@ -55,9 +55,9 @@ void USynthSamplePlayer::SeekToTime(float InTimeSecs, ESamplePlayerSeekType InSe
 			break;
 	}
 
-	SynthCommand([this, InTimeSecs, SeekType]()
+	SynthCommand([this, InTimeSecs, SeekType, bWrap]()
 	{
-		SampleBufferReader.SeekTime(InTimeSecs, SeekType);
+		SampleBufferReader.SeekTime(InTimeSecs, SeekType, bWrap);
 	});
 }
 
diff --git a/Engine/Source/Runtime/AudioMixer/Private/DSP/SampleBufferReader.cpp b/Engine/Source/Runtime/AudioMixer/Private/DSP/SampleBufferReader.cpp
--- a/Engine/Source/Runtime/AudioMixer/Private/DSP/SampleBufferReader.cpp
+++ b/Engine/Source/Runtime/AudioMixer/Private/DSP/SampleBufferReader.cpp
@@ -12,6 +12,9 @@ namespace Audio
 		, BufferNumFrames(0)
 		, BufferSampleRate(0)
 		, BufferNumChannels(0)
+		, FadeFrames(512)
+		, FadeValue(0.0f)
+		, FadeIncrement(1.0f / (float)FadeFrames)
 		, DeviceSampleRate(0.0f)
 		, BasePitch(1.0f)
 		, PitchScale(1.0f)
@@ -101,7 +104,7 @@ namespace Audio
 		if (BufferPtr)
 		{
 			check(BufferNumChannels > 0);
-			const float CurrentSeekFrame = ((float)BufferSampleRate * CurrentSeekTime) / BufferNumChannels;
+			const float CurrentSeekFrame = ((float)BufferSampleRate * CurrentSeekTime);
 
 			if (CurrentSeekType == ESeekType::FromBeginning)
 			{
@@ -118,7 +121,7 @@ namespace Audio
 
 			if (bWrap)
 			{
-				while (CurrentBufferFrameIndexInterpolated > (double)BufferNumFrames)
+				while (CurrentBufferFrameIndexInterpolated >= (double)BufferNumFrames)
 				{
 					CurrentBufferFrameIndexInterpolated -= (double)BufferNumFrames;
 				}
@@ -145,7 +148,15 @@ namespace Audio
 		CurrentSeekType = InSeekType;
 		bWrap = bInWrap;
 
-		UpdateSeekFrame();
+		if (bIsScrubMode)
+		{
+			UpdateSeekFrame();
+			UpdateScrubMinAndMax();
+		}
+		else
+		{
+			UpdateSeekFrame();
+		}
 	}
 
 	void FSampleBufferReader::SetScrubTimeWidth(const float InScrubTimeWidthSec)
@@ -167,6 +178,7 @@ namespace Audio
 		
 		// Anchor the current frame index as the scrub anchor
 		ScrubAnchorFrame = CurrentBufferFrameIndexInterpolated;
+		UpdateSeekFrame();
 		UpdateScrubMinAndMax();
 	}
 
@@ -179,6 +191,9 @@ namespace Audio
 			ScrubWidthFrames = (double)(DeviceSampleRate * FMath::Max(CurrentScrubWidthSec, 0.001f));
 			ScrubWidthFrames = FMath::Min((double)(BufferNumFrames - 1), ScrubWidthFrames);
 
+			// Don't allow the scrub width to be less than 2 times the scrubwidth frames
+			ScrubWidthFrames = FMath::Max(ScrubWidthFrames, (double)2*FadeFrames);
+
 			ScrubMinFrame = ScrubAnchorFrame - 0.5 * ScrubWidthFrames;
 			ScrubMaxFrame = ScrubAnchorFrame + 0.5 * ScrubWidthFrames;
 		}
@@ -237,7 +252,7 @@ namespace Audio
 				CurrentFrameIndex = FMath::FloorToInt(CurrentBufferFrameIndexInterpolated);
 				NextFrameIndex = CurrentFrameIndex + 1;
 				AlphaLerp = CurrentBufferFrameIndexInterpolated - (double)CurrentFrameIndex;
-				if (!bIsScrubMode)
+				if (!bIsScrubMode && !bWrap)
 				{
 					if (NextFrameIndex >= BufferNumFrames)
 					{
@@ -250,7 +265,7 @@ namespace Audio
 				CurrentFrameIndex = FMath::CeilToInt(CurrentBufferFrameIndexInterpolated);
 				NextFrameIndex = CurrentFrameIndex - 1;
 				AlphaLerp = (double)CurrentFrameIndex - CurrentBufferFrameIndexInterpolated;
-				if (!bIsScrubMode)
+				if (!bIsScrubMode && !bWrap)
 				{
 					if (NextFrameIndex < 0)
 					{
@@ -262,38 +277,61 @@ namespace Audio
 			if (!bIsFinished)
 			{
 				// Check for scrub boundaries and wrap. Note that we've already wrapped on the buffer boundary at this point.
-				if (bIsScrubMode)
+				if (bWrap || bIsScrubMode)
 				{
-					if (CurrentPitch > 0.0f && NextFrameIndex >= ScrubMaxFrame)
+					int32 MinWrapFrame = 0;
+					int32 MaxWrapFrame = BufferNumFrames;
+					if (bIsScrubMode)
+					{
+						MinWrapFrame = ScrubMinFrame;
+						MaxWrapFrame = ScrubMaxFrame;
+					}
+
+					if (CurrentPitch > 0.0f && NextFrameIndex >= MaxWrapFrame)
 					{
-						NextFrameIndex = ScrubMinFrame;
-						CurrentFrameIndex = (int32)(ScrubMaxFrame - 1.0f);
+						NextFrameIndex = MinWrapFrame;
+						CurrentFrameIndex = (int32)(MaxWrapFrame - 1.0f);
 						CurrentBufferFrameIndexInterpolated = FMath::Fmod(CurrentBufferFrameIndexInterpolated, 1.0) + (double)(NextFrameIndex);
 					}
-					else if (NextFrameIndex < ScrubMinFrame)
+					else if (NextFrameIndex < MinWrapFrame)
 					{
-						NextFrameIndex = (int32)(ScrubMaxFrame - 1);
-						CurrentFrameIndex = (int32)ScrubMinFrame;
+						NextFrameIndex = (int32)(MaxWrapFrame - 1);
+						CurrentFrameIndex = (int32)MinWrapFrame;
 						CurrentBufferFrameIndexInterpolated = FMath::Fmod(CurrentBufferFrameIndexInterpolated, 1.0) + (double)NextFrameIndex;
 					}
+
+					CurrentFrameIndex = FMath::CeilToInt(CurrentBufferFrameIndexInterpolated);
+					NextFrameIndex = CurrentFrameIndex - 1;
+					AlphaLerp = (double)CurrentFrameIndex - CurrentBufferFrameIndexInterpolated;
+
+					FadeValue = 1.0f;
+					int32 MaxFadeInFrame = MinWrapFrame + FadeFrames;
+					if (CurrentFrameIndex >= MinWrapFrame && CurrentFrameIndex < MaxFadeInFrame)
+					{
+						FadeValue = (float)(CurrentFrameIndex - MinWrapFrame) / FadeFrames;
+					}
+					else if (CurrentFrameIndex >= MaxWrapFrame - FadeFrames && CurrentFrameIndex < MaxWrapFrame)
+					{
+						FadeValue = 1.0f - (float)(CurrentFrameIndex - (MaxWrapFrame - FadeFrames)) / FadeFrames;
+					}
 				}
 
 				if (OutChannels == BufferNumChannels)
 				{
 					for (int32 Channel = 0; Channel < BufferNumChannels; ++Channel)
 					{
-						OutAudioBuffer[OutSampleIndex++] = GetSampleValueForChannel(Channel);
+						OutAudioBuffer[OutSampleIndex++] = FadeValue * GetSampleValueForChannel(Channel);
 					}
 				}
 				else if (OutChannels == 1 && BufferNumChannels == 2)
 				{
-					float LeftChannel = GetSampleValueForChannel(0);
-					float RightChannel = GetSampleValueForChannel(1);
+					float LeftChannel = FadeValue * GetSampleValueForChannel(0);
+					float RightChannel = FadeValue * GetSampleValueForChannel(1);
 					OutAudioBuffer[OutSampleIndex++] = 0.5f * (LeftChannel + RightChannel);
 				}
 				else if (OutChannels == 2 && BufferNumChannels == 1)
 				{
-					float Sample = GetSampleValueForChannel(0);
+					float Sample = FadeValue * GetSampleValueForChannel(0);
 					OutAudioBuffer[OutSampleIndex++] = 0.5f * Sample;
 					OutAudioBuffer[OutSampleIndex++] = 0.5f * Sample;
 				}
@@ -305,27 +343,24 @@ namespace Audio
 		return bIsFinished;
 	}
 
-	float FSampleBufferReader::GetSampleValueForChannel(const int32 Channel)
+	static int32 WrapIndex(int32 Value, int32 Max)
 	{
-		int32 WrappedCurrentFrameIndex = CurrentFrameIndex;
-		if (CurrentFrameIndex < 0)
+		if (Value < 0)
 		{
-			WrappedCurrentFrameIndex = CurrentFrameIndex + BufferNumFrames;
+			Value += Max;
 		}
-		else if (CurrentFrameIndex >= BufferNumFrames)
+		else if (Value >= Max)
 		{
-			WrappedCurrentFrameIndex = CurrentFrameIndex - BufferNumFrames;
+			Value -= Max;
 		}
+		return Value;
+	}
 
-		int32 WrappedNextFrameIndex = NextFrameIndex;
-		if (NextFrameIndex < 0)
-		{
-			WrappedNextFrameIndex = NextFrameIndex + BufferNumFrames;
-		}
-		else if (NextFrameIndex >= BufferNumFrames)
-		{
-			WrappedCurrentFrameIndex = NextFrameIndex - BufferNumFrames;
-		}
+	float FSampleBufferReader::GetSampleValueForChannel(const int32 Channel)
+	{
+		// Wrap the current frame index
+		int32 WrappedCurrentFrameIndex = WrapIndex(CurrentFrameIndex, BufferNumFrames);
+		int32 WrappedNextFrameIndex = WrapIndex(NextFrameIndex, BufferNumFrames);
 
 		// Update the current playback time
 		PlaybackProgress = (float)(WrappedCurrentFrameIndex / BufferSampleRate);
diff --git a/Engine/Source/Runtime/AudioMixer/Public/DSP/SampleBufferReader.h b/Engine/Source/Runtime/AudioMixer/Public/DSP/SampleBufferReader.h
--- a/Engine/Source/Runtime/AudioMixer/Public/DSP/SampleBufferReader.h
+++ b/Engine/Source/Runtime/AudioMixer/Public/DSP/SampleBufferReader.h
@@ -37,7 +37,7 @@ namespace Audio
 
 		// Puts the wave reader into scrub mode
 		void SetScrubMode(const bool bInIsScrubMode);
-
+		
 		// Sets the scrub width. The sound will loop between the scrub width region and the current frame
 		void SetScrubTimeWidth(const float InScrubTimeWidthSec);
 
@@ -71,6 +71,9 @@ namespace Audio
 		int32 BufferNumFrames;
 		int32 BufferSampleRate;
 		int32 BufferNumChannels;
+		int32 FadeFrames;
+		float FadeValue;
+		float FadeIncrement;
 
 		float DeviceSampleRate;
 
