Fix prevents sample buffer reader from accessing garbage memory.

diff --git a/Engine/Plugins/Runtime/Synthesis/Source/Synthesis/Private/SynthComponentWaveTable.cpp b/Engine/Plugins/Runtime/Synthesis/Source/Synthesis/Private/SynthComponentWaveTable.cpp
--- a/Engine/Plugins/Runtime/Synthesis/Source/Synthesis/Private/SynthComponentWaveTable.cpp
+++ b/Engine/Plugins/Runtime/Synthesis/Source/Synthesis/Private/SynthComponentWaveTable.cpp
@@ -142,6 +142,7 @@ void USynthSamplePlayer::TickComponent(float DeltaTime, enum ELevelTick TickType
 		SynthCommand([this, NewSampleBuffer]()
 		{
 			SampleBuffer = NewSampleBuffer;
+			SampleBufferReader.ClearBuffer();
 
 			// Clear the pending sound waves queue since we've now loaded a new buffer of data
 			SoundWaveLoader.Reset();
@@ -159,7 +160,7 @@ int32 USynthSamplePlayer::OnGenerateAudio(float* OutAudio, int32 NumSamples)
 		const int32 BufferNumSamples = SampleBuffer.GetNumSamples();
 		const int32 BufferNumChannels = SampleBuffer.GetNumChannels();
 		const int32 BufferSampleRate = SampleBuffer.GetSampleRate();
-		SampleBufferReader.SetBuffer(&BufferData, BufferNumSamples, BufferNumChannels, BufferSampleRate);
+		SampleBufferReader.SetBuffer(BufferData, BufferNumSamples, BufferNumChannels, BufferSampleRate);
 		SampleDurationSec = BufferNumSamples / (BufferSampleRate * BufferNumChannels);
 	}
 
diff --git a/Engine/Source/Runtime/AudioMixer/Private/DSP/Granulator.cpp b/Engine/Source/Runtime/AudioMixer/Private/DSP/Granulator.cpp
--- a/Engine/Source/Runtime/AudioMixer/Private/DSP/Granulator.cpp
+++ b/Engine/Source/Runtime/AudioMixer/Private/DSP/Granulator.cpp
@@ -262,7 +262,7 @@ namespace Audio
 
 			SampleBufferReader.ClearBuffer();
 
-			SampleBufferReader.SetBuffer(&Buffer, NumBufferSamples, BufferChannels, BufferSampleRate);
+			SampleBufferReader.SetBuffer(Buffer, NumBufferSamples, BufferChannels, BufferSampleRate);
 
 			// Setup the sample buffer reader
 			SampleBufferReader.SetPitch(CurrentPitch);
diff --git a/Engine/Source/Runtime/AudioMixer/Private/DSP/SampleBufferReader.cpp b/Engine/Source/Runtime/AudioMixer/Private/DSP/SampleBufferReader.cpp
--- a/Engine/Source/Runtime/AudioMixer/Private/DSP/SampleBufferReader.cpp
+++ b/Engine/Source/Runtime/AudioMixer/Private/DSP/SampleBufferReader.cpp
@@ -71,9 +71,12 @@ namespace Audio
 		PlaybackProgress = 0.0f;
 	}
 
-	void FSampleBufferReader::SetBuffer(const int16** InBufferPtr, const int32 InNumBufferSamples, const int32 InNumChannels, const int32 InBufferSampleRate)
-	{	
-		BufferPtr = *InBufferPtr;
+	void FSampleBufferReader::SetBuffer(const int16* InBufferPtr, const int32 InNumBufferSamples, const int32 InNumChannels, const int32 InBufferSampleRate)
+	{
+		// Re-init on setting a new buffer
+		Init(InBufferSampleRate);
+
+		BufferPtr = InBufferPtr;
 		BufferNumSamples = InNumBufferSamples;
 		BufferNumChannels = InNumChannels;
 		BufferSampleRate = InBufferSampleRate;
@@ -315,6 +318,10 @@ namespace Audio
 						FadeValue = 1.0f - (float)(CurrentFrameIndex - (MaxWrapFrame - FadeFrames)) / FadeFrames;
 					}
 				}
+				else
+				{
+					FadeValue = 1.0f;
+				}
 
 				if (OutChannels == BufferNumChannels)
 				{
@@ -358,19 +365,23 @@ namespace Audio
 
 	float FSampleBufferReader::GetSampleValueForChannel(const int32 Channel)
 	{
-		// Wrap the current frame index
-		int32 WrappedCurrentFrameIndex = WrapIndex(CurrentFrameIndex, BufferNumFrames);
-		int32 WrappedNextFrameIndex = WrapIndex(NextFrameIndex, BufferNumFrames);
+		if (BufferPtr)
+		{
+			// Wrap the current frame index
+			int32 WrappedCurrentFrameIndex = WrapIndex(CurrentFrameIndex, BufferNumFrames);
+			int32 WrappedNextFrameIndex = WrapIndex(NextFrameIndex, BufferNumFrames);
 
-		// Update the current playback time
-		PlaybackProgress = (float)(WrappedCurrentFrameIndex / BufferSampleRate);
+			// Update the current playback time
+			PlaybackProgress = (float)(WrappedCurrentFrameIndex / BufferSampleRate);
 
-		const int32 CurrentBufferSampleIndex = BufferNumChannels * WrappedCurrentFrameIndex + Channel;
-		const int32 NextBufferSampleIndex = BufferNumChannels * WrappedNextFrameIndex + Channel;
+			const int32 CurrentBufferSampleIndex = BufferNumChannels * WrappedCurrentFrameIndex + Channel;
+			const int32 NextBufferSampleIndex = BufferNumChannels * WrappedNextFrameIndex + Channel;
 
-		const float CurrentSampleValue = GetSampleValue(BufferPtr, CurrentBufferSampleIndex);
-		const float NextSampleValue = GetSampleValue(BufferPtr, NextBufferSampleIndex);
-		return FMath::Lerp(CurrentSampleValue, NextSampleValue, AlphaLerp);
+			const float CurrentSampleValue = GetSampleValue(BufferPtr, CurrentBufferSampleIndex);
+			const float NextSampleValue = GetSampleValue(BufferPtr, NextBufferSampleIndex);
+			return FMath::Lerp(CurrentSampleValue, NextSampleValue, AlphaLerp);
+		}
+		return 0.0f;
 	}
 
 }
diff --git a/Engine/Source/Runtime/AudioMixer/Public/DSP/SampleBufferReader.h b/Engine/Source/Runtime/AudioMixer/Public/DSP/SampleBufferReader.h
--- a/Engine/Source/Runtime/AudioMixer/Public/DSP/SampleBufferReader.h
+++ b/Engine/Source/Runtime/AudioMixer/Public/DSP/SampleBufferReader.h
@@ -27,7 +27,7 @@ namespace Audio
 		void Init(const int32 InSampleRate);
 
 		// This must be a completely loaded buffer. This buffer reader doesn't OWN the buffer memory.
-		void SetBuffer(const int16** InBufferPtr, const int32 InNumBufferSamples, const int32 InNumChannels, const int32 InBufferSampleRate);
+		void SetBuffer(const int16* InBufferPtr, const int32 InNumBufferSamples, const int32 InNumChannels, const int32 InBufferSampleRate);
 
 		// Seeks the buffer the given time in seconds. Returns true if succeeded.
 		void SeekTime(const float InTimeSec, const ESeekType::Type InSeekType = ESeekType::FromBeginning, const bool bWrap = true);
diff --git a/Engine/Source/Runtime/Engine/Private/AudioDevice.cpp b/Engine/Source/Runtime/Engine/Private/AudioDevice.cpp
--- a/Engine/Source/Runtime/Engine/Private/AudioDevice.cpp
+++ b/Engine/Source/Runtime/Engine/Private/AudioDevice.cpp
@@ -4586,11 +4586,16 @@ void FAudioDevice::Precache(USoundWave* SoundWave, bool bSynchronous, bool bTrac
 
 
 	// We're already precaching this sound wave so no need to precache again
-	if (SoundWave->DecompressionType != DTYPE_Setup)
+	if (SoundWave->DecompressionType != DTYPE_Setup && !bForceFullDecompression)
 	{
 		return; 
 	}
 
+	if (bForceFullDecompression)
+	{
+		SoundWave->SetPrecacheState(ESoundWavePrecacheState::NotStarted);
+	}
+
 	if (!bSynchronous && SoundWave->GetPrecacheState() == ESoundWavePrecacheState::NotStarted)
 	{
 		if (IsInGameThread())
