[Release 4.20] Merge updates for Stopping Voices.

diff --git a/Engine/Source/Runtime/AudioMixer/Private/AudioMixer.cpp b/Engine/Source/Runtime/AudioMixer/Private/AudioMixer.cpp
--- a/Engine/Source/Runtime/AudioMixer/Private/AudioMixer.cpp
+++ b/Engine/Source/Runtime/AudioMixer/Private/AudioMixer.cpp
@@ -15,6 +15,15 @@ FAutoConsoleVariableRef CVarLogRenderTimes(
 	TEXT("0: Not Log, 1: Log"),
 	ECVF_Default);
 
+// Command for setting the audio render thread priority.
+static int32 SetRenderThreadPriorityCVar = (int32)TPri_Highest;
+FAutoConsoleVariableRef CVarLogRenderThreadPriority(
+	TEXT("au.RenderThreadPriority"),
+	LogRenderTimesCVar,
+	TEXT("Sets audio render thread priority. Defaults to 3.\n")
+	TEXT("0: Normal, 1: Above Normal, 2: Below Normal, 3: Highest, 4: Lowest, 5: Slightly Below Normal, 6: Time Critical"),
+	ECVF_Default);
+
 DEFINE_STAT(STAT_AudioMixerRenderAudio);
 DEFINE_STAT(STAT_AudioMixerSourceManagerUpdate);
 DEFINE_STAT(STAT_AudioMixerSourceBuffers);
@@ -190,6 +199,8 @@ namespace Audio
 			}
 			break;
 		}
+
+		bIsReady = false;
 	}
 
 	/**
@@ -380,7 +391,7 @@ namespace Audio
 		check(AudioFadeEvent != nullptr);
 
 		check(AudioRenderThread == nullptr);
-		AudioRenderThread = FRunnableThread::Create(this, *FString::Printf(TEXT("AudioMixerRenderThread(%d)"), AudioMixerTaskCounter.Increment()), 0, TPri_Highest, FPlatformAffinity::GetAudioThreadMask());
+		AudioRenderThread = FRunnableThread::Create(this, *FString::Printf(TEXT("AudioMixerRenderThread(%d)"), AudioMixerTaskCounter.Increment()), 0, (EThreadPriority)SetRenderThreadPriorityCVar, FPlatformAffinity::GetAudioThreadMask());
 		check(AudioRenderThread != nullptr);
 	}
 
diff --git a/Engine/Source/Runtime/AudioMixer/Private/AudioMixerBuffer.cpp b/Engine/Source/Runtime/AudioMixer/Private/AudioMixerBuffer.cpp
--- a/Engine/Source/Runtime/AudioMixer/Private/AudioMixerBuffer.cpp
+++ b/Engine/Source/Runtime/AudioMixer/Private/AudioMixerBuffer.cpp
@@ -11,7 +11,6 @@ namespace Audio
 		: FSoundBuffer(InAudioDevice)
 		, RealtimeAsyncHeaderParseTask(nullptr)
 		, DecompressionState(nullptr)
-		, SoundWaveProcedural(nullptr)
 		, BufferType(InBufferType)
 		, SampleRate(InWave->GetSampleRateForCurrentPlatform())
 		, BitsPerSample(16) // TODO: support more bits, currently hard-coded to 16
@@ -73,13 +72,6 @@ namespace Audio
 			// nothing
 			break;
 		}
-
-		// Mark the procedural sound wave as being ok to be destroyed now
-		if (SoundWaveProcedural)
-		{
-			SoundWaveProcedural->OnEndGenerate();
-			SoundWaveProcedural->bIsReadyForDestroy = true;
-		}
 	}
 
 	int32 FMixerBuffer::GetSize()
@@ -307,13 +299,6 @@ namespace Audio
 		Buffer->ResourceID = 0;
 		InWave->ResourceID = 0;
 
-		// Don't allow the procedural sound wave to be destroyed until we're done with it
-		Buffer->SoundWaveProcedural = Cast<USoundWaveProcedural>(InWave);
-		if (Buffer->SoundWaveProcedural)
-		{
-			Buffer->SoundWaveProcedural->bIsReadyForDestroy = false;
-		}
-
 		return Buffer;
 	}
 
@@ -416,14 +401,4 @@ namespace Audio
 			RealtimeAsyncHeaderParseTask = nullptr;
 		}
 	}
-
-	void FMixerBuffer::OnBeginGenerate()
-	{
-		if (SoundWaveProcedural)
-		{
-			SoundWaveProcedural->OnBeginGenerate();
-		}
-	}
-
-
 }
\ No newline at end of file
diff --git a/Engine/Source/Runtime/AudioMixer/Private/AudioMixerBuffer.h b/Engine/Source/Runtime/AudioMixer/Private/AudioMixerBuffer.h
--- a/Engine/Source/Runtime/AudioMixer/Private/AudioMixerBuffer.h
+++ b/Engine/Source/Runtime/AudioMixer/Private/AudioMixerBuffer.h
@@ -62,9 +62,6 @@ namespace Audio
 		int32 GetNumChannels() const { return NumChannels; }
 		void InitSampleRate(const float InSampleRate) { SampleRate = InSampleRate; }
 
-		// Called on a buffer before it starts generating audio. Only used for the case of a procedural sound wave.
-		void OnBeginGenerate();
-
 	private:
 
 		/** Async task for parsing real-time decompressed compressed info headers. */
@@ -73,9 +70,6 @@ namespace Audio
 		/** Wrapper to handle the decompression of audio codecs. */
 		ICompressedAudioInfo* DecompressionState;
 
-		/** The ptr to this buffers sound wave procedural if it was created from a procedural sound wave. */
-		USoundWaveProcedural* SoundWaveProcedural;
-
 		/** Format of the sound referenced by this buffer */
 		EBufferType::Type BufferType;
 
diff --git a/Engine/Source/Runtime/AudioMixer/Private/AudioMixerChannelMaps.cpp b/Engine/Source/Runtime/AudioMixer/Private/AudioMixerChannelMaps.cpp
--- a/Engine/Source/Runtime/AudioMixer/Private/AudioMixerChannelMaps.cpp
+++ b/Engine/Source/Runtime/AudioMixer/Private/AudioMixerChannelMaps.cpp
@@ -202,7 +202,7 @@ namespace Audio
 		return Index;
 	}
 
-	void FMixerDevice::Get2DChannelMap(bool bIsVorbis, const ESubmixChannelFormat InSubmixChannelType, const int32 NumSourceChannels, const bool bIsCenterChannelOnly, TArray<float>& OutChannelMap) const
+	void FMixerDevice::Get2DChannelMap(bool bIsVorbis, const ESubmixChannelFormat InSubmixChannelType, const int32 NumSourceChannels, const bool bIsCenterChannelOnly, Audio::AlignedFloatBuffer& OutChannelMap) const
 	{
 		int32 NumOutputChannels = GetNumChannelsForSubmixFormat(InSubmixChannelType);
 		if (NumSourceChannels > 8 || NumOutputChannels > 8)
@@ -263,8 +263,23 @@ namespace Audio
 				else
 				{
 					// Mapping out to more than 2 channels, mono sources should be equally spread to left and right
-					OutChannelMap.Add(0.707f);
-					OutChannelMap.Add(0.707f);
+					switch (MonoChannelUpmixMethod)
+					{
+						case EMonoChannelUpmixMethod::Linear:
+							OutChannelMap.Add(0.5f);
+							OutChannelMap.Add(0.5f);
+							break;
+
+						case EMonoChannelUpmixMethod::EqualPower:
+							OutChannelMap.Add(0.707f);
+							OutChannelMap.Add(0.707f);
+							break;
+
+						case EMonoChannelUpmixMethod::FullVolume:
+							OutChannelMap.Add(1.0f);
+							OutChannelMap.Add(1.0f);
+							break;
+					}
 
 					for (int32 OutputChannel = 2; OutputChannel < NumOutputChannels; ++OutputChannel)
 					{
diff --git a/Engine/Source/Runtime/AudioMixer/Private/AudioMixerDevice.cpp b/Engine/Source/Runtime/AudioMixer/Private/AudioMixerDevice.cpp
--- a/Engine/Source/Runtime/AudioMixer/Private/AudioMixerDevice.cpp
+++ b/Engine/Source/Runtime/AudioMixer/Private/AudioMixerDevice.cpp
@@ -15,11 +15,29 @@
 #include "UObject/UObjectIterator.h"
 #include "Runtime/HeadMountedDisplay/Public/IHeadMountedDisplayModule.h"
 #include "Misc/App.h"
+#include "Sound/AudioSettings.h"
 
 #if WITH_EDITOR
 #include "AudioEditorModule.h"
 #endif
 
+static int32 DisableSubmixReverbCVar = 0;
+FAutoConsoleVariableRef CVarDisableSubmixReverb(
+	TEXT("au.DisableReverbSubmix"),
+	DisableSubmixReverbCVar,
+	TEXT("Disables the reverb submix.\n")
+	TEXT("0: Not Disabled, 1: Disabled"),
+	ECVF_Default);
+
+static int32 DisableSubmixEffectEQCvar = 0;
+FAutoConsoleVariableRef CVarDisableSubmixEQ(
+	TEXT("au.DisableSubmixEffectEQ"),
+	DisableSubmixEffectEQCvar,
+	TEXT("Disables the eq submix.\n")
+	TEXT("0: Not Disabled, 1: Disabled"),
+	ECVF_Default);
+
+
 namespace Audio
 {
 	TArray<USoundSubmix*> FMixerDevice::MasterSubmixes;
@@ -109,6 +127,10 @@ namespace Audio
 
 		if (AudioMixerPlatform && AudioMixerPlatform->InitializeHardware())
 		{
+			const UAudioSettings* AudioSettings = GetDefault<UAudioSettings>();
+			MonoChannelUpmixMethod = AudioSettings->MonoChannelUpmixMethod;
+			PanningMethod = AudioSettings->PanningMethod;
+
 			// Set whether we're the main audio mixer
 			bIsMainAudioMixer = IsMainAudioDevice();
 
@@ -168,7 +190,7 @@ namespace Audio
 				// This extra source count is used for "stopping sources", which are sources
 				// which are fading out (very quickly) to avoid discontinuities when stopping sounds
 				FSourceManagerInitParams SourceManagerInitParams;
-				SourceManagerInitParams.NumSources = GetMaxChannels();
+				SourceManagerInitParams.NumSources = GetMaxChannels() + NumStoppingVoices;
 				SourceManagerInitParams.NumSourceWorkers = 4;
 
 				SourceManager.Init(SourceManagerInitParams);
@@ -270,6 +292,14 @@ namespace Audio
 
 	void FMixerDevice::UpdateHardware()
 	{
+		// If we're in editor, re-query these in case they changed. 
+		if (GIsEditor)
+		{
+			const UAudioSettings* AudioSettings = GetDefault<UAudioSettings>();
+			MonoChannelUpmixMethod = AudioSettings->MonoChannelUpmixMethod;
+			PanningMethod = AudioSettings->PanningMethod;
+		}
+
 		SourceManager.Update();
 
 		AudioMixerPlatform->OnHardwareUpdate();
@@ -463,19 +493,40 @@ namespace Audio
 			FMixerDevice::MasterSubmixes.Add(MasterSubmix);
 
 			// Master Reverb Plugin
-			USoundSubmix* ReverbPluginSubmix = NewObject<USoundSubmix>(USoundSubmix::StaticClass(), TEXT("Master Reverb Plugin Submix"));
-			ReverbPluginSubmix->AddToRoot();
-			FMixerDevice::MasterSubmixes.Add(ReverbPluginSubmix);
+			if (DisableSubmixReverbCVar)
+			{
+				FMixerDevice::MasterSubmixes.Add(nullptr);
+			}
+			else
+			{
+				USoundSubmix* ReverbPluginSubmix = NewObject<USoundSubmix>(USoundSubmix::StaticClass(), TEXT("Master Reverb Plugin Submix"));
+				ReverbPluginSubmix->AddToRoot();
+				FMixerDevice::MasterSubmixes.Add(ReverbPluginSubmix);
+			}
 
 			// Master Reverb
-			USoundSubmix* ReverbSubmix = NewObject<USoundSubmix>(USoundSubmix::StaticClass(), TEXT("Master Reverb Submix"));
-			ReverbSubmix->AddToRoot();
-			FMixerDevice::MasterSubmixes.Add(ReverbSubmix);
+			if (DisableSubmixReverbCVar)
+			{
+				FMixerDevice::MasterSubmixes.Add(nullptr);
+			}
+			else
+			{
+				USoundSubmix* ReverbSubmix = NewObject<USoundSubmix>(USoundSubmix::StaticClass(), TEXT("Master Reverb Submix"));
+				ReverbSubmix->AddToRoot();
+				FMixerDevice::MasterSubmixes.Add(ReverbSubmix);
+			}
 
 			// Master EQ
-			USoundSubmix* EQSubmix = NewObject<USoundSubmix>(USoundSubmix::StaticClass(), TEXT("Master EQ Submix"));
-			EQSubmix->AddToRoot();
-			FMixerDevice::MasterSubmixes.Add(EQSubmix);
+			if (DisableSubmixEffectEQCvar)
+			{
+				FMixerDevice::MasterSubmixes.Add(nullptr);
+			}
+			else
+			{
+				USoundSubmix* EQSubmix = NewObject<USoundSubmix>(USoundSubmix::StaticClass(), TEXT("Master EQ Submix"));
+				EQSubmix->AddToRoot();
+				FMixerDevice::MasterSubmixes.Add(EQSubmix);
+			}
 
 			// Master ambisonics
 			USoundSubmix* AmbisonicsSubmix = NewObject<USoundSubmix>(USoundSubmix::StaticClass(), TEXT("Master Ambisonics Submix"));
@@ -494,10 +545,17 @@ namespace Audio
 		{
 			for (int32 i = 0; i < EMasterSubmixType::Count; ++i)
 			{
-				FMixerSubmixPtr MixerSubmix = FMixerSubmixPtr(new FMixerSubmix(this));
-				MixerSubmix->Init(FMixerDevice::MasterSubmixes[i]);
+				if (FMixerDevice::MasterSubmixes[i])
+				{
+					FMixerSubmixPtr MixerSubmix = FMixerSubmixPtr(new FMixerSubmix(this));
+					MixerSubmix->Init(FMixerDevice::MasterSubmixes[i]);
 
-				MasterSubmixInstances.Add(MixerSubmix);
+					MasterSubmixInstances.Add(MixerSubmix);
+				}
+				else
+				{
+					MasterSubmixInstances.Add(nullptr);
+				}
 			}
 
 			FMixerSubmixPtr MasterSubmixInstance = MasterSubmixInstances[EMasterSubmixType::Master];
@@ -507,7 +565,7 @@ namespace Audio
 
 	
 			// Setup the master reverb plugin
-			if (ReverbPluginInterface.IsValid())
+			if (ReverbPluginInterface.IsValid() && MasterSubmixInstances[EMasterSubmixType::ReverbPlugin].IsValid())
 			{
 				auto ReverbPluginEffectSubmix = ReverbPluginInterface->GetEffectSubmix(FMixerDevice::MasterSubmixes[EMasterSubmixType::ReverbPlugin]);
 
@@ -521,7 +579,7 @@ namespace Audio
 				MasterReverbPluginSubmix->SetParentSubmix(MasterSubmixInstance);
 				MasterSubmixInstance->AddChildSubmix(MasterReverbPluginSubmix);
 			}
-			else
+			else if (MasterSubmixInstances[EMasterSubmixType::Reverb].IsValid())
 			{
 				// Setup the master reverb only if we don't have a reverb plugin
 
@@ -544,26 +602,29 @@ namespace Audio
 			}
 
 			// Setup the master EQ
-			USoundSubmix* MasterEQSoundSubmix = FMixerDevice::MasterSubmixes[EMasterSubmixType::EQ];
-			USubmixEffectSubmixEQPreset* EQPreset = NewObject<USubmixEffectSubmixEQPreset>(MasterEQSoundSubmix, TEXT("Master EQ Effect preset"));
-			EQPreset->AddToRoot();
-
-			FSoundEffectSubmix* EQEffectSubmix = static_cast<FSoundEffectSubmix*>(EQPreset->CreateNewEffect());
-			EQEffectSubmix->Init(InitData);
-			EQEffectSubmix->SetPreset(EQPreset);
-			EQEffectSubmix->SetEnabled(true);
-
-			const uint32 EQPresetId = EQPreset->GetUniqueID();
-
-			FMixerSubmixPtr MasterEQSubmix = MasterSubmixInstances[EMasterSubmixType::EQ];
-			MasterEQSubmix->AddSoundEffectSubmix(EQPresetId, MakeShareable(EQEffectSubmix));
-			MasterEQSubmix->SetParentSubmix(MasterSubmixInstance);
-			MasterSubmixInstance->AddChildSubmix(MasterEQSubmix);
-
-			// Add the ambisonics master submix
-			FMixerSubmixPtr MasterAmbisonicsSubmix = MasterSubmixInstances[EMasterSubmixType::Ambisonics];
-			MasterAmbisonicsSubmix->SetParentSubmix(MasterEQSubmix);
-			MasterEQSubmix->AddChildSubmix(MasterAmbisonicsSubmix);
+			if (FMixerDevice::MasterSubmixes[EMasterSubmixType::EQ])
+			{
+				USoundSubmix* MasterEQSoundSubmix = FMixerDevice::MasterSubmixes[EMasterSubmixType::EQ];
+				USubmixEffectSubmixEQPreset* EQPreset = NewObject<USubmixEffectSubmixEQPreset>(MasterEQSoundSubmix, TEXT("Master EQ Effect preset"));
+				EQPreset->AddToRoot();
+
+				FSoundEffectSubmix* EQEffectSubmix = static_cast<FSoundEffectSubmix*>(EQPreset->CreateNewEffect());
+				EQEffectSubmix->Init(InitData);
+				EQEffectSubmix->SetPreset(EQPreset);
+				EQEffectSubmix->SetEnabled(true);
+
+				const uint32 EQPresetId = EQPreset->GetUniqueID();
+
+				FMixerSubmixPtr MasterEQSubmix = MasterSubmixInstances[EMasterSubmixType::EQ];
+				MasterEQSubmix->AddSoundEffectSubmix(EQPresetId, MakeShareable(EQEffectSubmix));
+				MasterEQSubmix->SetParentSubmix(MasterSubmixInstance);
+				MasterSubmixInstance->AddChildSubmix(MasterEQSubmix);
+
+				// Add the ambisonics master submix
+				FMixerSubmixPtr MasterAmbisonicsSubmix = MasterSubmixInstances[EMasterSubmixType::Ambisonics];
+				MasterAmbisonicsSubmix->SetParentSubmix(MasterEQSubmix);
+				MasterEQSubmix->AddChildSubmix(MasterAmbisonicsSubmix);
+			}
 		}
 
 		// Now register all the non-core submixes
@@ -900,7 +961,7 @@ namespace Audio
 		return SourceManager.GetNumActiveSources();
 	}
 
-	void FMixerDevice::Get3DChannelMap(const ESubmixChannelFormat InSubmixType, const FWaveInstance* InWaveInstance, float EmitterAzimith, float NormalizedOmniRadius, TArray<float>& OutChannelMap)
+	void FMixerDevice::Get3DChannelMap(const ESubmixChannelFormat InSubmixType, const FWaveInstance* InWaveInstance, float EmitterAzimith, float NormalizedOmniRadius, Audio::AlignedFloatBuffer& OutChannelMap)
 	{
 		// If we're center-channel only, then no need for spatial calculations, but need to build a channel map
 		if (InWaveInstance->bCenterChannelOnly)
@@ -985,11 +1046,19 @@ namespace Audio
 		float PrevChannelPan; 
 		float NextChannelPan;
 
-		FMath::SinCos(&NextChannelPan, &PrevChannelPan, Fraction * 0.5f * PI);
+		if (PanningMethod == EPanningMethod::EqualPower)
+		{
+			FMath::SinCos(&NextChannelPan, &PrevChannelPan, Fraction * 0.5f * PI);
 
-		// Note that SinCos can return values slightly greater than 1.0 when very close to PI/2
-		NextChannelPan = FMath::Clamp(NextChannelPan, 0.0f, 1.0f);
-		PrevChannelPan = FMath::Clamp(PrevChannelPan, 0.0f, 1.0f);
+			// Note that SinCos can return values slightly greater than 1.0 when very close to PI/2
+			NextChannelPan = FMath::Clamp(NextChannelPan, 0.0f, 1.0f);
+			PrevChannelPan = FMath::Clamp(PrevChannelPan, 0.0f, 1.0f);
+		}
+		else
+		{
+			NextChannelPan = Fraction;
+			PrevChannelPan = 1.0f - Fraction;
+		}
 
 		float NormalizedOmniRadSquared = NormalizedOmniRadius * NormalizedOmniRadius;
 		float OmniAmount = 0.0f;
@@ -1006,7 +1075,7 @@ namespace Audio
 		{
 			NumSpatialChannels--;
 		}
-		float OmniPanFactor = 1.0f / FMath::Sqrt(NumSpatialChannels);
+		float OmniPanFactor = 1.0f / NumSpatialChannels;
 
 		float DefaultEffectivePan = !OmniAmount ? 0.0f : FMath::Lerp(0.0f, OmniPanFactor, OmniAmount);
 		const TArray<EAudioMixerChannel::Type>& ChannelArray = GetChannelArrayForSubmixChannelType(InSubmixType);
diff --git a/Engine/Source/Runtime/AudioMixer/Private/AudioMixerSource.cpp b/Engine/Source/Runtime/AudioMixer/Private/AudioMixerSource.cpp
--- a/Engine/Source/Runtime/AudioMixer/Private/AudioMixerSource.cpp
+++ b/Engine/Source/Runtime/AudioMixer/Private/AudioMixerSource.cpp
@@ -16,42 +16,28 @@ FAutoConsoleVariableRef CVarDisableHRTF(
 	TEXT("0: Not Disabled, 1: Disabled"),
 	ECVF_Default);
 
-static int32 DisableStoppingVoicesCvar = 0;
-FAutoConsoleVariableRef CVarDisableStoppingVoices(
-	TEXT("au.DisableStoppingVoices"),
-	DisableStoppingVoicesCvar,
-	TEXT("Forces all sounds to not fade out and instead stop immediately.\n")
-	TEXT("0: Not Disabled, 1: Disabled"),
-	ECVF_Default);
-
 namespace Audio
 {
 	FMixerSource::FMixerSource(FAudioDevice* InAudioDevice)
 		: FSoundSource(InAudioDevice)
 		, MixerDevice((FMixerDevice*)InAudioDevice)
 		, MixerBuffer(nullptr)
-		, MixerSourceVoice(nullptr)
-		, AsyncRealtimeAudioTask(nullptr)
-		, CurrentBuffer(0)
+		, MixerSourceBuffer(nullptr)
+		, MixerSourceVoice(nullptr)		
 		, PreviousAzimuth(-1.0f)
 		, bPlayedCachedBuffer(false)
 		, bPlaying(false)
 		, bLoopCallback(false)
 		, bIsFinished(false)
 		, bIsPlayingEffectTails(false)
-		, bBuffersToFlush(false)
 		, bResourcesNeedFreeing(false)
 		, bEditorWarnedChangedSpatialization(false)
 		, bUsingHRTFSpatialization(false)
 		, bIs3D(false)
 		, bDebugMode(false)
 		, bIsVorbis(false)
+		, bIsStoppingVoicesEnabled(InAudioDevice->IsStoppingVoicesEnabled())
 	{
-		// Create the source voice buffers
-		for (int32 i = 0; i < Audio::MAX_BUFFERS_QUEUED; ++i)
-		{
-			SourceVoiceBuffers.Add(FMixerSourceBufferPtr(new FMixerSourceVoiceBuffer()));
-		}
 	}
 
 	FMixerSource::~FMixerSource()
@@ -81,6 +67,10 @@ namespace Audio
 		// Reset our releasing bool
 		bIsReleasing = false;
 
+		bIsStoppingVoicesEnabled = ((FAudioDevice*)MixerDevice)->IsStoppingVoicesEnabled();
+
+		bIsStopping = false;
+
 		FSoundBuffer* SoundBuffer = static_cast<FSoundBuffer*>(MixerBuffer);
 		if (SoundBuffer->NumChannels > 0)
 		{
@@ -95,7 +85,7 @@ namespace Audio
 
 			// Initialize the source voice with the necessary format information
 			FMixerSourceVoiceInitParams InitParams;
-			InitParams.BufferQueueListener = this;
+			InitParams.SourceListener = this;
 			InitParams.NumInputChannels = InWaveInstance->WaveData->NumChannels;
 			InitParams.NumInputFrames = NumFrames;
 			InitParams.SourceVoice = MixerSourceVoice;
@@ -192,16 +182,22 @@ namespace Audio
 						{
 							// Send the voice to the EQ submix if it's enabled
 							const bool bIsEQDisabled = GetDefault<UAudioSettings>()->bDisableMasterEQ;
+							bool bUseMaster = true;
 							if (!bIsEQDisabled && IsEQFilterApplied())
 							{
-								// Default the submix to use to use the master submix if none are set
-								FMixerSourceSubmixSend SubmixSend;
-								SubmixSend.Submix = MixerDevice->GetMasterEQSubmix();
-								SubmixSend.SendLevel = 1.0f;
-								SubmixSend.bIsMainSend = true;
-								InitParams.SubmixSends.Add(SubmixSend);
+								if (MixerDevice->GetMasterEQSubmix().IsValid())
+								{
+									// Default the submix to use to use the master submix if none are set
+									FMixerSourceSubmixSend SubmixSend;
+									SubmixSend.Submix = MixerDevice->GetMasterEQSubmix();
+									SubmixSend.SendLevel = 1.0f;
+									SubmixSend.bIsMainSend = true;
+									InitParams.SubmixSends.Add(SubmixSend);
+									bUseMaster = false;
+								}
 							}
-							else
+
+							if (bUseMaster)
 							{
 								// Default the submix to use to use the master submix if none are set
 								FMixerSourceSubmixSend SubmixSend;
@@ -268,27 +264,17 @@ namespace Audio
 			// Update the buffer sample rate to the wave instance sample rate in case it was serialized incorrectly
 			MixerBuffer->InitSampleRate(InWaveInstance->WaveData->GetSampleRateForCurrentPlatform());
 
+			// Now we init the mixer source buffer
+			MixerSourceBuffer->Init();
+
+			// Hand off the mixer source buffer decoder
+			InitParams.MixerSourceBuffer = MixerSourceBuffer;
+			MixerSourceBuffer = nullptr;
+
 			if (MixerSourceVoice->Init(InitParams))
 			{
 				AUDIO_MIXER_CHECK(WaveInstance);
 
-				const EBufferType::Type BufferType = MixerBuffer->GetType();
-				switch (BufferType)
-				{
-					case EBufferType::PCM:
-					case EBufferType::PCMPreview:
-						SubmitPCMBuffers();
-						break;
-
-					case EBufferType::PCMRealTime:
-					case EBufferType::Streaming:
-						SubmitPCMRTBuffers();
-						break;
-
-					case EBufferType::Invalid:
-						break;
-				}
-
 				bInitialized = true;
 
 				Update();
@@ -308,6 +294,8 @@ namespace Audio
 			return;
 		}
 
+		++TickCount;
+
 		UpdatePitch();
 
 		UpdateVolume();
@@ -339,7 +327,8 @@ namespace Audio
 		check(MixerBuffer == nullptr);
 		check(AudioDevice);
 
-		MixerBuffer = FMixerBuffer::Init(AudioDevice, InWaveInstance->WaveData, InWaveInstance->StartTime > 0.0f);
+		bool bIsSeeking = InWaveInstance->StartTime > 0.0f;
+		MixerBuffer = FMixerBuffer::Init(AudioDevice, InWaveInstance->WaveData, bIsSeeking);
 		if (MixerBuffer)
 		{
 			Buffer = MixerBuffer;
@@ -347,6 +336,10 @@ namespace Audio
 
 			LPFFrequency = MAX_FILTER_FREQUENCY;
 			LastLPFFrequency = FLT_MAX;
+
+			HPFFrequency = 0.0;
+			LastHPFFrequency = FLT_MAX;
+
 			bIsFinished = false;
 
 			EBufferType::Type BufferType = MixerBuffer->GetType();
@@ -367,20 +360,15 @@ namespace Audio
 				}
 			}
 
-			// Set up buffer areas to decompress audio into
-			for (int32 BufferIndex = 0; BufferIndex < Audio::MAX_BUFFERS_QUEUED; ++BufferIndex)
+			check(!MixerSourceBuffer.IsValid());
+			MixerSourceBuffer = TSharedPtr<FMixerSourceBuffer>(new FMixerSourceBuffer());
+
+			if (MixerSourceBuffer->PreInit(MixerBuffer, InWaveInstance->WaveData, InWaveInstance->LoopingMode, bIsSeeking))
 			{
-				const uint32 TotalSamples = MONO_PCM_BUFFER_SAMPLES * MixerBuffer->NumChannels;
-				SourceVoiceBuffers[BufferIndex]->AudioData.Reset();
-				SourceVoiceBuffers[BufferIndex]->AudioData.AddZeroed(TotalSamples);
-				SourceVoiceBuffers[BufferIndex]->Samples = TotalSamples;
-				SourceVoiceBuffers[BufferIndex]->bRealTimeBuffer = true;
-				SourceVoiceBuffers[BufferIndex]->LoopCount = 0;
+				// We succeeded in preparing the buffer for initialization, but we are not technically initialized yet.
+				// If the buffer is asynchronously preparing a file-handle, we may not yet initialize the source.
+				return true;
 			}
-
-			// We succeeded in preparing the buffer for initialization, but we are not technically initialized yet.
-			// If the buffer is asynchronously preparing a file-handle, we may not yet initialize the source.
-			return true;
 		}
 
 		// Something went wrong with initializing the generator
@@ -391,36 +379,42 @@ namespace Audio
 	{
 		if (MixerBuffer && MixerBuffer->IsRealTimeSourceReady() && !bIsReleasing)
 		{
+			check(MixerSourceBuffer.IsValid());
+
 			// Check if we have a realtime audio task already (doing first decode)
-			if (AsyncRealtimeAudioTask)
+			if (MixerSourceBuffer->IsAsyncTaskInProgress())
 			{
 				// not ready
-				return AsyncRealtimeAudioTask->IsDone();
-			}
-			else if (WaveInstance->WaveData->bIsBus)
-			{
-				// Buses don't need to do anything to play audio
-				return true;
+				return MixerSourceBuffer->IsAsyncTaskDone();
 			}
-			else
+			else if (WaveInstance)
 			{
-				// Now check to see if we need to kick off a decode the first chunk of audio
-				const EBufferType::Type BufferType = MixerBuffer->GetType();
-				if ((BufferType == EBufferType::PCMRealTime || BufferType == EBufferType::Streaming) && WaveInstance->WaveData)
+				if (WaveInstance->WaveData->bIsBus)
+				{
+					// Buses don't need to do anything to play audio
+					return true;
+				}
+				else
 				{
-					// If any of these conditions meet, we need to do an initial async decode before we're ready to start playing the sound
-					if (WaveInstance->StartTime > 0.0f || WaveInstance->WaveData->bProcedural || WaveInstance->WaveData->bIsBus || !WaveInstance->WaveData->CachedRealtimeFirstBuffer)
+					// Now check to see if we need to kick off a decode the first chunk of audio
+					const EBufferType::Type BufferType = MixerBuffer->GetType();
+					if ((BufferType == EBufferType::PCMRealTime || BufferType == EBufferType::Streaming) && WaveInstance->WaveData)
 					{
-						// Before reading more PCMRT data, we first need to seek the buffer
-						if (WaveInstance->StartTime > 0.0f && !WaveInstance->WaveData->bIsBus && !WaveInstance->WaveData->bProcedural)
+						// If any of these conditions meet, we need to do an initial async decode before we're ready to start playing the sound
+						if (WaveInstance->StartTime > 0.0f || WaveInstance->WaveData->bProcedural || WaveInstance->WaveData->bIsBus || !WaveInstance->WaveData->CachedRealtimeFirstBuffer)
 						{
-							MixerBuffer->Seek(WaveInstance->StartTime);
-						}
+							// Before reading more PCMRT data, we first need to seek the buffer
+							if (WaveInstance->StartTime > 0.0f && !WaveInstance->WaveData->bIsBus && !WaveInstance->WaveData->bProcedural)
+							{
+								MixerBuffer->Seek(WaveInstance->StartTime);
+							}
 
-						ReadMorePCMRTData(0, EBufferReadMode::Asynchronous);
+							check(MixerSourceBuffer.IsValid());
+							MixerSourceBuffer->ReadMoreRealtimeData(0, EBufferReadMode::Asynchronous);
 
-						// not ready
-						return false;
+							// not ready
+							return false;
+						}
 					}
 				}
 			}
@@ -448,95 +442,66 @@ namespace Audio
 		bIsStopping = false;
 		Paused = false;
 		Playing = true;
-		bBuffersToFlush = false;
 		bLoopCallback = false;
 		bIsFinished = false;
 	}
 
-	void FMixerSource::StopInternal(bool bStopNow)
+	void FMixerSource::Stop()
 	{
-		IStreamingManager::Get().GetAudioStreamingManager().RemoveStreamingSoundSource(this);
-
-		if (WaveInstance)
+		if (bIsFinished)
+		{
+			StopNow();
+		}
+		else if (!bIsStopping)
 		{
-			// We stop procedural sources immediately, since we do not know whether implementations of USoundWaveProcedural will correctly handle delayed stops.
-			const bool bIsProcedural = WaveInstance->WaveData && WaveInstance->WaveData->DecompressionType == DTYPE_Procedural;
-
-			// If we've finished naturally or we're paused or we're told to immediately stop
-			if (bIsFinished || bIsProcedural || bStopNow || DisableStoppingVoicesCvar|| (MixerSourceVoice && MixerSourceVoice->IsPaused()))
-			{
-				FScopeLock Lock(&RenderThreadCritSect);
-
-				// Stop immediately stops a sound with no fade
-				if (MixerSourceVoice && Playing)
-				{
-					MixerSourceVoice->Stop();
-				}
-
-				Paused = false;
-				Playing = false;
-				bIsStopping = false;
-
-				FreeResources();
-			}
 			// Otherwise, we need to do a quick fade-out of the sound and put the state
 			// of the sound into "stopping" mode. This prevents this source from
 			// being put into the "free" pool and prevents the source from freeing its resources
 			// until the sound has finished naturally (i.e. faded all the way out)
-			else if (!bIsStopping)
-			{
-				FScopeLock Lock(&RenderThreadCritSect);
 
-				bIsStopping = true;
+			// StopFade will stop a sound with a very small fade to avoid discontinuities
+			if (MixerSourceVoice && Playing)
+			{
+				if (bIsStoppingVoicesEnabled && !WaveInstance->WaveData->bProcedural)
+				{
+					// Let the wave instance know it's stopping
+					WaveInstance->SetStopping(true);
 
-				// StopFade will stop a sound with a very small fade to avoid discontinuities
-				if (MixerSourceVoice && Playing)
+					// TODO: parameterize the number of fades
+					MixerSourceVoice->StopFade(512);
+					bIsStopping = true;
+				}
+				else
 				{
-					MixerSourceVoice->StopFade();
+					StopNow();
 				}
-
-				Paused = false;
 			}
-		}
 
-		if (!bIsStopping)
-		{
-			bInitialized = false;
-			FSoundSource::Stop();
+			Paused = false;
 		}
 	}
 
-
-	void FMixerSource::Stop()
-	{
-		StopInternal(false);
-	}
-
 	void FMixerSource::StopNow()
 	{
-		StopInternal(true);
-	}
-
-	void FMixerSource::EnsureStopped()
-	{
-		// Make sure the source finishes its fadeout
-		if (MixerSourceVoice)
-		{
-			MixerSourceVoice->EnsureStopped();
-			bIsStopping = false;
-			check(MixerSourceVoice->IsDone());
-		}
-	}
+		// Immediately stop the sound source
 
-	void FMixerSource::ReleaseStoppedSound()
-	{
-		// Do state reset and free resources
 		bInitialized = false;
-		Paused = false;
-		Playing = false;
+		IStreamingManager::Get().GetAudioStreamingManager().RemoveStreamingSoundSource(this);
+
 		bIsStopping = false;
 
-		FreeResources();
+		if (WaveInstance)
+		{
+			if (MixerSourceVoice && Playing)
+			{
+				MixerSourceVoice->Stop();
+			}
+
+			Paused = false;
+			Playing = false;
+
+			FreeResources();
+		}
 
 		FSoundSource::Stop();
 	}
@@ -559,16 +524,22 @@ namespace Audio
 	bool FMixerSource::IsFinished()
 	{
 		// A paused source is not finished.
-		if (Paused || !bInitialized)
+		if (Paused)
 		{
 			return false;
 		}
 
+		if (!bInitialized)
+		{
+			return true;
+		}
+
 		if (WaveInstance && MixerSourceVoice)
 		{
 			if (bIsFinished && MixerSourceVoice->IsSourceEffectTailsDone())
 			{
 				WaveInstance->NotifyFinished();
+				bIsStopping = false;
 				return true;
 			}
 			// Buses don't do buffer end callbacks, so we need to directly query doneness. 
@@ -579,6 +550,7 @@ namespace Audio
 				{
 					bIsFinished = true;
 					WaveInstance->NotifyFinished();
+					bIsStopping = false;
 					return true;
 				}
 			}
@@ -629,280 +601,18 @@ namespace Audio
 		return 0.0f;
 	}
 
-	void FMixerSource::SubmitPCMBuffers()
-	{
-		if (!AudioDevice)
-		{
-			UE_LOG(LogAudioMixer, Error, TEXT("SubmitPCMBuffers: Audio device is nullptr"));
-			return;
-		}
-
-		if (!MixerSourceVoice)
-		{
-			UE_LOG(LogAudioMixer, Error, TEXT("SubmitPCMBuffers: Source is nullptr"));
-			return;
-		}
-
-		CurrentBuffer = 0;
-
-		RawPCMDataBuffer.Data = nullptr;
-		RawPCMDataBuffer.DataSize = 0;
-		MixerBuffer->GetPCMData(&RawPCMDataBuffer.Data, &RawPCMDataBuffer.DataSize);
-
-		RawPCMDataBuffer.NumSamples = RawPCMDataBuffer.DataSize / sizeof(int16);
-		RawPCMDataBuffer.CurrentSample = 0;
-
-		// Only submit data if we've successfully loaded it
-		if (!RawPCMDataBuffer.Data || !RawPCMDataBuffer.DataSize)
-		{
-			UE_LOG(LogAudioMixer, Error, TEXT("Failed to load PCM data from sound source %s"), *WaveInstance->GetName());
-			return;
-		}
-
-		RawPCMDataBuffer.LoopCount = (WaveInstance->LoopingMode != LOOP_Never) ? Audio::LOOP_FOREVER : 0;
-
-		// Submit the first two format-converted chunks to the source voice
-		const uint32 NumSamplesPerBuffer = MONO_PCM_BUFFER_SAMPLES * MixerBuffer->NumChannels;
-		int16* RawPCMBufferDataPtr = (int16*)RawPCMDataBuffer.Data;
-
-		RawPCMDataBuffer.GetNextBuffer(SourceVoiceBuffers[0], NumSamplesPerBuffer);
-		MixerSourceVoice->SubmitBuffer(SourceVoiceBuffers[0], false);
-
-		CurrentBuffer = 1;
-	}
-
-	void FMixerSource::SubmitPCMRTBuffers()
-	{
-		CurrentBuffer = 0;
-
-		bPlayedCachedBuffer = false;
-		bool bIsSeeking = (WaveInstance->StartTime > 0.0f);
-		if (!bIsSeeking && WaveInstance->WaveData && WaveInstance->WaveData->CachedRealtimeFirstBuffer)
-		{
-			bPlayedCachedBuffer = true;
-
-			// Format convert the first cached buffers
-			const uint32 NumSamples = MONO_PCM_BUFFER_SAMPLES * MixerBuffer->NumChannels;
-			const uint32 BufferSize = MONO_PCM_BUFFER_SIZE * MixerBuffer->NumChannels;
-
-			int16* CachedBufferPtr0 = (int16*)WaveInstance->WaveData->CachedRealtimeFirstBuffer;
-			int16* CachedBufferPtr1 = (int16*)(WaveInstance->WaveData->CachedRealtimeFirstBuffer + BufferSize);
-			float* AudioData0 = SourceVoiceBuffers[0]->AudioData.GetData();
-			float* AudioData1 = SourceVoiceBuffers[1]->AudioData.GetData();
-			for (uint32 Sample = 0; Sample < NumSamples; ++Sample)
-			{
-				AudioData0[Sample] = CachedBufferPtr0[Sample] / 32768.0f;
-				AudioData1[Sample] = CachedBufferPtr1[Sample] / 32768.0f;
-			}
-
-			// Submit the already decoded and cached audio buffers
-			MixerSourceVoice->SubmitBuffer(SourceVoiceBuffers[0], false);
-			MixerSourceVoice->SubmitBuffer(SourceVoiceBuffers[1], false);
-
-			CurrentBuffer = 2;
-		}
-		else if (WaveInstance->WaveData && !WaveInstance->WaveData->bIsBus)
-		{
-			// We should have already kicked off and finished a task. 
-			check(AsyncRealtimeAudioTask != nullptr);
-
-			ProcessRealTimeSource(true, false);
-		}
-
-		// We should have set this
-		check(bResourcesNeedFreeing);
-	}
-
-	bool FMixerSource::ReadMorePCMRTData(const int32 BufferIndex, EBufferReadMode BufferReadMode, bool* OutLooped)
-	{
-		USoundWave* WaveData = WaveInstance->WaveData;
-
-		if (WaveData)
-		{
-			if (WaveData->bProcedural)
-			{
-				const int32 MaxSamples = MONO_PCM_BUFFER_SAMPLES * Buffer->NumChannels;
-
-				FProceduralAudioTaskData NewTaskData;
-				NewTaskData.ProceduralSoundWave = WaveData;
-				NewTaskData.AudioData = SourceVoiceBuffers[BufferIndex]->AudioData.GetData();
-				NewTaskData.NumSamples = MaxSamples;
-				NewTaskData.NumChannels = Buffer->NumChannels;
-				check(!AsyncRealtimeAudioTask);
-				AsyncRealtimeAudioTask = CreateAudioTask(NewTaskData);
-
-				// Procedural sound waves never loop
-				return false;
-			}
-			else if (!MixerBuffer->IsRealTimeBuffer())
-			{
-				check(RawPCMDataBuffer.Data != nullptr);
-
-				// Read the next raw PCM buffer into the source buffer index. This converts raw PCM to float.
-				const uint32 NumSamplesPerBuffer = MONO_PCM_BUFFER_SAMPLES * MixerBuffer->NumChannels;
-				return RawPCMDataBuffer.GetNextBuffer(SourceVoiceBuffers[BufferIndex], NumSamplesPerBuffer);
-			}
-			else
-			{
-				FDecodeAudioTaskData NewTaskData;
-				NewTaskData.MixerBuffer = MixerBuffer;
-				NewTaskData.AudioData = SourceVoiceBuffers[BufferIndex]->AudioData.GetData();
-				NewTaskData.bLoopingMode = WaveInstance->LoopingMode != LOOP_Never;
-				NewTaskData.bSkipFirstBuffer = (BufferReadMode == EBufferReadMode::AsynchronousSkipFirstFrame);
-				NewTaskData.NumFramesToDecode = MONO_PCM_BUFFER_SAMPLES;
-
-				check(!AsyncRealtimeAudioTask);
-				AsyncRealtimeAudioTask = CreateAudioTask(NewTaskData);
-
-				// Not looping
-				return false;
-			}
-		}
-		return true;
-	}
-
-	void FMixerSource::SubmitRealTimeSourceData(const bool bLooped, const bool bSubmitSynchronously)
-	{
-		// Have we reached the end of the sound
-		if (bLooped)
-		{
-			switch (WaveInstance->LoopingMode)
-			{
-				case LOOP_Never:
-					// Play out any queued buffers - once there are no buffers left, the state check at the beginning of IsFinished will fire
-					bBuffersToFlush = true;
-					break;
-
-				case LOOP_WithNotification:
-					// If we have just looped, and we are looping, send notification
-					// This will trigger a WaveInstance->NotifyFinished() in the FXAudio2SoundSournce::IsFinished() function on main thread.
-					bLoopCallback = true;
-					break;
-
-				case LOOP_Forever:
-					// Let the sound loop indefinitely
-					break;
-			}
-		}
-
-		if (MixerSourceVoice && SourceVoiceBuffers[CurrentBuffer]->AudioData.Num() > 0)
-		{
-			MixerSourceVoice->SubmitBuffer(SourceVoiceBuffers[CurrentBuffer], bSubmitSynchronously);
-		}
-	}
-
-	void FMixerSource::ProcessRealTimeSource(const bool bBlockForData, const bool bSubmitSynchronously)
-	{
-		const bool bGetMoreData = bBlockForData || (AsyncRealtimeAudioTask == nullptr);
-		if (AsyncRealtimeAudioTask)
-		{
-			const bool bTaskDone = AsyncRealtimeAudioTask->IsDone();
-			if (bTaskDone || bBlockForData)
-			{
-				bool bLooped = false;
-
-				if (!bTaskDone)
-				{
-					AsyncRealtimeAudioTask->EnsureCompletion();
-				}
-
-				switch (AsyncRealtimeAudioTask->GetType())
-				{
-					case EAudioTaskType::Decode:
-					{
-						FDecodeAudioTaskResults TaskResult;
-						AsyncRealtimeAudioTask->GetResult(TaskResult);
-
-						SourceVoiceBuffers[CurrentBuffer]->Samples = MONO_PCM_BUFFER_SAMPLES * MixerBuffer->NumChannels;
-						bLooped = TaskResult.bLooped;
-					}
-					break;
-
-					case EAudioTaskType::Procedural:
-					{
-						FProceduralAudioTaskResults TaskResult;
-						AsyncRealtimeAudioTask->GetResult(TaskResult);
-
-						SourceVoiceBuffers[CurrentBuffer]->Samples = TaskResult.NumSamplesWritten;
-					}
-					break;
-				}
-
-				delete AsyncRealtimeAudioTask;
-				AsyncRealtimeAudioTask = nullptr;
-				
-				SubmitRealTimeSourceData(bLooped, bSubmitSynchronously);
-			}
-		}
-
-		if (bGetMoreData)
-		{
-			// Update the buffer index
-			if (++CurrentBuffer > 2)
-			{
-				CurrentBuffer = 0;
-			}
-
-			EBufferReadMode DataReadMode;
-			if (bPlayedCachedBuffer)
-			{
-				bPlayedCachedBuffer = false;
-				DataReadMode = EBufferReadMode::AsynchronousSkipFirstFrame;
-			}
-			else
-			{
-				DataReadMode = EBufferReadMode::Asynchronous;
-			}
-			const bool bLooped = ReadMorePCMRTData(CurrentBuffer, DataReadMode);
-
-			// If this was a synchronous read, then immediately write it
-			if (AsyncRealtimeAudioTask == nullptr)
-			{
-				SubmitRealTimeSourceData(bLooped, bSubmitSynchronously);
-			}
-		}
-	}
-
 	void FMixerSource::OnBeginGenerate()
 	{
-		if (MixerBuffer)
-		{
-			MixerBuffer->OnBeginGenerate();
-		}
 	}
 
-	void FMixerSource::OnSourceBufferEnd()
+	void FMixerSource::OnRelease()
 	{
-		FScopeLock Lock(&RenderThreadCritSect);
-
-		if (Playing && MixerSourceVoice)
-		{
-			int32 BuffersQueued = MixerSourceVoice->GetNumBuffersQueued();
-
-			if (BuffersQueued == 0 && bBuffersToFlush)
-			{
-				bIsFinished = true;
-			}
-			else if (!bBuffersToFlush && BuffersQueued <= (Audio::MAX_BUFFERS_QUEUED - 1))
-			{
-				// OnSourceBufferEnd is always called from render thread and the source needs to be 
-				// processed and any decoded buffers submitted to render thread synchronously			
-				const bool bSubmitSynchronously = true;
-				ProcessRealTimeSource(BuffersQueued < (Audio::MAX_BUFFERS_QUEUED - 1), bSubmitSynchronously);
-			}
-		}
+		bIsReleasing = false;
 	}
 
-	void FMixerSource::OnRelease(TArray<FPendingReleaseData*>& OutPendingReleaseData)
+	void FMixerSource::OnDone()
 	{
-		FPendingReleaseData* PendingReleaseData = nullptr;
-
-		while (PendingReleases.Dequeue(PendingReleaseData))
-		{
-			OutPendingReleaseData.Add(PendingReleaseData);
-		}
-
-		bIsReleasing = false;
+		bIsFinished = true;
 	}
 
 	void FMixerSource::FreeResources()
@@ -912,17 +622,12 @@ namespace Audio
 			MixerBuffer->EnsureHeaderParseTaskFinished();
 		}
 
-		FPendingReleaseData* PendingDecodeTask = nullptr;
+		check(!bIsStopping);
+		check(!Playing);
 
+		// Make a new pending release data ptr to pass off release data
 		if (MixerSourceVoice)
 		{
-			// Hand off the ptr of the async task so it can be shutdown on the audio render thread.
-			if (AsyncRealtimeAudioTask)
-			{
-				PendingDecodeTask = new FPendingReleaseData();
-				PendingDecodeTask->Task = AsyncRealtimeAudioTask;
-			}
-
 			// We're now "releasing" so don't recycle this voice until we get notified that the source has finished
 			bIsReleasing = true;
 
@@ -931,36 +636,9 @@ namespace Audio
 			MixerSourceVoice = nullptr;
 		}
 
-		if (bResourcesNeedFreeing)
-		{
-			// If we have a buffer, we can't delete until the async decoding task has been ensured to complete.
-			if (Buffer)
-			{
-				check(Buffer->ResourceID == 0);
-
-				if (!PendingDecodeTask)
-				{
-					PendingDecodeTask = new FPendingReleaseData();
-				}
-
-				PendingDecodeTask->Buffer = Buffer;
-			}
-
-			CurrentBuffer = 0;
-		}
-
-		if (PendingDecodeTask)
-		{
-			PendingReleases.Enqueue(PendingDecodeTask);
-		}
-
-		// Reset the raw PCM buffer data
-		RawPCMDataBuffer = FRawPCMDataBuffer();
-
+		MixerSourceBuffer = nullptr;
 		MixerBuffer = nullptr;
-		AsyncRealtimeAudioTask = nullptr;
 		Buffer = nullptr;
-		bBuffersToFlush = false;
 		bLoopCallback = false;
 		bResourcesNeedFreeing = false;
 		NumTotalFrames = 0;
@@ -1075,14 +753,18 @@ namespace Audio
 			// Send the source audio to the reverb plugin if enabled
 			if (UseReverbPlugin())
 			{
-				FMixerSubmixPtr MasterReverbPluginSubmix = MixerDevice->GetMasterReverbPluginSubmix();
-				MixerSourceVoice->SetSubmixSendInfo(MasterReverbPluginSubmix, ReverbSendLevel);
+				if (MixerDevice->GetMasterReverbPluginSubmix().IsValid())
+				{
+					MixerSourceVoice->SetSubmixSendInfo(MixerDevice->GetMasterReverbPluginSubmix(), ReverbSendLevel);
+				}
 			}
 			else 
 			{
 				// Send the source audio to the master reverb
-				FMixerSubmixPtr MasterReverbSubmix = MixerDevice->GetMasterReverbSubmix();
-				MixerSourceVoice->SetSubmixSendInfo(MasterReverbSubmix, ReverbSendLevel);
+				if (MixerDevice->GetMasterReverbSubmix().IsValid())
+				{
+					MixerSourceVoice->SetSubmixSendInfo(MixerDevice->GetMasterReverbSubmix(), ReverbSendLevel);
+				}
 			}
 		}
 
@@ -1130,7 +812,7 @@ namespace Audio
 		}
 	}
 
-	bool FMixerSource::ComputeMonoChannelMap(const ESubmixChannelFormat SubmixChannelType, TArray<float>& OutChannelMap)
+	bool FMixerSource::ComputeMonoChannelMap(const ESubmixChannelFormat SubmixChannelType, Audio::AlignedFloatBuffer& OutChannelMap)
 	{
 		if (UseObjectBasedSpatialization())
 		{
@@ -1162,7 +844,7 @@ namespace Audio
 		return false;
 	}
 
-	bool FMixerSource::ComputeStereoChannelMap(const ESubmixChannelFormat InSubmixChannelType, TArray<float>& OutChannelMap)
+	bool FMixerSource::ComputeStereoChannelMap(const ESubmixChannelFormat InSubmixChannelType, Audio::AlignedFloatBuffer& OutChannelMap)
 	{
 		if (!UseObjectBasedSpatialization() && WaveInstance->bUseSpatialization && (!FMath::IsNearlyEqual(WaveInstance->AbsoluteAzimuth, PreviousAzimuth, 0.01f) || MixerSourceVoice->NeedsSpeakerMap()))
 		{
@@ -1205,7 +887,7 @@ namespace Audio
 		return false;
 	}
 
-	bool FMixerSource::ComputeChannelMap(const ESubmixChannelFormat InSubmixChannelType, const int32 NumSourceChannels, TArray<float>& OutChannelMap)
+	bool FMixerSource::ComputeChannelMap(const ESubmixChannelFormat InSubmixChannelType, const int32 NumSourceChannels, Audio::AlignedFloatBuffer& OutChannelMap)
 	{
 		if (NumSourceChannels == 1)
 		{
@@ -1252,52 +934,4 @@ namespace Audio
 				WaveInstance->ReverbPluginSettings != nullptr;
 	}
 
-	bool FRawPCMDataBuffer::GetNextBuffer(FMixerSourceBufferPtr OutSourceBufferPtr, const uint32 NumSampleToGet)
-	{
-		// TODO: support loop counts
-		float* OutBufferPtr = OutSourceBufferPtr->AudioData.GetData();
-		int16* DataPtr = (int16*)Data;
-
-		if (LoopCount == Audio::LOOP_FOREVER)
-		{
-			bool bLooped = false;
-			for (uint32 Sample = 0; Sample < NumSampleToGet; ++Sample)
-			{
-				OutBufferPtr[Sample] = DataPtr[CurrentSample++] / 32768.0f;
-
-				// Loop around if we're looping
-				if (CurrentSample >= NumSamples)
-				{
-					CurrentSample = 0;
-					bLooped = true;
-				}
-			}
-			return bLooped;
-		}
-		else if (CurrentSample < NumSamples)
-		{
-			uint32 Sample = 0;
-			while (Sample < NumSampleToGet && CurrentSample < NumSamples)
-			{
-				OutBufferPtr[Sample++] = (float)DataPtr[CurrentSample++] / 32768.0f;
-			}
-
-			// Zero out the rest of the buffer
-			while (Sample < NumSampleToGet)
-			{
-				OutBufferPtr[Sample++] = 0.0f;
-			}
-		}
-		else
-		{
-			for (uint32 Sample = 0; Sample < NumSampleToGet; ++Sample)
-			{
-				OutBufferPtr[Sample] = 0.0f;
-			}
-		}
-
-		// If the current sample is greater or equal to num samples we hit the end of the buffer
-		return CurrentSample >= NumSamples;
-	}
-
 }
diff --git a/Engine/Source/Runtime/AudioMixer/Private/AudioMixerSource.h b/Engine/Source/Runtime/AudioMixer/Private/AudioMixerSource.h
--- a/Engine/Source/Runtime/AudioMixer/Private/AudioMixerSource.h
+++ b/Engine/Source/Runtime/AudioMixer/Private/AudioMixerSource.h
@@ -12,37 +12,14 @@ namespace Audio
 	class FMixerSourceVoice;
 	class FMixerSource;
 	class FMixerBuffer;
-	class ISourceBufferQueueListener;
-
-	static const int32 MAX_BUFFERS_QUEUED = 3;
-	static const int32 LOOP_FOREVER = -1;
-
-	struct FRawPCMDataBuffer
-	{
-		uint8* Data;
-		uint32 DataSize;
-		int32 LoopCount;
-		uint32 CurrentSample;
-		uint32 NumSamples;
-
-		bool GetNextBuffer(FMixerSourceBufferPtr OutSourceBufferPtr, const uint32 NumSampleToGet);
-
-
-		FRawPCMDataBuffer()
-			: Data(nullptr)
-			, DataSize(0)
-			, LoopCount(0)
-			, CurrentSample(0)
-			, NumSamples(0)
-		{}
-	};
+	class ISourceListener;
 
 	/** 
 	 * FMixerSource
 	 * Class which implements a sound source object for the audio mixer module.
 	 */
 	class FMixerSource :	public FSoundSource, 
-							public ISourceBufferQueueListener
+							public ISourceListener
 	{
 	public:
 		/** Constructor. */
@@ -60,8 +37,6 @@ namespace Audio
 		virtual void Stop() override;
 		virtual void StopNow() override;
 		virtual bool IsStopping() override { return bIsStopping; }
-		virtual void EnsureStopped() override;
-		virtual void ReleaseStoppedSound() override;
 		virtual void Pause() override;
 		virtual bool IsFinished() override;
 		virtual FString Describe(bool bUseLongName) override;
@@ -69,45 +44,15 @@ namespace Audio
 		virtual float GetEnvelopeValue() const override;
 		//~ End FSoundSource Interface
 
-		//~Begin ISourceBufferQueueListener
+		//~ Begin ISourceListneer
 		virtual void OnBeginGenerate() override;
-		virtual void OnSourceBufferEnd() override;
-		virtual void OnRelease(TArray<FPendingReleaseData*>& OutPendingReleaseData) override;
-		//~End ISourceBufferQueueListener
+		virtual void OnRelease() override;
+		virtual void OnDone() override;
+		virtual void OnLoopEnd() override { bLoopCallback = true; };
+		//~ End ISourceListneer
 
 	private:
 
-		/** Enum describing the data-read mode of an audio buffer. */
-		enum class EBufferReadMode : uint8
-		{
-			/** Read the next buffer synchronously. */
-			Synchronous,
-
-			/** Read the next buffer asynchronously. */
-			Asynchronous,
-
-			/** Read the next buffer asynchronously but skip the first chunk of audio. */
-			AsynchronousSkipFirstFrame
-		};
-
-		/** Internal stop function. */
-		void StopInternal(bool bStopNow);
-
-		/** Submit the current decoded PCM buffer to the contained source voice. */
-		void SubmitPCMBuffers();
-
-		/** Submit the current decoded PCMRT (PCM RealTime) buffer to the contained source voice. */
-		void SubmitPCMRTBuffers();
-
-		/** Reads more PCMRT data for the real-time decoded audio buffers. */
-		bool ReadMorePCMRTData(const int32 BufferIndex, EBufferReadMode BufferReadMode, bool* OutLooped = nullptr);
-
-		/** Called when a buffer finishes for a real-time source and more buffers need to be read and submitted. */
-		void ProcessRealTimeSource(const bool bBlockForData, const bool bOnRenderThread);
-
-		/** Submits new real-time decoded buffers to a source voice. */
-		void SubmitRealTimeSourceData(const bool bLooped, const bool bOnRenderThread);
-
 		/** Frees any resources for this sound source. */
 		void FreeResources();
 
@@ -127,13 +72,13 @@ namespace Audio
 		void UpdateChannelMaps();
 
 		/** Computes the mono-channel map. */
-		bool ComputeMonoChannelMap(const ESubmixChannelFormat SubmixChannelType, TArray<float>& OutChannelMap);
+		bool ComputeMonoChannelMap(const ESubmixChannelFormat SubmixChannelType, Audio::AlignedFloatBuffer& OutChannelMap);
 
 		/** Computes the stereo-channel map. */
-		bool ComputeStereoChannelMap(const ESubmixChannelFormat SubmixChannelType, TArray<float>& OutChannelMap);
+		bool ComputeStereoChannelMap(const ESubmixChannelFormat SubmixChannelType, Audio::AlignedFloatBuffer& OutChannelMap);
 
 		/** Compute the channel map based on the number of output and source channels. */
-		bool ComputeChannelMap(const ESubmixChannelFormat SubmixChannelType, const int32 NumSourceChannels, TArray<float>& OutChannelMap);
+		bool ComputeChannelMap(const ESubmixChannelFormat SubmixChannelType, const int32 NumSourceChannels, Audio::AlignedFloatBuffer& OutChannelMap);
 
 		/** Whether or not we should create the source voice with the HRTF spatializer. */
 		bool UseObjectBasedSpatialization() const;
@@ -151,17 +96,12 @@ namespace Audio
 
 		FMixerDevice* MixerDevice;
 		FMixerBuffer* MixerBuffer;
+		TSharedPtr<FMixerSourceBuffer> MixerSourceBuffer;
 		FMixerSourceVoice* MixerSourceVoice;
-		IAudioTask* AsyncRealtimeAudioTask;
-
-		// Queue of pending release data. Pushed from audio thread, updated on audio render thread.
-		TQueue<FPendingReleaseData*> PendingReleases;
-
-		FCriticalSection RenderThreadCritSect;
 
 		struct FChannelMapInfo
 		{
-			TArray<float> ChannelMap;
+			Audio::AlignedFloatBuffer ChannelMap;
 			bool bUsed;
 
 			FChannelMapInfo()
@@ -172,17 +112,8 @@ namespace Audio
 		// Mapping of channel map types to channel maps. Determined by what submixes this source sends its audio to.
 		FChannelMapInfo ChannelMaps[(int32) ESubmixChannelFormat::Count];
 
-		int32 CurrentBuffer;
 		float PreviousAzimuth;
 
-		// The decoded source buffers are using a shared pointer because the audio mixer thread 
-		// will need to have a ref while playing back the sound. There is a small window where a 
-		// flushed/destroyed sound can destroy the buffer before the sound finishes using the buffer.
-		TArray<FMixerSourceBufferPtr> SourceVoiceBuffers;
-
-		// Raw uncompressed, non-float PCM data (int16)
-		FRawPCMDataBuffer RawPCMDataBuffer;
-
 		FSpatializationParams SpatializationParams;
 
 		FThreadSafeBool bPlayedCachedBuffer;
@@ -191,7 +122,6 @@ namespace Audio
 		FThreadSafeBool bLoopCallback;
 		FThreadSafeBool bIsFinished;
 		FThreadSafeBool bIsPlayingEffectTails;
-		FThreadSafeBool bBuffersToFlush;
 		FThreadSafeBool bFreeAsyncTask;
 
 		// Whether or not we're currently releasing our resources. Prevents recycling the source until release is finished.
@@ -203,5 +133,6 @@ namespace Audio
 		uint32 bIs3D : 1;
 		uint32 bDebugMode : 1;
 		uint32 bIsVorbis : 1;
+		uint32 bIsStoppingVoicesEnabled : 1;
 	};
 }
diff --git a/Engine/Source/Runtime/AudioMixer/Private/AudioMixerSourceBuffer.cpp b/Engine/Source/Runtime/AudioMixer/Private/AudioMixerSourceBuffer.cpp
new file mode 100644
--- /dev/null
+++ b/Engine/Source/Runtime/AudioMixer/Private/AudioMixerSourceBuffer.cpp
@@ -0,0 +1,419 @@
+// Copyright 1998-2018 Epic Games, Inc. All Rights Reserved.
+
+#include "AudioMixerSourceBuffer.h"
+
+namespace Audio
+{
+	bool FRawPCMDataBuffer::GetNextBuffer(FMixerSourceVoiceBuffer* OutSourceBufferPtr, const uint32 NumSampleToGet)
+	{
+		// TODO: support loop counts
+		float* OutBufferPtr = OutSourceBufferPtr->AudioData.GetData();
+		int16* DataPtr = (int16*)Data;
+
+		if (LoopCount == Audio::LOOP_FOREVER)
+		{
+			bool bLooped = false;
+			for (uint32 Sample = 0; Sample < NumSampleToGet; ++Sample)
+			{
+				OutBufferPtr[Sample] = DataPtr[CurrentSample++] / 32768.0f;
+
+				// Loop around if we're looping
+				if (CurrentSample >= NumSamples)
+				{
+					CurrentSample = 0;
+					bLooped = true;
+				}
+			}
+			return bLooped;
+		}
+		else if (CurrentSample < NumSamples)
+		{
+			uint32 Sample = 0;
+			while (Sample < NumSampleToGet && CurrentSample < NumSamples)
+			{
+				OutBufferPtr[Sample++] = (float)DataPtr[CurrentSample++] / 32768.0f;
+			}
+
+			// Zero out the rest of the buffer
+			while (Sample < NumSampleToGet)
+			{
+				OutBufferPtr[Sample++] = 0.0f;
+			}
+		}
+		else
+		{
+			for (uint32 Sample = 0; Sample < NumSampleToGet; ++Sample)
+			{
+				OutBufferPtr[Sample] = 0.0f;
+			}
+		}
+
+		// If the current sample is greater or equal to num samples we hit the end of the buffer
+		return CurrentSample >= NumSamples;
+	}
+
+	FMixerSourceBuffer::FMixerSourceBuffer()
+		: NumBuffersQeueued(0)
+		, CurrentBuffer(0)
+		, MixerBuffer(nullptr)
+		, SoundWave(nullptr)
+		, AsyncRealtimeAudioTask(nullptr)
+		, LoopingMode(ELoopingMode::LOOP_Never)
+		, bBufferFinished(false)
+		, bPlayedCachedBuffer(false)
+		, bIsSeeking(false)
+		, bLoopCallback(false)
+	{
+	}
+
+	FMixerSourceBuffer::~FMixerSourceBuffer()
+	{
+		OnEndGenerate();
+	}
+
+	bool FMixerSourceBuffer::PreInit(FMixerBuffer* InBuffer, USoundWave* InWave, ELoopingMode InLoopingMode, bool bInIsSeeking)
+	{
+		MixerBuffer = InBuffer;
+		check(MixerBuffer);
+
+		// May or may not be nullptr
+		SoundWave = InWave;
+
+		if (SoundWave->bProcedural)
+		{
+			if (SoundWave->bIsSoundActive)
+			{
+				UE_LOG(LogAudioMixer, Warning, TEXT("Procedural sound wave is reinitializing even though it is currently actively generating audio. Please stop sound before trying to play it again."));
+				return false;
+			}
+	
+			// We flag that this sound wave is active for the lifetime of this object since we use it for decoding, etc.
+			SoundWave->bIsSoundActive = true;
+		}
+
+		LoopingMode = InLoopingMode;
+		bIsSeeking = bInIsSeeking;
+		bLoopCallback = false;
+
+		BufferQueue.Empty();
+
+		const uint32 TotalSamples = MONO_PCM_BUFFER_SAMPLES * MixerBuffer->NumChannels;
+		for (int32 BufferIndex = 0; BufferIndex < Audio::MAX_BUFFERS_QUEUED; ++BufferIndex)
+		{
+			SourceVoiceBuffers.Add(TSharedPtr<FMixerSourceVoiceBuffer>(new FMixerSourceVoiceBuffer()));
+
+			SourceVoiceBuffers[BufferIndex]->AudioData.AddZeroed(TotalSamples);
+			SourceVoiceBuffers[BufferIndex]->Samples = TotalSamples;
+			SourceVoiceBuffers[BufferIndex]->bRealTimeBuffer = true;
+			SourceVoiceBuffers[BufferIndex]->LoopCount = 0;
+		}
+		return true;
+	}
+
+	void FMixerSourceBuffer::Init()
+	{
+		const EBufferType::Type BufferType = MixerBuffer->GetType();
+		switch (BufferType)
+		{
+			case EBufferType::PCM:
+			case EBufferType::PCMPreview:
+				SubmitInitialPCMBuffers();
+				break;
+
+			case EBufferType::PCMRealTime:
+			case EBufferType::Streaming:
+				SubmitInitialRealtimeBuffers();
+				break;
+
+			case EBufferType::Invalid:
+				break;
+		}
+	}
+
+	void FMixerSourceBuffer::OnBufferEnd()
+	{
+		if ((NumBuffersQeueued == 0 && bBufferFinished) || !SoundWave)
+		{
+			return;
+		}
+
+		ProcessRealTimeSource();
+	}
+
+	TSharedPtr<FMixerSourceVoiceBuffer> FMixerSourceBuffer::GetNextBuffer()
+	{
+		TSharedPtr<FMixerSourceVoiceBuffer> NewBufferPtr;
+		BufferQueue.Dequeue(NewBufferPtr);
+		--NumBuffersQeueued;
+		return NewBufferPtr;
+	}
+
+	void FMixerSourceBuffer::SubmitInitialPCMBuffers()
+	{
+		CurrentBuffer = 0;
+
+		RawPCMDataBuffer.Data = nullptr;
+		RawPCMDataBuffer.DataSize = 0;
+		MixerBuffer->GetPCMData(&RawPCMDataBuffer.Data, &RawPCMDataBuffer.DataSize);
+
+		RawPCMDataBuffer.NumSamples = RawPCMDataBuffer.DataSize / sizeof(int16);
+		RawPCMDataBuffer.CurrentSample = 0;
+
+		// Only submit data if we've successfully loaded it
+		if (!RawPCMDataBuffer.Data || !RawPCMDataBuffer.DataSize)
+		{
+			return;
+		}
+
+		RawPCMDataBuffer.LoopCount = (LoopingMode != LOOP_Never) ? Audio::LOOP_FOREVER : 0;
+
+		// Submit the first two format-converted chunks to the source voice
+		const uint32 NumSamplesPerBuffer = MONO_PCM_BUFFER_SAMPLES * MixerBuffer->NumChannels;
+		int16* RawPCMBufferDataPtr = (int16*)RawPCMDataBuffer.Data;
+
+		RawPCMDataBuffer.GetNextBuffer(SourceVoiceBuffers[0].Get(), NumSamplesPerBuffer);
+
+		SubmitBuffer(SourceVoiceBuffers[0]);
+
+		CurrentBuffer = 1;
+	}
+
+	void FMixerSourceBuffer::SubmitInitialRealtimeBuffers()
+	{
+		CurrentBuffer = 0;
+
+		bPlayedCachedBuffer = false;
+		if (!bIsSeeking && SoundWave && SoundWave->CachedRealtimeFirstBuffer)
+		{
+			bPlayedCachedBuffer = true;
+
+			// Format convert the first cached buffers
+			const uint32 NumSamples = MONO_PCM_BUFFER_SAMPLES * MixerBuffer->NumChannels;
+			const uint32 BufferSize = MONO_PCM_BUFFER_SIZE * MixerBuffer->NumChannels;
+
+			int16* CachedBufferPtr0 = (int16*)SoundWave->CachedRealtimeFirstBuffer;
+			int16* CachedBufferPtr1 = (int16*)(SoundWave->CachedRealtimeFirstBuffer + BufferSize);
+			float* AudioData0 = SourceVoiceBuffers[0]->AudioData.GetData();
+			float* AudioData1 = SourceVoiceBuffers[1]->AudioData.GetData();
+			for (uint32 Sample = 0; Sample < NumSamples; ++Sample)
+			{
+				AudioData0[Sample] = CachedBufferPtr0[Sample] / 32768.0f;
+				AudioData1[Sample] = CachedBufferPtr1[Sample] / 32768.0f;
+			}
+
+			// Submit the already decoded and cached audio buffers
+			SubmitBuffer(SourceVoiceBuffers[0]);
+			SubmitBuffer(SourceVoiceBuffers[1]);
+
+			CurrentBuffer = 2;
+		}
+		else if (SoundWave && !SoundWave->bIsBus)
+		{
+			// We should have already kicked off and finished a task. 
+			check(AsyncRealtimeAudioTask != nullptr);
+
+			ProcessRealTimeSource();
+		}
+	}
+
+	bool FMixerSourceBuffer::ReadMoreRealtimeData(const int32 BufferIndex, EBufferReadMode BufferReadMode)
+	{
+		if (SoundWave && SoundWave->bProcedural)
+		{
+			const int32 MaxSamples = MONO_PCM_BUFFER_SAMPLES * MixerBuffer->NumChannels;
+
+			FProceduralAudioTaskData NewTaskData;
+			NewTaskData.ProceduralSoundWave = SoundWave;
+			NewTaskData.AudioData = SourceVoiceBuffers[BufferIndex]->AudioData.GetData();
+			NewTaskData.NumSamples = MaxSamples;
+			NewTaskData.NumChannels = MixerBuffer->NumChannels;
+			check(!AsyncRealtimeAudioTask);
+			AsyncRealtimeAudioTask = CreateAudioTask(NewTaskData);
+
+			// Procedural sound waves never loop
+			return false;
+		}
+		else if (!MixerBuffer->IsRealTimeBuffer())
+		{
+			check(RawPCMDataBuffer.Data != nullptr);
+
+			// Read the next raw PCM buffer into the source buffer index. This converts raw PCM to float.
+			const uint32 NumSamplesPerBuffer = MONO_PCM_BUFFER_SAMPLES * MixerBuffer->NumChannels;
+			return RawPCMDataBuffer.GetNextBuffer(SourceVoiceBuffers[BufferIndex].Get(), NumSamplesPerBuffer);
+		}
+
+		FDecodeAudioTaskData NewTaskData;
+		NewTaskData.MixerBuffer = MixerBuffer;
+		NewTaskData.AudioData = SourceVoiceBuffers[BufferIndex]->AudioData.GetData();
+		NewTaskData.bLoopingMode = LoopingMode != LOOP_Never;
+		NewTaskData.bSkipFirstBuffer = (BufferReadMode == EBufferReadMode::AsynchronousSkipFirstFrame);
+		NewTaskData.NumFramesToDecode = MONO_PCM_BUFFER_SAMPLES;
+
+		check(!AsyncRealtimeAudioTask);
+		AsyncRealtimeAudioTask = CreateAudioTask(NewTaskData);
+
+		return false;
+	}
+
+	void FMixerSourceBuffer::SubmitRealTimeSourceData(const bool bLooped)
+	{
+		// Have we reached the end of the sound
+		if (bLooped)
+		{
+			switch (LoopingMode)
+			{
+			case LOOP_Never:
+				// Play out any queued buffers - once there are no buffers left, the state check at the beginning of IsFinished will fire
+				bBufferFinished = true;
+				break;
+
+			case LOOP_WithNotification:
+				// If we have just looped, and we are looping, send notification
+				// This will trigger a WaveInstance->NotifyFinished() in the FXAudio2SoundSournce::IsFinished() function on main thread.
+				bLoopCallback = true;
+				break;
+
+			case LOOP_Forever:
+				// Let the sound loop indefinitely
+				break;
+			}
+		}
+
+		if (SourceVoiceBuffers[CurrentBuffer]->AudioData.Num() > 0)
+		{
+			SubmitBuffer(SourceVoiceBuffers[CurrentBuffer]);
+		}
+	}
+
+	void FMixerSourceBuffer::ProcessRealTimeSource()
+	{
+		if (AsyncRealtimeAudioTask)
+		{
+			AsyncRealtimeAudioTask->EnsureCompletion();
+
+			bool bLooped = false;
+
+			switch (AsyncRealtimeAudioTask->GetType())
+			{
+			case EAudioTaskType::Decode:
+			{
+				FDecodeAudioTaskResults TaskResult;
+				AsyncRealtimeAudioTask->GetResult(TaskResult);
+
+				SourceVoiceBuffers[CurrentBuffer]->Samples = MONO_PCM_BUFFER_SAMPLES * MixerBuffer->NumChannels;
+				bLooped = TaskResult.bLooped;
+			}
+			break;
+
+			case EAudioTaskType::Procedural:
+			{
+				FProceduralAudioTaskResults TaskResult;
+				AsyncRealtimeAudioTask->GetResult(TaskResult);
+
+				SourceVoiceBuffers[CurrentBuffer]->Samples = TaskResult.NumSamplesWritten;
+			}
+			break;
+			}
+
+			delete AsyncRealtimeAudioTask;
+			AsyncRealtimeAudioTask = nullptr;
+
+			SubmitRealTimeSourceData(bLooped);
+		}
+
+		if (!AsyncRealtimeAudioTask)
+		{
+			// Update the buffer index
+			if (++CurrentBuffer > 2)
+			{
+				CurrentBuffer = 0;
+			}
+
+			EBufferReadMode DataReadMode;
+			if (bPlayedCachedBuffer)
+			{
+				bPlayedCachedBuffer = false;
+				DataReadMode = EBufferReadMode::AsynchronousSkipFirstFrame;
+			}
+			else
+			{
+				DataReadMode = EBufferReadMode::Asynchronous;
+			}
+
+			const bool bLooped = ReadMoreRealtimeData(CurrentBuffer, DataReadMode);
+
+			// If this was a synchronous read, then immediately write it
+			if (AsyncRealtimeAudioTask == nullptr)
+			{
+				SubmitRealTimeSourceData(bLooped);
+			}
+		}
+	}
+
+	void FMixerSourceBuffer::SubmitBuffer(TSharedPtr<FMixerSourceVoiceBuffer> InSourceVoiceBuffer)
+	{
+		NumBuffersQeueued++;
+		BufferQueue.Enqueue(InSourceVoiceBuffer);
+	}
+
+	bool FMixerSourceBuffer::IsAsyncTaskInProgress() const
+	{ 
+		return AsyncRealtimeAudioTask != nullptr; 
+	}
+
+	bool FMixerSourceBuffer::IsAsyncTaskDone() const 
+	{
+		if (AsyncRealtimeAudioTask)
+		{
+			return AsyncRealtimeAudioTask->IsDone();
+		}
+		return true; 
+	}
+	
+	void FMixerSourceBuffer::EnsureAsyncTaskFinishes() 
+	{
+		if (AsyncRealtimeAudioTask) 
+		{ 
+			AsyncRealtimeAudioTask->EnsureCompletion(); 
+		}
+	}
+
+	bool FMixerSourceBuffer::IsBeginDestroy()
+	{
+		return SoundWave && SoundWave->bIsBeginDestroy;
+	}
+
+	void FMixerSourceBuffer::ClearSoundWave()
+	{
+		// Call on end generate right now, before destructor
+		OnEndGenerate();
+	}
+
+	void FMixerSourceBuffer::OnBeginGenerate()
+	{
+		if (SoundWave && SoundWave->bProcedural)
+		{
+			SoundWave->bIsSoundActive = true;
+
+			SoundWave->OnBeginGenerate();
+		}
+	}
+
+	void FMixerSourceBuffer::OnEndGenerate()
+	{
+		// Make sure the async task finishes!
+		EnsureAsyncTaskFinishes();
+
+		if (SoundWave)
+		{
+			if (SoundWave->bProcedural)
+			{
+				SoundWave->OnEndGenerate();
+				SoundWave->bIsSoundActive = false;
+			}
+
+			SoundWave = nullptr;
+		}
+	}
+
+}
diff --git a/Engine/Source/Runtime/AudioMixer/Private/AudioMixerSourceBuffer.h b/Engine/Source/Runtime/AudioMixer/Private/AudioMixerSourceBuffer.h
new file mode 100644
--- /dev/null
+++ b/Engine/Source/Runtime/AudioMixer/Private/AudioMixerSourceBuffer.h
@@ -0,0 +1,117 @@
+// Copyright 1998-2018 Epic Games, Inc. All Rights Reserved.
+
+#pragma once
+
+#include "CoreMinimal.h"
+#include "AudioMixerBuffer.h"
+#include "AudioMixerSourceManager.h"
+
+namespace Audio
+{
+	struct FMixerSourceVoiceBuffer;
+
+	static const int32 MAX_BUFFERS_QUEUED = 3;
+	static const int32 LOOP_FOREVER = -1;
+
+	struct FRawPCMDataBuffer
+	{
+		uint8* Data;
+		uint32 DataSize;
+		int32 LoopCount;
+		uint32 CurrentSample;
+		uint32 NumSamples;
+
+		bool GetNextBuffer(FMixerSourceVoiceBuffer* OutSourceBufferPtr, const uint32 NumSampleToGet);
+
+		FRawPCMDataBuffer()
+			: Data(nullptr)
+			, DataSize(0)
+			, LoopCount(0)
+			, CurrentSample(0)
+			, NumSamples(0)
+		{}
+	};
+
+	/** Enum describing the data-read mode of an audio buffer. */
+	enum class EBufferReadMode : uint8
+	{
+		/** Read the next buffer asynchronously. */
+		Asynchronous,
+
+		/** Read the next buffer asynchronously but skip the first chunk of audio. */
+		AsynchronousSkipFirstFrame
+	};
+
+	/** Class which handles decoding audio for a particular source buffer. */
+	class FMixerSourceBuffer
+	{
+	public:
+		FMixerSourceBuffer();
+		~FMixerSourceBuffer();
+
+		bool PreInit(FMixerBuffer* InBuffer, USoundWave* InWave, ELoopingMode InLoopingMode, bool bInIsSeeking);
+		void Init();
+
+		// Called by source manager when needing more buffers
+		void OnBufferEnd();
+
+		// Return the number of buffers enqueued on the mixer source buffer
+		int32 GetNumBuffersQueued() const { return NumBuffersQeueued; }
+		
+		// Returns the next enqueued buffer, returns nullptr if no buffers enqueued
+		TSharedPtr<FMixerSourceVoiceBuffer> GetNextBuffer();
+
+		// Returns if buffer looped
+		bool DidBufferLoop() const { return bLoopCallback; }
+
+		// Returns true if buffer finished
+		bool DidBufferFinish() const { return bBufferFinished; }
+
+		// Called to start an async task to read more data
+		bool ReadMoreRealtimeData(const int32 BufferIndex, EBufferReadMode BufferReadMode);
+
+		// Returns true if async task is in progress
+		bool IsAsyncTaskInProgress() const;
+
+		// Returns true if the async task is done
+		bool IsAsyncTaskDone() const;
+
+		// Ensures the async task finishes
+		void EnsureAsyncTaskFinishes();
+
+		// Checks if sound wave is flagged begin destroy
+		bool IsBeginDestroy();
+
+		// Clear the sound wave reference
+		void ClearSoundWave();
+
+		// Begin and end generation on the audio render thread (audio mixer only)
+		void OnBeginGenerate();
+		void OnEndGenerate();
+
+	private:
+
+		void SubmitInitialPCMBuffers();
+		void SubmitInitialRealtimeBuffers();
+		void SubmitRealTimeSourceData(const bool bLooped);
+		void ProcessRealTimeSource();
+		void SubmitBuffer(TSharedPtr<FMixerSourceVoiceBuffer> InSourceVoiceBuffer);
+
+
+		int32 NumBuffersQeueued;
+		FRawPCMDataBuffer RawPCMDataBuffer;
+
+		TArray<TSharedPtr<FMixerSourceVoiceBuffer>> SourceVoiceBuffers;
+		TQueue<TSharedPtr<FMixerSourceVoiceBuffer>> BufferQueue;
+		int32 CurrentBuffer;
+		FMixerBuffer* MixerBuffer;
+		USoundWave* SoundWave;
+		IAudioTask* AsyncRealtimeAudioTask;
+		ELoopingMode LoopingMode;
+		bool bBufferFinished;
+		bool bPlayedCachedBuffer;
+		bool bIsSeeking;
+		bool bLoopCallback;
+
+	};
+}
diff --git a/Engine/Source/Runtime/AudioMixer/Private/AudioMixerSourceDecode.cpp b/Engine/Source/Runtime/AudioMixer/Private/AudioMixerSourceDecode.cpp
--- a/Engine/Source/Runtime/AudioMixer/Private/AudioMixerSourceDecode.cpp
+++ b/Engine/Source/Runtime/AudioMixer/Private/AudioMixerSourceDecode.cpp
@@ -46,6 +46,13 @@ public:
 		{
 			case EAudioTaskType::Procedural:
 			{
+				// Make sure we've been flagged as active
+				if (!ProceduralTaskData.ProceduralSoundWave->bIsSoundActive)
+				{
+					UE_LOG(LogAudioMixer, Warning, TEXT("Warning: Attempted to run async procedural task on inactive sound. Aborting."));
+					return;
+				}
+
 				// If we're not a float format, we need to convert the format to float
 				const EAudioMixerStreamDataFormat::Type FormatType = ProceduralTaskData.ProceduralSoundWave->GetGeneratedPCMDataFormat();
 				if (FormatType != EAudioMixerStreamDataFormat::Float)
diff --git a/Engine/Source/Runtime/AudioMixer/Private/AudioMixerSourceDecode.h b/Engine/Source/Runtime/AudioMixer/Private/AudioMixerSourceDecode.h
--- a/Engine/Source/Runtime/AudioMixer/Private/AudioMixerSourceDecode.h
+++ b/Engine/Source/Runtime/AudioMixer/Private/AudioMixerSourceDecode.h
@@ -24,6 +24,13 @@ namespace Audio
 
 		// The number of channels of the procedural buffer
 		int32 NumChannels;
+
+		FProceduralAudioTaskData()
+			: ProceduralSoundWave(nullptr)
+			, AudioData(nullptr)
+			, NumSamples(0)
+			, NumChannels(0)
+		{}
 	};
 
 	// Data needed for a decode audio task
@@ -43,6 +50,14 @@ namespace Audio
 
 		// Whether or not to skip the first buffer
 		bool bSkipFirstBuffer;
+
+		FDecodeAudioTaskData()
+			: MixerBuffer(nullptr)
+			, AudioData(nullptr)
+			, NumFramesToDecode(0)
+			, bLoopingMode(false)
+			, bSkipFirstBuffer(false)
+		{}
 	};
 
 	// Data needed for a header parse audio task
@@ -53,6 +68,11 @@ namespace Audio
 
 		// The sound wave object which contains the encoded file
 		USoundWave* SoundWave;
+
+		FHeaderParseAudioTaskData()
+			: MixerBuffer(nullptr)
+			, SoundWave(nullptr)
+		{}
 	};
 
 	// Results from procedural audio task
diff --git a/Engine/Source/Runtime/AudioMixer/Private/AudioMixerSourceManager.cpp b/Engine/Source/Runtime/AudioMixer/Private/AudioMixerSourceManager.cpp
--- a/Engine/Source/Runtime/AudioMixer/Private/AudioMixerSourceManager.cpp
+++ b/Engine/Source/Runtime/AudioMixer/Private/AudioMixerSourceManager.cpp
@@ -18,6 +18,53 @@ FAutoConsoleVariableRef CVarDisableParallelSourceProcessing(
 	TEXT("0: Not Disabled, 1: Disabled"),
 	ECVF_Default);
 
+static int32 DisableFilteringCvar = 0;
+FAutoConsoleVariableRef CVarDisableFiltering(
+	TEXT("au.DisableFiltering"),
+	DisableFilteringCvar,
+	TEXT("Disables using the per-source lowpass and highpass filter.\n")
+	TEXT("0: Not Disabled, 1: Disabled"),
+	ECVF_Default);
+
+static int32 DisableHPFilteringCvar = 0;
+FAutoConsoleVariableRef CVarDisableHPFiltering(
+	TEXT("au.DisableHPFiltering"),
+	DisableHPFilteringCvar,
+	TEXT("Disables using the per-source highpass filter.\n")
+	TEXT("0: Not Disabled, 1: Disabled"),
+	ECVF_Default);
+
+static int32 DisableEnvelopeFollowingCvar = 0;
+FAutoConsoleVariableRef CVarDisableEnvelopeFollowing(
+	TEXT("au.DisableEnvelopeFollowing"),
+	DisableEnvelopeFollowingCvar,
+	TEXT("Disables using the envlope follower for source envelope tracking.\n")
+	TEXT("0: Not Disabled, 1: Disabled"),
+	ECVF_Default);
+
+static int32 DisableSourceEffectsCvar = 0;
+FAutoConsoleVariableRef CVarDisableSourceEffects(
+	TEXT("au.DisableSourceEffects"),
+	DisableSourceEffectsCvar,
+	TEXT("Disables using any source effects.\n")
+	TEXT("0: Not Disabled, 1: Disabled"),
+	ECVF_Default);
+
+static int32 DisableDistanceAttenuationCvar = 0;
+FAutoConsoleVariableRef CVarDisableDistanceAttenuation(
+	TEXT("au.DisableDistanceAttenuation"),
+	DisableDistanceAttenuationCvar,
+	TEXT("Disables using any Distance Attenuation.\n")
+	TEXT("0: Not Disabled, 1: Disabled"),
+	ECVF_Default);
+
+static int32 BypassAudioPluginsCvar = 0;
+FAutoConsoleVariableRef CVarBypassAudioPlugins(
+	TEXT("au.BypassAudioPlugins"),
+	BypassAudioPluginsCvar,
+	TEXT("Bypasses any audio plugin processing.\n")
+	TEXT("0: Not Disabled, 1: Disabled"),
+	ECVF_Default);
 
 #define ONEOVERSHORTMAX (3.0517578125e-5f) // 1/32768
 #define ENVELOPE_TAIL_THRESHOLD (1.58489e-5f) // -96 dB
@@ -41,6 +88,9 @@ FAutoConsoleVariableRef CVarDisableParallelSourceProcessing(
 
 #endif
 
+// Disable subframe timing logic
+#define AUDIO_SUBFRAME_ENABLED 0
+
 namespace Audio
 {
 	/*************************************************************************
@@ -97,7 +147,19 @@ namespace Audio
 		{
 			FSourceInfo& SourceInfo = SourceInfos[i];
 
-			SourceInfo.BufferQueueListener = nullptr;
+			SourceInfo.MixerSourceBuffer = nullptr;
+
+			SourceInfo.VolumeSourceStart = -1.0f;
+			SourceInfo.VolumeSourceDestination = -1.0f;
+			SourceInfo.VolumeFadeSlope = 0.0f;
+			SourceInfo.VolumeFadeStart = 0.0f;
+			SourceInfo.VolumeFadeFramePosition = 0;
+			SourceInfo.VolumeFadeNumFrames = 0;
+
+			SourceInfo.DistanceAttenuationSourceStart = -1.0f;
+			SourceInfo.DistanceAttenuationSourceDestination = -1.0f;
+
+			SourceInfo.SourceListener = nullptr;
 			SourceInfo.CurrentPCMBuffer = nullptr;	
 			SourceInfo.CurrentAudioChunkNumFrames = 0;
 			SourceInfo.CurrentFrameAlpha = 0.0f;
@@ -135,7 +197,7 @@ namespace Audio
 			SourceInfo.bOutputToBusOnly = false;
 			SourceInfo.bIsVorbis = false;
 			SourceInfo.bIsBypassingLPF = false;
-			SourceInfo.bIsBypassingLPF = false;
+			SourceInfo.bIsBypassingHPF = false;
 
 			SourceInfo.NumInputChannels = 0;
 			SourceInfo.NumPostEffectChannels = 0;
@@ -277,12 +339,19 @@ namespace Audio
 		SourceInfo.BusId = INDEX_NONE;
 		SourceInfo.BusDurationFrames = INDEX_NONE;
 
+		// Free the mixer source buffer data
+		if (SourceInfo.MixerSourceBuffer.IsValid())
+		{
+			PendingSourceBuffers.Add(SourceInfo.MixerSourceBuffer);
+			SourceInfo.MixerSourceBuffer = nullptr;
+		}
+
 		// Call OnRelease on the BufferQueueListener to give it a chance 
 		// to release any resources it owns on the audio render thread
-		if (SourceInfo.BufferQueueListener)
+		if (SourceInfo.SourceListener)
 		{
-			SourceInfo.BufferQueueListener->OnRelease(PendingReleaseData);
-			SourceInfo.BufferQueueListener = nullptr;
+			SourceInfo.SourceListener->OnRelease();
+			SourceInfo.SourceListener = nullptr;
 		}
 
 		// Remove the mixer source from its submix sends
@@ -322,13 +391,20 @@ namespace Audio
 
 		// Reset all state and data
 		SourceInfo.PitchSourceParam.Init();
-		SourceInfo.VolumeSourceParam.Init();
-		SourceInfo.DistanceAttenuationSourceParam.Init();
+		SourceInfo.VolumeSourceStart = -1.0f;
+		SourceInfo.VolumeSourceDestination = -1.0f;
+		SourceInfo.VolumeFadeSlope = 0.0f;
+		SourceInfo.VolumeFadeStart = 0.0f;
+		SourceInfo.VolumeFadeFramePosition = 0;
+		SourceInfo.VolumeFadeNumFrames = 0;
+
+		SourceInfo.DistanceAttenuationSourceStart = -1.0f;
+		SourceInfo.DistanceAttenuationSourceDestination = -1.0f;
 		SourceInfo.LPFCutoffFrequencyParam.Init();
+		SourceInfo.HPFCutoffFrequencyParam.Init();
 
 		SourceInfo.LowPassFilter.Reset();
 		SourceInfo.HighPassFilter.Reset();
-		SourceInfo.BufferQueue.Empty();
 		SourceInfo.CurrentPCMBuffer = nullptr;
 		SourceInfo.CurrentAudioChunkNumFrames = 0;
 		SourceInfo.SourceBuffer.Reset();
@@ -460,7 +536,7 @@ namespace Audio
 		AUDIO_MIXER_CHECK(SourceId < NumTotalSources);
 		AUDIO_MIXER_CHECK(GameThreadInfo.bIsBusy[SourceId]);
 		AUDIO_MIXER_CHECK(!GameThreadInfo.bIsDebugMode[SourceId]);
-		AUDIO_MIXER_CHECK(InitParams.BufferQueueListener != nullptr);
+		AUDIO_MIXER_CHECK(InitParams.SourceListener != nullptr);
 		AUDIO_MIXER_CHECK_GAME_THREAD(MixerDevice);
 
 #if AUDIO_MIXER_ENABLE_DEBUG_MODE
@@ -477,6 +553,10 @@ namespace Audio
 
 			FSourceInfo& SourceInfo = SourceInfos[SourceId];
 
+			// Initialize the mixer source buffer decoder with the given mixer buffer
+			SourceInfo.MixerSourceBuffer = InitParams.MixerSourceBuffer;
+			SourceInfo.MixerSourceBuffer->OnBeginGenerate();
+
 			SourceInfo.bIsPlaying = false;
 			SourceInfo.bIsPaused = false;
 			SourceInfo.bIsStopping = false;
@@ -488,10 +568,9 @@ namespace Audio
 			SourceInfo.bIsVorbis = InitParams.bIsVorbis;
 			SourceInfo.AudioComponentID = InitParams.AudioComponentID;
 
-			SourceInfo.BufferQueueListener = InitParams.BufferQueueListener;
-
 			// Call initialization from the render thread so anything wanting to do any initialization here can do so (e.g. procedural sound waves)
-			SourceInfo.BufferQueueListener->OnBeginGenerate();
+			SourceInfo.SourceListener = InitParams.SourceListener;
+			SourceInfo.SourceListener->OnBeginGenerate();
 
 			SourceInfo.NumInputChannels = InitParams.NumInputChannels;
 			SourceInfo.NumInputFrames = InitParams.NumInputFrames;
@@ -502,7 +581,7 @@ namespace Audio
 			SourceInfo.HighPassFilter.Init(MixerDevice->SampleRate, InitParams.NumInputChannels, 0, nullptr);
 			SourceInfo.HighPassFilter.SetFilterType(EFilter::Type::HighPass);
 
-			SourceInfo.SourceEnvelopeFollower = Audio::FEnvelopeFollower(MixerDevice->SampleRate, (float)InitParams.EnvelopeFollowerAttackTime, (float)InitParams.EnvelopeFollowerReleaseTime, Audio::EPeakMode::Peak);
+			SourceInfo.SourceEnvelopeFollower = Audio::FEnvelopeFollower(MixerDevice->SampleRate / NumOutputFrames, (float)InitParams.EnvelopeFollowerAttackTime, (float)InitParams.EnvelopeFollowerReleaseTime, Audio::EPeakMode::Peak);
 
 			// Create the spatialization plugin source effect
 			if (InitParams.bUseHRTFSpatialization)
@@ -724,18 +803,21 @@ namespace Audio
 			SourceInfo.bIsPlaying = false;
 			SourceInfo.bIsPaused = false;
 			SourceInfo.bIsActive = false;
+			SourceInfo.bIsStopping = false;
 
 			AUDIO_MIXER_DEBUG_LOG(SourceId, TEXT("Is immediately stopping"));
 		});
 	}
 
-	void FMixerSourceManager::StopFade(const int32 SourceId)
+	void FMixerSourceManager::StopFade(const int32 SourceId, const int32 NumFrames)
 	{
 		AUDIO_MIXER_CHECK(SourceId < NumTotalSources);
 		AUDIO_MIXER_CHECK(GameThreadInfo.bIsBusy[SourceId]);
+		AUDIO_MIXER_CHECK(NumFrames > 0);
 		AUDIO_MIXER_CHECK_GAME_THREAD(MixerDevice);
 
-		AudioMixerThreadCommand([this, SourceId]()
+
+		AudioMixerThreadCommand([this, SourceId, NumFrames]()
 		{
 			AUDIO_MIXER_CHECK_AUDIO_PLAT_THREAD(MixerDevice);
 
@@ -743,9 +825,25 @@ namespace Audio
 
 			SourceInfo.bIsPaused = false;
 			SourceInfo.bIsStopping = true;
-
-			// Set the volume to 0.0 on the source. This will naturally interpolate to 0.0 through the output frames
-			SourceInfo.VolumeSourceParam.SetValue(0.0f, NumOutputFrames);
+			
+			// Only allow multiple of 4 fade frames and positive
+			int32 NumFadeFrames = AlignArbitrary(NumFrames, 4);
+			if (NumFadeFrames <= 0)
+			{
+				// Stop immediately if we've been given no fade frames
+				SourceInfo.bIsPlaying = false;
+				SourceInfo.bIsPaused = false;
+				SourceInfo.bIsActive = false;
+				SourceInfo.bIsStopping = false;
+			}
+			else
+			{
+				// compute the fade slope
+				SourceInfo.VolumeFadeStart = SourceInfo.VolumeSourceStart;
+				SourceInfo.VolumeFadeNumFrames = NumFadeFrames;
+				SourceInfo.VolumeFadeSlope = -SourceInfo.VolumeSourceStart / SourceInfo.VolumeFadeNumFrames;
+				SourceInfo.VolumeFadeFramePosition = 0;
+			}
 
 			AUDIO_MIXER_DEBUG_LOG(SourceId, TEXT("Is stopping with fade"));
 		});
@@ -799,7 +897,13 @@ namespace Audio
 			// Only set the volume if we're not stopping. Stopping sources are setting their volume to 0.0.
 			if (!SourceInfo.bIsStopping)
 			{
-				SourceInfo.VolumeSourceParam.SetValue(Volume, NumOutputFrames);
+				// If we've not yet set a volume, we need to immediately set the start and destination to be the same value (to avoid an initial fade in)
+				if (SourceInfos[SourceId].VolumeSourceDestination < 0.0f)
+				{
+					SourceInfos[SourceId].VolumeSourceStart = Volume;
+				}
+
+				SourceInfos[SourceId].VolumeSourceDestination = Volume;
 			}
 		});
 	}
@@ -815,7 +919,13 @@ namespace Audio
 			AUDIO_MIXER_CHECK_AUDIO_PLAT_THREAD(MixerDevice);
 			check(NumOutputFrames > 0);
 
-			SourceInfos[SourceId].DistanceAttenuationSourceParam.SetValue(DistanceAttenuation, NumOutputFrames);
+			// If we've not yet set a distance attenuation, we need to immediately set the start and destination to be the same value (to avoid an initial fade in)
+			if (SourceInfos[SourceId].DistanceAttenuationSourceDestination < 0.0f)
+			{
+				SourceInfos[SourceId].DistanceAttenuationSourceStart = DistanceAttenuation;
+			}
+
+			SourceInfos[SourceId].DistanceAttenuationSourceDestination = DistanceAttenuation;
 		});
 	}
 
@@ -833,7 +943,7 @@ namespace Audio
 		});
 	}
 
-	void FMixerSourceManager::SetChannelMap(const int32 SourceId, const ESubmixChannelFormat SubmixChannelType, const TArray<float>& ChannelMap, const bool bInIs3D, const bool bInIsCenterChannelOnly)
+	void FMixerSourceManager::SetChannelMap(const int32 SourceId, const ESubmixChannelFormat SubmixChannelType, const Audio::AlignedFloatBuffer& ChannelMap, const bool bInIs3D, const bool bInIsCenterChannelOnly)
 	{
 		AUDIO_MIXER_CHECK(SourceId < NumTotalSources);
 		AUDIO_MIXER_CHECK(GameThreadInfo.bIsBusy[SourceId]);
@@ -864,7 +974,7 @@ namespace Audio
 				// If this is true, then the device changed while the command was in-flight
 				if (ChannelMap.Num() != ChannelMapSize)
 				{
-					TArray<float> NewChannelMap;
+					Audio::AlignedFloatBuffer NewChannelMap;
 
 					// If 3d then just zero it out, we'll get another channel map shortly
 					if (bInIs3D)
@@ -923,30 +1033,6 @@ namespace Audio
 		});
 	}
 
-	void FMixerSourceManager::SubmitBuffer(const int32 SourceId, FMixerSourceBufferPtr InSourceVoiceBuffer, const bool bSubmitSynchronously)
-	{
-		AUDIO_MIXER_CHECK(SourceId < NumTotalSources);
-		AUDIO_MIXER_CHECK(InSourceVoiceBuffer->Samples <= (uint32)InSourceVoiceBuffer->AudioData.Num());
-
-		// If we're submitting synchronously then don't use AudioMixerThreadCommand but submit directly
-		if (bSubmitSynchronously)
-		{
-			SourceInfos[SourceId].BufferQueue.Enqueue(InSourceVoiceBuffer);
-		}
-		else
-		{
-			AUDIO_MIXER_CHECK_GAME_THREAD(MixerDevice);
-
-			// Use a thread command
-			AudioMixerThreadCommand([this, SourceId, InSourceVoiceBuffer]()
-			{
-				AUDIO_MIXER_CHECK_AUDIO_PLAT_THREAD(MixerDevice);
-
-				SourceInfos[SourceId].BufferQueue.Enqueue(InSourceVoiceBuffer);
-			});
-		}
-	}
-
 	void FMixerSourceManager::SetSubmixSendInfo(const int32 SourceId, const FMixerSourceSubmixSend& InSubmixSend)
 	{
 		AUDIO_MIXER_CHECK(SourceId < NumTotalSources);
@@ -1067,20 +1153,13 @@ namespace Audio
 					break;
 				}
 
-				SourceInfo.BufferQueueListener->OnSourceBufferEnd();
+				SourceInfo.MixerSourceBuffer->OnBufferEnd();
 			}
 
 			// If we have audio in our queue, we're still playing
-			if (!SourceInfo.BufferQueue.IsEmpty())
+			if (SourceInfo.MixerSourceBuffer->GetNumBuffersQueued() > 0)
 			{
-				FMixerSourceBufferPtr NewBufferPtr;
-				SourceInfo.BufferQueue.Dequeue(NewBufferPtr);
-				SourceInfo.CurrentPCMBuffer = NewBufferPtr;
-
-				AUDIO_MIXER_CHECK(MixerSources[SourceId]->NumBuffersQueued.GetValue() > 0);
-
-				MixerSources[SourceId]->NumBuffersQueued.Decrement();
-
+				SourceInfo.CurrentPCMBuffer = SourceInfo.MixerSourceBuffer->GetNextBuffer();
 				SourceInfo.CurrentAudioChunkNumFrames = SourceInfo.CurrentPCMBuffer->Samples / NumChannels;
 
 				// Subtract the number of frames in the current buffer from our frame index.
@@ -1206,6 +1285,8 @@ namespace Audio
 			}
 			else
 			{
+
+#if AUDIO_SUBFRAME_ENABLED
 				// If we're not going to start yet, just continue
 				double StartFraction = (SourceInfo.StartTime - AudioRenderThreadTime) / AudioClockDelta;
 				if (StartFraction >= 1.0)
@@ -1228,6 +1309,10 @@ namespace Audio
 				// Update sample index to the frame we're starting on, accounting for source channels
 				int32 SampleIndex = StartFrame * SourceInfo.NumInputChannels;
 				bool bWriteOutZeros = true;
+#else
+				int32 SampleIndex = 0;
+				int32 StartFrame = 0;
+#endif
 
 				for (int32 Frame = StartFrame; Frame < NumOutputFrames; ++Frame)
 				{
@@ -1309,24 +1394,23 @@ namespace Audio
 
 	void FMixerSourceManager::ApplyDistanceAttenuation(FSourceInfo& SourceInfo, int32 NumSamples)
 	{
-		float* PostDistanceAttenBufferPtr = SourceInfo.SourceBuffer.GetData();
-
-		// Interpolate the distance attenuation value to avoid discontinuities
-		float CurrentDistanceAttenuationValue = SourceInfo.DistanceAttenuationSourceParam.GetValue();
-		for (int32 Sample = 0; Sample < NumSamples; Sample += SourceInfo.NumInputChannels)
+		if (DisableDistanceAttenuationCvar)
 		{
-			for (int32 Chan = 0; Chan < SourceInfo.NumInputChannels; ++Chan)
-			{
-				PostDistanceAttenBufferPtr[Sample + Chan] = PostDistanceAttenBufferPtr[Sample + Chan] * CurrentDistanceAttenuationValue;
-			}
-
-			CurrentDistanceAttenuationValue = SourceInfo.DistanceAttenuationSourceParam.Update();
+			return;
 		}
-		SourceInfo.DistanceAttenuationSourceParam.Reset();
+
+		float* PostDistanceAttenBufferPtr = SourceInfo.SourceBuffer.GetData();
+		Audio::FadeBufferFast(PostDistanceAttenBufferPtr, SourceInfo.SourceBuffer.Num(), SourceInfo.DistanceAttenuationSourceStart, SourceInfo.DistanceAttenuationSourceDestination);
+		SourceInfo.DistanceAttenuationSourceStart = SourceInfo.DistanceAttenuationSourceDestination;
 	}
 
 	void FMixerSourceManager::ComputePluginAudio(FSourceInfo& SourceInfo, int32 SourceId, int32 NumSamples)
 	{
+		if (BypassAudioPluginsCvar)
+		{
+			return;
+		}
+
 		// Apply the distance attenuation before sending to plugins
 		float* PreDistanceAttenBufferPtr = SourceInfo.PreDistanceAttenuationBuffer.GetData();
 		float* PostDistanceAttenBufferPtr = SourceInfo.SourceBuffer.GetData();
@@ -1387,10 +1471,7 @@ namespace Audio
 				const float* ReverbPluginOutputBufferPtr = SourceInfo.ReverbPluginOutputBuffer.GetData();
 				const float* AudioPluginOutputDataPtr = SourceInfo.AudioPluginOutputData.AudioBuffer.GetData();
 
-				for (int32 i = 0; i < NumSamples; ++i)
-				{
-					PostDistanceAttenBufferPtr[i] = ReverbPluginOutputBufferPtr[i] + AudioPluginOutputDataPtr[i];
-				}
+				Audio::SumBuffers(ReverbPluginOutputBufferPtr, AudioPluginOutputDataPtr, PostDistanceAttenBufferPtr, NumSamples);
 			}
 			else
 			{
@@ -1400,10 +1481,7 @@ namespace Audio
 		else if (bShouldMixInReverb)
 		{
 			const float* ReverbPluginOutputBufferPtr = SourceInfo.ReverbPluginOutputBuffer.GetData();
-			for (int32 i = 0; i < NumSamples; ++i)
-			{
-				PostDistanceAttenBufferPtr[i] += ReverbPluginOutputBufferPtr[i];
-			}
+			Audio::MixInBufferFast(ReverbPluginOutputBufferPtr, PostDistanceAttenBufferPtr, NumSamples);
 		}
 
 		// If the source has HRTF processing enabled, run it through the spatializer
@@ -1490,33 +1568,42 @@ namespace Audio
 			float* PreDistanceAttenBufferPtr = SourceInfo.PreDistanceAttenuationBuffer.GetData();
 			const int32 NumSamples = SourceInfo.PreDistanceAttenuationBuffer.Num();
 
-			// Apply non-distance attenuation volume
-			float CurrentVolumeValue = SourceInfo.VolumeSourceParam.GetValue();
-
-			for (int32 Frame = 0; Frame < NumOutputFrames; ++Frame)
+			// Update volume fade information if we're stopping
+			if (SourceInfo.bIsStopping)
 			{
-#if AUDIO_MIXER_ENABLE_DEBUG_MODE				
-				CurrentVolumeValue = (bIsDebugModeEnabled && !SourceInfo.bIsDebugMode) ? 0.0f : SourceInfo.VolumeSourceParam.Update();
+				const int32 NumFadeFrames = FMath::Min(SourceInfo.VolumeFadeNumFrames - SourceInfo.VolumeFadeFramePosition, NumOutputFrames);
 
-#else
-				CurrentVolumeValue = SourceInfo.VolumeSourceParam.Update();
-#endif
-				int32 SampleIndex = SourceInfo.NumInputChannels * Frame;
+				SourceInfo.VolumeFadeFramePosition += NumFadeFrames;
+				SourceInfo.VolumeSourceDestination = SourceInfo.VolumeFadeSlope * (float) SourceInfo.VolumeFadeFramePosition + SourceInfo.VolumeFadeStart;
 
-				// Scale by current volume value (note: not including distance attenuation). TODO: do this as a SIMD operation in its own loop
-				for (int32 Channel = 0; Channel < SourceInfo.NumInputChannels; ++Channel)
+				if (FMath::IsNearlyZero(SourceInfo.VolumeSourceDestination, KINDA_SMALL_NUMBER))
 				{
-					PreDistanceAttenBufferPtr[SampleIndex + Channel] *= CurrentVolumeValue;
+					SourceInfo.VolumeSourceDestination = 0.0f;
 				}
-			}
 
-			// Reset the volume interpolator
-			SourceInfo.VolumeSourceParam.Reset();
+				const int32 NumFadeSamples = NumFadeFrames * SourceInfo.NumInputChannels;
+
+				Audio::FadeBufferFast(PreDistanceAttenBufferPtr, NumFadeSamples, SourceInfo.VolumeSourceStart, SourceInfo.VolumeSourceDestination);
+
+				// Zero the rest of the buffer
+				if (NumFadeFrames < NumOutputFrames)
+				{
+					int32 SamplesLeft = NumSamples - NumFadeSamples;
+					FMemory::Memzero(&PreDistanceAttenBufferPtr[NumFadeFrames], sizeof(float)*SamplesLeft);
+				}
+
+				SourceInfo.VolumeSourceStart = SourceInfo.VolumeSourceDestination;
+			}
+			else
+			{
+				Audio::FadeBufferFast(PreDistanceAttenBufferPtr, NumSamples, SourceInfo.VolumeSourceStart, SourceInfo.VolumeSourceDestination);
+				SourceInfo.VolumeSourceStart = SourceInfo.VolumeSourceDestination;
+			}
 
 			// Now process the effect chain if it exists
-			if (SourceInfo.SourceEffects.Num() > 0)
+			if (!DisableSourceEffectsCvar && SourceInfo.SourceEffects.Num() > 0)
 			{
-				SourceInfo.SourceEffectInputData.CurrentVolume = CurrentVolumeValue;
+				SourceInfo.SourceEffectInputData.CurrentVolume = SourceInfo.VolumeSourceDestination;
 
 				SourceInfo.SourceEffectOutputData.AudioFrame.Reset();
 				SourceInfo.SourceEffectOutputData.AudioFrame.AddZeroed(SourceInfo.NumInputChannels);
@@ -1559,27 +1646,24 @@ namespace Audio
 					}
 				}
 			}
-
-			// Compute the source envelope using pre-distance attenuation buffer
-			float AverageSampleValue;
-			for (int32 Sample = 0; Sample < NumSamples;)
+			if (!DisableEnvelopeFollowingCvar)
 			{
-				AverageSampleValue = 0.0f;
-				for (int32 Chan = 0; Chan < SourceInfo.NumInputChannels; ++Chan)
-				{
-					AverageSampleValue += PreDistanceAttenBufferPtr[Sample++];
-				}
-				AverageSampleValue /= SourceInfo.NumInputChannels;
+				// Compute the source envelope using pre-distance attenuation buffer
+				float AverageSampleValue = Audio::GetAverageAmplitude(PreDistanceAttenBufferPtr, NumSamples);
 				SourceInfo.SourceEnvelopeFollower.ProcessAudio(AverageSampleValue);
-			}
 
-			// Copy the current value of the envelope follower (block-rate value)
-			SourceInfo.SourceEnvelopeValue = SourceInfo.SourceEnvelopeFollower.GetCurrentValue();
+				// Copy the current value of the envelope follower (block-rate value)
+				SourceInfo.SourceEnvelopeValue = SourceInfo.SourceEnvelopeFollower.GetCurrentValue();
 
-			SourceInfo.bEffectTailsDone = SourceInfo.bEffectTailsDone || SourceInfo.SourceEnvelopeValue < ENVELOPE_TAIL_THRESHOLD;
+				SourceInfo.bEffectTailsDone = SourceInfo.bEffectTailsDone || SourceInfo.SourceEnvelopeValue < ENVELOPE_TAIL_THRESHOLD;
+			}
+			else
+			{
+				SourceInfo.bEffectTailsDone = true;
+			}
 
 			// Only scale with distance attenuation and send to source audio to plugins if we're not in output-to-bus only mode
-			if (!SourceInfo.bOutputToBusOnly)
+			if (!SourceInfo.bOutputToBusOnly && !DisableFilteringCvar)
 			{
 				float* PostDistanceAttenBufferPtr = SourceInfo.SourceBuffer.GetData();
 
@@ -1614,7 +1698,7 @@ namespace Audio
 						FMemory::Memcpy(&PostDistanceAttenBufferPtr[SampleIndex], &PreDistanceAttenBufferPtr[SampleIndex], SourceInfo.NumInputChannels * sizeof(float));
 					}
 
-					if (HPFFreq > 0.0f)
+					if (!DisableHPFilteringCvar && HPFFreq > 0.0f)
 					{
 						//If we stopped processing the low pass filter, we need to clear the filter's memory to prevent pops.
 						if (SourceInfo.bIsBypassingHPF)
@@ -1641,12 +1725,22 @@ namespace Audio
 				// Send source audio to plugins
 				ComputePluginAudio(SourceInfo, SourceId, NumSamples);
 			}
+			else if (DisableFilteringCvar)
+			{
+				float* PostDistanceAttenBufferPtr = SourceInfo.SourceBuffer.GetData();
+				FMemory::Memcpy(PostDistanceAttenBufferPtr, PreDistanceAttenBufferPtr, NumSamples * sizeof(float));
+
+				// Apply distance attenuation
+				ApplyDistanceAttenuation(SourceInfo, NumSamples);
+
+				// Send source audio to plugins
+				ComputePluginAudio(SourceInfo, SourceId, NumSamples);
+			}
 
 			// Check the source effect tails condition
 			if (SourceInfo.bIsLastBuffer && SourceInfo.bEffectTailsDone)
 			{
 				// If we're done and our tails our done, clear everything out
-				SourceInfo.BufferQueue.Empty();
 				SourceInfo.CurrentFrameValues.Reset();
 				SourceInfo.NextFrameValues.Reset();
 				SourceInfo.CurrentPCMBuffer = nullptr;
@@ -1823,23 +1917,8 @@ namespace Audio
 				const int32 OutWetBufferSize = OutWetBuffer.Num();
 				float* OutWetBufferPtr = OutWetBuffer.GetData();
 
-				for (int32 SampleIndex = 0; SampleIndex < OutWetBufferSize; ++SampleIndex)
-				{
-					OutWetBufferPtr[SampleIndex] += SourceOutputBufferPtr[SampleIndex] * SendLevel;
-				}
+				Audio::MixInBufferFast(SourceOutputBufferPtr, OutWetBufferPtr, OutWetBufferSize, SendLevel);
 			}
-
-// TOODO: SIMD
-// 			const VectorRegister SendLevelScalar = VectorLoadFloat1(SendLevel);
-// 			for (int32 SampleIndex = 0; SampleIndex < OutWetBuffer.Num(); SampleIndex += 4)
-// 			{
-// 				VectorRegister OutputBufferRegister = VectorLoadAligned(&SourceOutputBuffer[SampleIndex]);
-// 				VectorRegister Temp = VectorMultiply(OutputBufferRegister, SendLevelScalar);
-// 				VectorRegister OutputWetBuffer = VectorLoadAligned(&OutWetBuffer[SampleIndex]);
-// 				Temp = VectorAdd(OutputWetBuffer, Temp);
-// 
-// 				VectorStoreAligned(Temp, &OutWetBuffer[SampleIndex]);
-// 			}
 		}
 	}
 
@@ -2022,6 +2101,9 @@ namespace Audio
 			if (SourceInfo.bIsLastBuffer)
 			{
 				SourceInfo.bIsDone = true;
+
+				// Notify that we're now done with this source
+				SourceInfo.SourceListener->OnDone();
 			}
 
 			GameThreadInfo.bIsDone[SourceId] = SourceInfo.bIsDone;
@@ -2038,16 +2120,18 @@ namespace Audio
 			// Check for the stopping condition to "turn the sound off"
 			if (!SourceInfo.bIsDone && SourceInfo.bIsStopping)
 			{
-				const float CurrentVolumeValue = SourceInfo.VolumeSourceParam.GetValue();
-				if (CurrentVolumeValue == 0.0f)
+				if (SourceInfo.VolumeSourceDestination == 0.0f)
 				{
+					SourceInfo.bIsStopping = false;
 					SourceInfo.bIsDone = true;
 					GameThreadInfo.bIsDone[SourceId] = true;
-
-					// Allow any waits to clear if they were waiting on this source to stop
-					MixerSources[SourceId]->NotifyStopped();
 				}
 			}
+			else if (SourceInfo.bIsDone)
+			{
+				SourceInfo.bIsStopping = false;
+				GameThreadInfo.bIsDone[SourceId] = true;
+			}
 		}
 	}
 
@@ -2060,8 +2144,6 @@ namespace Audio
 
 	void FMixerSourceManager::PumpCommandQueue()
 	{
-		AUDIO_MIXER_CHECK_AUDIO_PLAT_THREAD(MixerDevice);
-
 		int32 CurrentRenderThreadIndex = RenderThreadCommandBufferIndex.GetValue();
 
 		FCommands& Commands = CommandBuffers[CurrentRenderThreadIndex];
@@ -2080,45 +2162,52 @@ namespace Audio
 
 	void FMixerSourceManager::UpdatePendingReleaseData(bool bForceWait)
 	{
-		// Don't block, but let tasks finish naturally
-		for (int32 i = PendingReleaseData.Num() - 1; i >= 0; --i)
+		// Check for any pending delete procedural sound waves
+		for (int32 SourceId = 0; SourceId < NumTotalSources; ++SourceId)
 		{
-			FPendingReleaseData* DataEntry = PendingReleaseData[i];
-			if (DataEntry->Task)
+			FSourceInfo& SourceInfo = SourceInfos[SourceId];
+			if (SourceInfo.MixerSourceBuffer.IsValid())
 			{
-				bool bDeleteData = false;
-				if (bForceWait)
+				// If we've been flagged to begin destroy
+				if (SourceInfo.MixerSourceBuffer->IsBeginDestroy())
 				{
-					DataEntry->Task->EnsureCompletion();
-					bDeleteData = true;
-				}
-				else if (DataEntry->Task->IsDone())
-				{
-					bDeleteData = true;
+					SourceInfo.MixerSourceBuffer->ClearSoundWave();
+						
+					// Clear out the mixer source buffer
+					SourceInfo.MixerSourceBuffer = nullptr;
+
+					// Signal that this source is done
+					SourceInfo.SourceListener->OnDone();
+
+					// Set the sound to be done playing
+					// This will flag the sound to be released
+					SourceInfo.bIsPlaying = false;
+					SourceInfo.bIsPaused = false;
+					SourceInfo.bIsActive = false;
+					SourceInfo.bIsStopping = false;
 				}
+			}
+		}
 
-				if (bDeleteData)
-				{
-					delete DataEntry->Task;
-					DataEntry->Task = nullptr;
 
-					if (DataEntry->Buffer)
-					{
-						delete DataEntry->Buffer;
-					}
-
-					delete DataEntry;
-					PendingReleaseData[i] = nullptr;
+		// Don't block, but let tasks finish naturally
+		for (int32 i = PendingSourceBuffers.Num() - 1; i >= 0; --i)
+		{
+			FMixerSourceBuffer* MixerSourceBuffer = PendingSourceBuffers[i].Get();
 
-					PendingReleaseData.RemoveAtSwap(i, 1, false);
-				}
-			}
-			else if (DataEntry->Buffer)
+			bool bDeleteSourceBuffer = true;
+			if (bForceWait)
 			{
-				delete DataEntry->Buffer;
-				DataEntry->Buffer = nullptr;
+				MixerSourceBuffer->EnsureAsyncTaskFinishes();
+			}
+			else if (!MixerSourceBuffer->IsAsyncTaskDone())
+			{			
+				bDeleteSourceBuffer = false;
+			}
 
-				PendingReleaseData.RemoveAtSwap(i, 1, false);
+			if (bDeleteSourceBuffer)
+			{
+				PendingSourceBuffers.RemoveAtSwap(i, 1, false);
 			}
 		}
 	}
diff --git a/Engine/Source/Runtime/AudioMixer/Private/AudioMixerSourceManager.h b/Engine/Source/Runtime/AudioMixer/Private/AudioMixerSourceManager.h
--- a/Engine/Source/Runtime/AudioMixer/Private/AudioMixerSourceManager.h
+++ b/Engine/Source/Runtime/AudioMixer/Private/AudioMixerSourceManager.h
@@ -13,8 +13,10 @@
 #include "DSP/Filter.h"
 #include "DSP/EnvelopeFollower.h"
 #include "DSP/ParamInterpolator.h"
+#include "DSP/BufferVectorOperations.h"
 #include "IAudioExtensionPlugin.h"
 #include "Containers/Queue.h"
+#include "AudioMixerSourceBuffer.h"
 
 
 namespace Audio
@@ -22,6 +24,8 @@ namespace Audio
 	class FMixerSubmix;
 	class FMixerDevice;
 	class FMixerSourceVoice;
+	class FMixerSourceBuffer;
+	class ISourceListener;
 
 	/** Struct defining a source voice buffer. */
 	struct FMixerSourceVoiceBuffer
@@ -44,32 +48,23 @@ namespace Audio
 		}
 	};
 
-	typedef TSharedPtr<FMixerSourceVoiceBuffer, ESPMode::ThreadSafe> FMixerSourceBufferPtr;
 	typedef TSharedPtr<FMixerSubmix, ESPMode::ThreadSafe> FMixerSubmixPtr;
 
-	// Task used to store pending release/decode data
-	struct FPendingReleaseData
-	{
-		FSoundBuffer* Buffer;
-		IAudioTask* Task;
-
-		FPendingReleaseData()
-			: Buffer(nullptr)
-			, Task(nullptr)
-		{}
-	};
-
-	class ISourceBufferQueueListener
+	class ISourceListener
 	{
 	public:
 		// Called before a source begins to generate audio. 
 		virtual void OnBeginGenerate() = 0;
 
-		// Called when the current buffer is finished and a new one needs to be queued
-		virtual void OnSourceBufferEnd() = 0;
+		// Called when a loop point is hit
+		virtual void OnLoopEnd() = 0;
+
+		// Called the source is released on the render thread
+		virtual void OnRelease() = 0;
+
+		// Called when the source finishes on the render thread
+		virtual void OnDone() = 0;
 
-		// Called when the buffer queue listener is released. Allows cleaning up any resources from render thread.
-		virtual void OnRelease(TArray<FPendingReleaseData*>& OutPendingReleaseData) = 0;
 	};
 
 	struct FMixerSourceSubmixSend
@@ -93,7 +88,8 @@ namespace Audio
 
 	struct FMixerSourceVoiceInitParams
 	{
-		ISourceBufferQueueListener* BufferQueueListener;
+		TSharedPtr<FMixerSourceBuffer> MixerSourceBuffer;
+		ISourceListener* SourceListener;
 		TArray<FMixerSourceSubmixSend> SubmixSends;
 		TArray<FMixerBusSend> BusSends[(int32)EBusSendType::Count];
 		uint32 BusId;
@@ -117,9 +113,11 @@ namespace Audio
 		bool bOutputToBusOnly;
 		bool bIsVorbis;
 		bool bIsAmbisonics;
+		bool bIsSeeking;
 
 		FMixerSourceVoiceInitParams()
-			: BufferQueueListener(nullptr)
+			: MixerSourceBuffer(nullptr)
+			, SourceListener(nullptr)
 			, BusId(INDEX_NONE)
 			, BusDuration(0.0f)
 			, SourceEffectChainId(INDEX_NONE)
@@ -138,6 +136,7 @@ namespace Audio
 			, bOutputToBusOnly(false)
 			, bIsVorbis(false)
 			, bIsAmbisonics(false)
+			, bIsSeeking(false)
 		{}
 	};
 
@@ -151,7 +150,7 @@ namespace Audio
 			ChannelValues.Reset();
 		}
 
-		FORCEINLINE void SetChannelMap(const TArray<float>& ChannelMap, const int32 InNumInterpFrames)
+		FORCEINLINE void SetChannelMap(const Audio::AlignedFloatBuffer& ChannelMap, const int32 InNumInterpFrames)
 		{
 			if (ChannelValues.Num() != ChannelMap.Num())
 			{
@@ -235,21 +234,19 @@ namespace Audio
 
 		void Play(const int32 SourceId);
 		void Stop(const int32 SourceId);
-		void StopFade(const int32 SourceId);
+		void StopFade(const int32 SourceId, const int32 NumFrames);
 		void Pause(const int32 SourceId);
 		void SetPitch(const int32 SourceId, const float Pitch);
 		void SetVolume(const int32 SourceId, const float Volume);
 		void SetDistanceAttenuation(const int32 SourceId, const float DistanceAttenuation);
 		void SetSpatializationParams(const int32 SourceId, const FSpatializationParams& InParams);
-		void SetChannelMap(const int32 SourceId, const ESubmixChannelFormat SubmixChannelType, const TArray<float>& InChannelMap, const bool bInIs3D, const bool bInIsCenterChannelOnly);
+		void SetChannelMap(const int32 SourceId, const ESubmixChannelFormat SubmixChannelType, const Audio::AlignedFloatBuffer& InChannelMap, const bool bInIs3D, const bool bInIsCenterChannelOnly);
 		void SetLPFFrequency(const int32 SourceId, const float Frequency);
 		void SetHPFFrequency(const int32 SourceId, const float Frequency);
 
 		void SetListenerTransforms(const TArray<FTransform>& ListenerTransforms);
 		const TArray<FTransform>* GetListenerTransforms() const;
 
-		void SubmitBuffer(const int32 SourceId, FMixerSourceBufferPtr InSourceVoiceBuffer, const bool bSubmitSynchronously);
-
 		int64 GetNumFramesPlayed(const int32 SourceId) const;
 		float GetEnvelopeValue(const int32 SourceId) const;
 		bool IsDone(const int32 SourceId) const;
@@ -369,18 +366,18 @@ namespace Audio
 			FSourceInfo() {}
 			~FSourceInfo() {}
 
-			// Raw PCM buffer data
-			TQueue<FMixerSourceBufferPtr> BufferQueue;
-			ISourceBufferQueueListener* BufferQueueListener;
+			// Object which handles source buffer decoding
+			TSharedPtr<FMixerSourceBuffer> MixerSourceBuffer;
+			ISourceListener* SourceListener;
 
 			// Data used for rendering sources
-			FMixerSourceBufferPtr CurrentPCMBuffer;
+			TSharedPtr<FMixerSourceVoiceBuffer> CurrentPCMBuffer;
 			int32 CurrentAudioChunkNumFrames;
 
 			// The post-attenuation source buffer, used to send audio to submixes
-			TArray<float> SourceBuffer;
-			TArray<float> PreEffectBuffer;
-			TArray<float> PreDistanceAttenuationBuffer;
+			Audio::AlignedFloatBuffer SourceBuffer;
+			Audio::AlignedFloatBuffer PreEffectBuffer;
+			Audio::AlignedFloatBuffer PreDistanceAttenuationBuffer;
 
 			TArray<float> CurrentFrameValues;
 			TArray<float> NextFrameValues;
@@ -404,8 +401,15 @@ namespace Audio
 
 			// Interpolated source params
 			FParam PitchSourceParam;
-			FParam VolumeSourceParam;
-			FParam DistanceAttenuationSourceParam;
+			float VolumeSourceStart;
+			float VolumeSourceDestination;
+			float VolumeFadeSlope;
+			float VolumeFadeStart;
+			int32 VolumeFadeFramePosition;
+			int32 VolumeFadeNumFrames;
+
+			float DistanceAttenuationSourceStart;
+			float DistanceAttenuationSourceDestination;
 			FParam LPFCutoffFrequencyParam;
 			FParam HPFCutoffFrequencyParam;
 
@@ -428,11 +432,11 @@ namespace Audio
 			float SourceEnvelopeValue;
 
 			FSpatializationParams SpatParams;
-			TArray<float> ScratchChannelMap;
+			Audio::AlignedFloatBuffer ScratchChannelMap;
 
 			// Output data, after computing a block of sample data, this is read back from mixers
-			TArray<float> ReverbPluginOutputBuffer;
-			TArray<float>* PostEffectBuffers;
+			Audio::AlignedFloatBuffer ReverbPluginOutputBuffer;
+			Audio::AlignedFloatBuffer* PostEffectBuffers;
 
 			// Data needed for outputting to submixes
 			FSubmixChannelTypeInfo SubmixChannelInfo[(int32) ESubmixChannelFormat::Count];
@@ -484,7 +488,7 @@ namespace Audio
 		TArray<FAsyncTask<FAudioMixerSourceWorker>*> SourceWorkers;
 
 		// Array of task data waiting to finished. Processed on audio render thread.
-		TArray<FPendingReleaseData*> PendingReleaseData;
+		TArray<TSharedPtr<FMixerSourceBuffer>> PendingSourceBuffers;
 
 		// General information about sources in source manager accessible from game thread
 		struct FGameThreadInfo
diff --git a/Engine/Source/Runtime/AudioMixer/Private/AudioMixerSourceVoice.cpp b/Engine/Source/Runtime/AudioMixer/Private/AudioMixerSourceVoice.cpp
--- a/Engine/Source/Runtime/AudioMixer/Private/AudioMixerSourceVoice.cpp
+++ b/Engine/Source/Runtime/AudioMixer/Private/AudioMixerSourceVoice.cpp
@@ -14,17 +14,11 @@ namespace Audio
 
 	FMixerSourceVoice::FMixerSourceVoice()
 	{
-		SourceStoppedEvent = FPlatformProcess::GetSynchEventFromPool();;
-
 		Reset(nullptr);
 	}
 
 	FMixerSourceVoice::~FMixerSourceVoice()
 	{
-		if (SourceStoppedEvent != nullptr)
-		{
-			FPlatformProcess::ReturnSynchEventToPool(SourceStoppedEvent);
-		}
 	}
 
 	void FMixerSourceVoice::Reset(FMixerDevice* InMixerDevice)
@@ -40,7 +34,6 @@ namespace Audio
 			SourceManager = nullptr;
 		}
 
-		NumBuffersQueued.Reset();
 		Pitch = -1.0f;
 		Volume = -1.0f;
 		DistanceAttenuation = -1.0f;
@@ -53,6 +46,7 @@ namespace Audio
 		bIsActive = false;
 		bIsBus = false;
 		bOutputToBusOnly = false;
+		bStopFadedOut = false;
 		SubmixSends.Reset();
 	}
 
@@ -62,7 +56,7 @@ namespace Audio
 
 		if (SourceManager->GetFreeSourceId(SourceId))
 		{
-			AUDIO_MIXER_CHECK(InitParams.BufferQueueListener != nullptr);
+			AUDIO_MIXER_CHECK(InitParams.SourceListener != nullptr);
 			AUDIO_MIXER_CHECK(InitParams.NumInputChannels > 0);
 
 			bOutputToBusOnly = InitParams.bOutputToBusOnly;
@@ -73,14 +67,11 @@ namespace Audio
 				SubmixSends.Add(InitParams.SubmixSends[i].Submix->GetId(), InitParams.SubmixSends[i]);
 			}
 
-			if (SourceStoppedEvent)
-			{
-				SourceStoppedEvent->Reset();
-			}
-
+			bStopFadedOut = false;
 			SourceManager->InitSource(SourceId, InitParams);
 			return true;
 		}
+
 		return false;
 	}
 
@@ -91,17 +82,6 @@ namespace Audio
 		SourceManager->ReleaseSourceId(SourceId);
 	}
 
-	void FMixerSourceVoice::SubmitBuffer(FMixerSourceBufferPtr InSourceVoiceBuffer, const bool bSubmitSynchronously)
-	{
-		NumBuffersQueued.Increment();
-		SourceManager->SubmitBuffer(SourceId, InSourceVoiceBuffer, bSubmitSynchronously);
-	}
-
-	int32 FMixerSourceVoice::GetNumBuffersQueued() const
-	{
-		return NumBuffersQueued.GetValue();
-	}
-
 	void FMixerSourceVoice::SetPitch(const float InPitch)
 	{
 		AUDIO_MIXER_CHECK_GAME_THREAD(MixerDevice);
@@ -157,7 +137,7 @@ namespace Audio
 		}
 	}
 
-	void FMixerSourceVoice::SetChannelMap(ESubmixChannelFormat InChannelType, TArray<float>& InChannelMap, const bool bInIs3D, const bool bInIsCenterChannelOnly)
+	void FMixerSourceVoice::SetChannelMap(ESubmixChannelFormat InChannelType, const Audio::AlignedFloatBuffer& InChannelMap, const bool bInIs3D, const bool bInIsCenterChannelOnly)
 	{
 		AUDIO_MIXER_CHECK_GAME_THREAD(MixerDevice);
 
@@ -189,34 +169,19 @@ namespace Audio
 		bIsPlaying = false;
 		bIsPaused = false;
 		bIsActive = false;
+		// We are instantly fading out with this stop command
+		bStopFadedOut = true;
 		SourceManager->Stop(SourceId);
 	}
 
-	void FMixerSourceVoice::StopFade()
+	void FMixerSourceVoice::StopFade(int32 NumFrames)
 	{
 		AUDIO_MIXER_CHECK_GAME_THREAD(MixerDevice);
 
 		bIsPaused = false;
-		SourceManager->StopFade(SourceId);
-	}
-
-	void FMixerSourceVoice::EnsureStopped()
-	{
-		if (SourceStoppedEvent)
-		{
-			SourceStoppedEvent->Wait();
-		}
+		SourceManager->StopFade(SourceId, NumFrames);
 	}
 
-	void FMixerSourceVoice::NotifyStopped()
-	{
-		if (SourceStoppedEvent)
-		{
-			SourceStoppedEvent->Trigger();
-		}
-	}
-
-
 	void FMixerSourceVoice::Pause()
 	{
 		AUDIO_MIXER_CHECK_GAME_THREAD(MixerDevice);
@@ -315,7 +280,7 @@ namespace Audio
 		}
 	}
 
-	void FMixerSourceVoice::OnMixBus(FMixerSourceBufferPtr OutMixerSourceBuffer)
+	void FMixerSourceVoice::OnMixBus(FMixerSourceVoiceBuffer* OutMixerSourceBuffer)
 	{
 		AUDIO_MIXER_CHECK_AUDIO_PLAT_THREAD(MixerDevice);
 
diff --git a/Engine/Source/Runtime/AudioMixer/Private/AudioMixerSourceVoice.h b/Engine/Source/Runtime/AudioMixer/Private/AudioMixerSourceVoice.h
--- a/Engine/Source/Runtime/AudioMixer/Private/AudioMixerSourceVoice.h
+++ b/Engine/Source/Runtime/AudioMixer/Private/AudioMixerSourceVoice.h
@@ -32,12 +32,6 @@ namespace Audio
 		// Releases the source voice back to the source buffer pool
 		void Release();
 
-		// Queues the given buffer data to the internal queue of audio buffers.
-		void SubmitBuffer(FMixerSourceBufferPtr InSourceVoiceBuffer, const bool bOnRenderThread);
-
-		// Returns the number of buffers that are currently queued to be played.
-		int32 GetNumBuffersQueued() const;
-
 		// Sets the source voice pitch value.
 		void SetPitch(const float InPitch);
 
@@ -54,7 +48,7 @@ namespace Audio
 		void SetHPFFrequency(const float InFrequency);
 
 		// Sets the source voice's channel map (2d or 3d).
-		void SetChannelMap(ESubmixChannelFormat InChannelType, TArray<float>& InChannelMap, const bool bInIs3D, const bool bInIsCenterChannelOnly);
+		void SetChannelMap(ESubmixChannelFormat InChannelType, const Audio::AlignedFloatBuffer& InChannelMap, const bool bInIs3D, const bool bInIsCenterChannelOnly);
 
 		// Sets params used by HRTF spatializer
 		void SetSpatializationParams(const FSpatializationParams& InParams);
@@ -69,13 +63,7 @@ namespace Audio
 		void Stop();
 
 		// Does a faded stop (to avoid discontinuity)
-		void StopFade();
-
-		// Makes sure the source has stopped
-		void EnsureStopped();
-
-		// Triggers the stop event
-		void NotifyStopped();
+		void StopFade(int32 NumFrames);
 
 		// Queries if the voice is playing
 		bool IsPlaying() const;
@@ -90,7 +78,7 @@ namespace Audio
 		bool IsDone() const;
 
 		// Queries if the source has finished its fade out.
-		bool IsStopFadedOut() const;
+		bool IsStopFadedOut() const { return bStopFadedOut; }
 
 		// Queries if the source ffect tails have finished
 		bool IsSourceEffectTailsDone() const;
@@ -111,18 +99,17 @@ namespace Audio
 		void SetSubmixSendInfo(FMixerSubmixPtr Submix, const float SendLevel);
 
 		// Called when the source is a bus and needs to mix other sources together to generate output
-		void OnMixBus(FMixerSourceBufferPtr OutMixerSourceBuffer);
+		void OnMixBus(FMixerSourceVoiceBuffer* OutMixerSourceBuffer);
 
 	private:
 
 		friend class FMixerSourceManager;
 
-		FEvent* SourceStoppedEvent;
 		FMixerSourceManager* SourceManager;
 		TMap<uint32, FMixerSourceSubmixSend> SubmixSends;
 		FMixerDevice* MixerDevice;
 		TMap<ESubmixChannelFormat, TArray<float>> ChannelMaps;
-		FThreadSafeCounter NumBuffersQueued;
+		FThreadSafeBool bStopFadedOut;
 		float Pitch;
 		float Volume;
 		float DistanceAttenuation;
diff --git a/Engine/Source/Runtime/AudioMixer/Private/AudioMixerSubmix.cpp b/Engine/Source/Runtime/AudioMixer/Private/AudioMixerSubmix.cpp
--- a/Engine/Source/Runtime/AudioMixer/Private/AudioMixerSubmix.cpp
+++ b/Engine/Source/Runtime/AudioMixer/Private/AudioMixerSubmix.cpp
@@ -271,7 +271,7 @@ namespace Audio
 	{
 		// Retrieve ptr to the cached downmix channel map from the mixer device
 		int32 NewChannelCount = MixerDevice->GetNumChannelsForSubmixFormat(InNewChannelType);
-		TArray<float> ChannelMap;
+		Audio::AlignedFloatBuffer ChannelMap;
 		MixerDevice->Get2DChannelMap(false, InNewChannelType, NumChannels, false, ChannelMap);
 		float* ChannelMapPtr = ChannelMap.GetData();
 
@@ -543,11 +543,7 @@ namespace Audio
 				}
 				else
 				{
-					// Mix the output of the submix into the output buffer
-					for (int32 i = 0; i < NumSamples; ++i)
-					{
-						BufferPtr[i] += ScratchBufferPtr[i];
-					}
+					Audio::MixInBufferFast(ScratchBufferPtr, BufferPtr, NumSamples);
 				}
 			}
 		}
diff --git a/Engine/Source/Runtime/AudioMixer/Private/Components/SynthComponent.cpp b/Engine/Source/Runtime/AudioMixer/Private/Components/SynthComponent.cpp
--- a/Engine/Source/Runtime/AudioMixer/Private/Components/SynthComponent.cpp
+++ b/Engine/Source/Runtime/AudioMixer/Private/Components/SynthComponent.cpp
@@ -283,7 +283,7 @@ void USynthComponent::OnUnregister()
 bool USynthComponent::IsReadyForOwnerToAutoDestroy() const
 {
 	const bool bIsAudioComponentReadyForDestroy = !AudioComponent || (AudioComponent && !AudioComponent->IsPlaying());
-	const bool bIsSynthSoundReadyForDestroy = !Synth || Synth->bIsReadyForDestroy;
+	const bool bIsSynthSoundReadyForDestroy = !Synth || !Synth->bIsSoundActive;
 	return bIsAudioComponentReadyForDestroy && bIsSynthSoundReadyForDestroy;
 }
 
@@ -373,6 +373,7 @@ void USynthComponent::Start()
 		AudioComponent->AttenuationSettings = AttenuationSettings;
 		AudioComponent->bOverrideAttenuation = bOverrideAttenuation;
 		AudioComponent->bIsUISound = bIsUISound;
+		AudioComponent->bIsPreviewSound = bIsPreviewSound;
 		AudioComponent->bAllowSpatialization = bAllowSpatialization;
 		AudioComponent->ConcurrencySettings = ConcurrencySettings;
 		AudioComponent->AttenuationOverrides = AttenuationOverrides;
diff --git a/Engine/Source/Runtime/AudioMixer/Private/DSP/BufferVectorOperations.cpp b/Engine/Source/Runtime/AudioMixer/Private/DSP/BufferVectorOperations.cpp
new file mode 100644
--- /dev/null
+++ b/Engine/Source/Runtime/AudioMixer/Private/DSP/BufferVectorOperations.cpp
@@ -0,0 +1,287 @@
+// Copyright 1998-2017 Epic Games, Inc. All Rights Reserved.
+
+#include "DSP/BufferVectorOperations.h"
+
+#define AUDIO_USE_SIMD 1
+
+namespace Audio
+{
+	void BufferMultiplyByConstant(const AlignedFloatBuffer& InFloatBuffer, float InValue, AlignedFloatBuffer& OutFloatBuffer)
+	{
+		check(InFloatBuffer.Num() >= 4);
+
+		// Prepare output buffer
+		OutFloatBuffer.Reset();
+		OutFloatBuffer.AddUninitialized(InFloatBuffer.Num());
+
+		check(InFloatBuffer.Num() == OutFloatBuffer.Num());
+
+		const int32 NumSamples = InFloatBuffer.Num();
+
+		// Get ptrs to audio buffers to avoid bounds check in non-shipping builds
+		const float* InBufferPtr = InFloatBuffer.GetData();
+		float* OutBufferPtr = OutFloatBuffer.GetData();
+
+#if !AUDIO_USE_SIMD
+		for (int32 i = 0; i < NumSamples; ++i)
+		{
+			OutBufferPtr[i] = InValue * InBufferPtr[i];
+		}
+#else
+
+		// Can only SIMD on multiple of 4 buffers, we'll do normal multiples on last bit
+		const int32 NumSamplesRemaining = NumSamples % 4;
+		const int32 NumSamplesToSimd = NumSamples - NumSamplesRemaining;
+
+		// Load the single value we want to multiply all values by into a vector register
+		const VectorRegister MultiplyValue = VectorLoadFloat1(&InValue);
+		for (int32 i = 0; i < NumSamplesToSimd; i += 4)
+		{
+			// Load the next 4 samples of the input buffer into a register
+			VectorRegister InputBufferRegister = VectorLoadAligned(&InBufferPtr[i]);
+
+			// Perform the multiply
+			VectorRegister Temp = VectorMultiply(InputBufferRegister, MultiplyValue);
+
+			// Store results into the output buffer
+			VectorStoreAligned(Temp, &OutBufferPtr[i]);
+		}
+
+		// Perform remaining non-simd values left over
+		for (int32 i = 0; i < NumSamplesRemaining; ++i)
+		{
+			OutBufferPtr[NumSamplesToSimd + i] = InValue * InBufferPtr[NumSamplesToSimd + i];
+		}
+#endif
+	}
+
+	void FadeBufferFast(AlignedFloatBuffer& OutFloatBuffer, const float StartValue, const float EndValue)
+	{
+		FadeBufferFast(OutFloatBuffer.GetData(), OutFloatBuffer.Num(), StartValue, EndValue);
+	}
+
+	void FadeBufferFast(float* OutFloatBuffer, int32 NumSamples, const float StartValue, const float EndValue)
+	{
+		checkf(IsAligned<float*>(OutFloatBuffer, 4), TEXT("Memory must be aligned to use vector operations."));
+		checkf(NumSamples % 4 == 0, TEXT("Please use a buffer size that is a multiple of 4."));
+
+#if !AUDIO_USE_SIMD
+		float Gain = StartValue;
+		if (FMath::IsNearlyEqual(StartValue, EndValue))
+		{
+			// No need to do anything if start and end values are both 0.0
+			if (StartValue == 0.0f)
+			{
+				FMemory::Memset(OutFloatBuffer, 0, sizeof(float)*NumSamples);
+			}
+			else
+			{
+				// Only need to do a buffer multiply if start and end values are the same
+				for (int32 i = 0; i < NumSamples; ++i)
+				{
+					OutFloatBuffer[i] = OutFloatBuffer[i] * Gain;
+				}
+			}
+		}
+		else
+		{
+			// Do a fade from start to end
+			const float DeltaValue = ((EndValue - StartValue) / NumSamples);
+			for (int32 i = 0; i < NumSamples; ++i)
+			{
+				OutFloatBuffer[i] = OutFloatBuffer[i] * Gain;
+				Gain += DeltaValue;
+			}
+		}
+#else
+		const int32 NumIterations = NumSamples / 4;
+
+		if (FMath::IsNearlyEqual(StartValue, EndValue))
+		{
+			// No need to do anything if start and end values are both 0.0
+			if (StartValue == 0.0f)
+			{
+				FMemory::Memset(OutFloatBuffer, 0, sizeof(float)*NumSamples);
+			}
+			else
+			{
+				// Only need to do a buffer multiply if start and end values are the same
+				const float DeltaValue = ((EndValue - StartValue) / NumIterations);
+
+				VectorRegister Gain = VectorLoadFloat1(&StartValue);
+
+				for (int32 i = 0; i < NumSamples; i += 4)
+				{
+					VectorRegister Output = VectorLoadAligned(&OutFloatBuffer[i]);
+					Output = VectorMultiply(Output, Gain);
+					VectorStoreAligned(Output, &OutFloatBuffer[i]);
+				}
+			}
+		}
+		else
+		{
+			const float DeltaValue = ((EndValue - StartValue) / NumIterations);
+
+			VectorRegister Gain = VectorLoadFloat1(&StartValue);
+			VectorRegister Delta = VectorLoadFloat1(&DeltaValue);
+
+			for (int32 i = 0; i < NumSamples; i += 4)
+			{
+				VectorRegister Output = VectorLoadAligned(&OutFloatBuffer[i]);
+				Output = VectorMultiply(Output, Gain);
+				Gain = VectorAdd(Gain, Delta);
+				VectorStoreAligned(Output, &OutFloatBuffer[i]);
+			}
+		}
+#endif
+	}
+
+	void MixInBufferFast(const AlignedFloatBuffer& InFloatBuffer, AlignedFloatBuffer& BufferToSumTo, const float Gain)
+	{
+		MixInBufferFast(InFloatBuffer.GetData(), BufferToSumTo.GetData(), InFloatBuffer.Num(), Gain);
+	}
+
+	void MixInBufferFast(const float* InFloatBuffer, float* BufferToSumTo, int32 NumSamples, const float Gain)
+	{
+		checkf(IsAligned<const float*>(InFloatBuffer, 4), TEXT("Memory must be aligned to use vector operations."));
+		checkf(IsAligned<float*>(BufferToSumTo, 4), TEXT("Memory must be aligned to use vector operations."));
+		checkf(NumSamples % 4 == 0, TEXT("Please use a buffer size that is a multiple of 4."));
+
+#if !AUDIO_USE_SIMD
+		for (int32 i = 0; i < NumSamples; ++i)
+		{
+			BufferToSumTo[i] += InFloatBuffer[i] * Gain;
+		}
+#else
+		VectorRegister GainVector = VectorLoadFloat1(&Gain);
+
+		for (int32 i = 0; i < NumSamples; i += 4)
+		{
+			VectorRegister Output = VectorLoadAligned(&BufferToSumTo[i]);
+			VectorRegister Input  = VectorLoadAligned(&InFloatBuffer[i]);
+			Output = VectorMultiplyAdd(Input, GainVector, Output);
+			VectorStoreAligned(Output, &BufferToSumTo[i]);
+		}
+#endif
+	}
+
+	void MixInBufferFast(const float* InFloatBuffer, float* BufferToSumTo, int32 NumSamples)
+	{
+		checkf(IsAligned<const float*>(InFloatBuffer, 4), TEXT("Memory must be aligned to use vector operations."));
+		checkf(IsAligned<float*>(BufferToSumTo, 4), TEXT("Memory must be aligned to use vector operations."));
+		checkf(NumSamples % 4 == 0, TEXT("Please use a buffer size that is a multiple of 4."));
+
+#if !AUDIO_USE_SIMD
+		for (int32 i = 0; i < NumSamples; ++i)
+		{
+			BufferToSumTo[i] += InFloatBuffer[i];
+		}
+#else
+		for (int32 i = 0; i < NumSamples; i += 4)
+		{
+			VectorRegister Output = VectorLoadAligned(&BufferToSumTo[i]);
+			VectorRegister Input = VectorLoadAligned(&InFloatBuffer[i]);
+			Output = VectorAdd(Input, Output);
+			VectorStoreAligned(Output, &BufferToSumTo[i]);
+		}
+#endif
+	}
+
+	void SumBuffers(const AlignedFloatBuffer& InFloatBuffer1, const AlignedFloatBuffer& InFloatBuffer2, AlignedFloatBuffer& OutputBuffer)
+	{
+		SumBuffers(InFloatBuffer1.GetData(), InFloatBuffer2.GetData(), OutputBuffer.GetData(), OutputBuffer.Num());
+	}
+
+	void SumBuffers(const float* InFloatBuffer1, const float* InFloatBuffer2, float* OutputBuffer, int32 NumSamples)
+	{
+		checkf(IsAligned<const float*>(InFloatBuffer1, 4), TEXT("Memory must be aligned to use vector operations."));
+		checkf(IsAligned<const float*>(InFloatBuffer2, 4), TEXT("Memory must be aligned to use vector operations."));
+		checkf(IsAligned<float*>(OutputBuffer, 4), TEXT("Memory must be aligned to use vector operations."));
+		checkf(NumSamples % 4 == 0, TEXT("Please use a buffer size that is a multiple of 4."));
+
+#if !AUDIO_USE_SIMD
+		for (int32 i = 0; i < NumSamples; ++i)
+		{
+			OutputBuffer[i] = InFloatBuffer1[i] + InFloatBuffer2[i];
+		}
+#else
+		for (int32 i = 0; i < NumSamples; i += 4)
+		{
+			VectorRegister Input1 = VectorLoadAligned(&InFloatBuffer1[i]);
+			VectorRegister Input2 = VectorLoadAligned(&InFloatBuffer2[i]);
+
+			VectorRegister Output = VectorAdd(Input1, Input2);
+			VectorStoreAligned(Output, &OutputBuffer[i]);
+		}
+#endif
+	}
+
+	float GetMagnitude(const AlignedFloatBuffer& Buffer)
+	{
+		return GetMagnitude(Buffer.GetData(), Buffer.Num());
+	}
+
+	float GetMagnitude(const float* Buffer, int32 NumSamples)
+	{
+		checkf(NumSamples % 4 == 0, TEXT("Please use a buffer size that is a multiple of 4."));
+
+#if !AUDIO_USE_SIMD
+		float Sum = 0.0f;
+		for (int32 i = 0; i < NumSamples; ++i)
+		{
+			Sum += Buffer[i] * Buffer[i];
+		}
+		return FMath::Sqrt(Sum);
+#else
+		VectorRegister Sum = VectorZero();
+
+		const float Exponent = 2.0f;
+		VectorRegister ExponentVector = VectorLoadFloat1(&Exponent);
+
+		for (int32 i = 0; i < NumSamples; i += 4)
+		{
+			VectorRegister Input = VectorPow(VectorLoadAligned(&Buffer[i]), ExponentVector);
+			Sum = VectorAdd(Sum, Input);
+		}
+
+		float PartionedSums[4];
+		VectorStoreAligned(Sum, PartionedSums);
+
+		return FMath::Sqrt(PartionedSums[0] + PartionedSums[1] + PartionedSums[2] + PartionedSums[3]);
+#endif
+	}
+
+	float GetAverageAmplitude(const AlignedFloatBuffer& Buffer)
+	{
+		checkf(Buffer.Num() % 4 == 0, TEXT("Please use a buffer size that is a multiple of 4."));
+
+		return GetAverageAmplitude(Buffer.GetData(), Buffer.Num());
+	}
+
+	float GetAverageAmplitude(const float* Buffer, int32 NumSamples)
+	{
+		checkf(NumSamples % 4 == 0, TEXT("Please use a buffer size that is a multiple of 4."));
+
+#if !AUDIO_USE_SIMD
+		float Sum = 0.0f;
+		for (int32 i = 0; i < NumSamples; ++i)
+		{
+			Sum += Buffer[i];
+		}
+		return Sum / NumSamples;
+#else
+		VectorRegister Sum = VectorZero();
+
+		for (int32 i = 0; i < NumSamples; i += 4)
+		{
+			VectorRegister Input = VectorAbs(VectorLoadAligned(&Buffer[i]));
+			Sum = VectorAdd(Sum, Input);
+		}
+
+		float PartionedSums[4];
+		VectorStore(Sum, PartionedSums);
+		
+		return (PartionedSums[0] + PartionedSums[1] + PartionedSums[2] + PartionedSums[3]) / NumSamples;
+#endif
+	}
+}
diff --git a/Engine/Source/Runtime/AudioMixer/Private/Effects/AudioMixerSubmixEffectEQ.cpp b/Engine/Source/Runtime/AudioMixer/Private/Effects/AudioMixerSubmixEffectEQ.cpp
--- a/Engine/Source/Runtime/AudioMixer/Private/Effects/AudioMixerSubmixEffectEQ.cpp
+++ b/Engine/Source/Runtime/AudioMixer/Private/Effects/AudioMixerSubmixEffectEQ.cpp
@@ -5,14 +5,6 @@
 #include "Misc/ScopeLock.h"
 #include "AudioMixer.h"
 
-static int32 DisableSubmixEffectEQCvar = 0;
-FAutoConsoleVariableRef CVarDisableSubmixEQ(
-	TEXT("au.DisableSubmixEffectEQ"),
-	DisableSubmixEffectEQCvar,
-	TEXT("Disables the eq submix.\n")
-	TEXT("0: Not Disabled, 1: Disabled"),
-	ECVF_Default);
-
 
 static bool IsEqual(const FSubmixEffectSubmixEQSettings& Left, const FSubmixEffectSubmixEQSettings& Right)
 {
@@ -101,7 +93,7 @@ void FSubmixEffectSubmixEQ::OnProcessAudio(const FSoundEffectSubmixInputData& In
 	Audio::AlignedFloatBuffer& InAudioBuffer = *InData.AudioBuffer;
 	Audio::AlignedFloatBuffer& OutAudioBuffer = *OutData.AudioBuffer;
 
-	if (bEQSettingsSet && !DisableSubmixEffectEQCvar && RenderThreadEQSettings.EQBands.Num() > 0)
+	if (bEQSettingsSet && RenderThreadEQSettings.EQBands.Num() > 0)
 	{
 		// Feed every other channel through the EQ filters
 		int32 NumFilters = InData.NumChannels / 2;
diff --git a/Engine/Source/Runtime/AudioMixer/Private/Effects/AudioMixerSubmixEffectReverb.cpp b/Engine/Source/Runtime/AudioMixer/Private/Effects/AudioMixerSubmixEffectReverb.cpp
--- a/Engine/Source/Runtime/AudioMixer/Private/Effects/AudioMixerSubmixEffectReverb.cpp
+++ b/Engine/Source/Runtime/AudioMixer/Private/Effects/AudioMixerSubmixEffectReverb.cpp
@@ -8,15 +8,6 @@
 
 class UReverbEffect;
 
-static int32 DisableSubmixReverbCVar = 0;
-FAutoConsoleVariableRef CVarDisableSubmixReverb(
-	TEXT("au.DisableReverbSubmix"),
-	DisableSubmixReverbCVar,
-	TEXT("Disables the reverb submix.\n")
-	TEXT("0: Not Disabled, 1: Disabled"),
-	ECVF_Default);
-
-
 FSubmixEffectReverb::FSubmixEffectReverb()
 	: bIsEnabled(false)
 {
@@ -75,7 +66,7 @@ void FSubmixEffectReverb::OnPresetChanged()
 void FSubmixEffectReverb::OnProcessAudio(const FSoundEffectSubmixInputData& InData, FSoundEffectSubmixOutputData& OutData)
 {
 	check(InData.NumChannels == 2);
-	if (OutData.NumChannels < 2 || !bIsEnabled || (DisableSubmixReverbCVar != 0)) 
+	if (OutData.NumChannels < 2 || !bIsEnabled) 
 	{
 		// Not supported
 		return;
diff --git a/Engine/Source/Runtime/AudioMixer/Public/AudioMixerDevice.h b/Engine/Source/Runtime/AudioMixer/Public/AudioMixerDevice.h
--- a/Engine/Source/Runtime/AudioMixer/Public/AudioMixerDevice.h
+++ b/Engine/Source/Runtime/AudioMixer/Public/AudioMixerDevice.h
@@ -139,8 +139,8 @@ namespace Audio
 
 		int32 GetNumOutputFrames() const { return PlatformSettings.CallbackBufferFrameSize; }
 
-		void Get3DChannelMap(const ESubmixChannelFormat InSubmixChannelType, const FWaveInstance* InWaveInstance, const float EmitterAzimuth, const float NormalizedOmniRadius, TArray<float>& OutChannelMap);
-		void Get2DChannelMap(bool bIsVorbis, const ESubmixChannelFormat InSubmixChannelType, const int32 NumSourceChannels, const bool bIsCenterChannelOnly, TArray<float>& OutChannelMap) const;
+		void Get3DChannelMap(const ESubmixChannelFormat InSubmixChannelType, const FWaveInstance* InWaveInstance, const float EmitterAzimuth, const float NormalizedOmniRadius, Audio::AlignedFloatBuffer& OutChannelMap);
+		void Get2DChannelMap(bool bIsVorbis, const ESubmixChannelFormat InSubmixChannelType, const int32 NumSourceChannels, const bool bIsCenterChannelOnly, Audio::AlignedFloatBuffer& OutChannelMap) const;
 
 		int32 GetDeviceSampleRate() const;
 		int32 GetDeviceOutputChannels() const;
@@ -221,6 +221,12 @@ namespace Audio
 		/** Channel type arrays for submix channel types. */
 		TMap<ESubmixChannelFormat, TArray<EAudioMixerChannel::Type>> ChannelArrays;
 
+		/** What upmix method to use for mono channel upmixing. */
+		EMonoChannelUpmixMethod MonoChannelUpmixMethod;
+
+		/** What panning method to use for panning. */
+		EPanningMethod PanningMethod;
+
 		/** The audio output stream parameters used to initialize the audio hardware. */
 		FAudioMixerOpenStreamParams OpenStreamParams;
 
diff --git a/Engine/Source/Runtime/AudioMixer/Public/Components/SynthComponent.h b/Engine/Source/Runtime/AudioMixer/Public/Components/SynthComponent.h
--- a/Engine/Source/Runtime/AudioMixer/Public/Components/SynthComponent.h
+++ b/Engine/Source/Runtime/AudioMixer/Public/Components/SynthComponent.h
@@ -154,6 +154,10 @@ public:
 	UPROPERTY(EditAnywhere, BlueprintReadWrite, Category = Sound)
 	uint8 bIsUISound : 1;
 
+	/** Whether or not this synth is playing as a preview sound */
+	UPROPERTY()
+	uint8 bIsPreviewSound : 1;
+
 	/** Call if creating this synth component not via an actor component in BP, but in code or some other location. Optionally override the sample rate of the sound wave, otherwise it uses the audio device's sample rate. */
 	void Initialize(int32 SampleRateOverride = INDEX_NONE);
 
diff --git a/Engine/Source/Runtime/AudioMixer/Public/DSP/BufferVectorOperations.h b/Engine/Source/Runtime/AudioMixer/Public/DSP/BufferVectorOperations.h
new file mode 100644
--- /dev/null
+++ b/Engine/Source/Runtime/AudioMixer/Public/DSP/BufferVectorOperations.h
@@ -0,0 +1,44 @@
+// Copyright 1998-2017 Epic Games, Inc. All Rights Reserved.
+
+#pragma once
+
+#include "CoreMinimal.h"
+
+#if PLATFORM_SWITCH
+// Switch uses page alignment for submitted buffers
+#define AUDIO_BUFFER_ALIGNMENT 4096
+#else
+#define AUDIO_BUFFER_ALIGNMENT 16
+#endif
+
+namespace Audio
+{
+	typedef TArray<float, TAlignedHeapAllocator<AUDIO_BUFFER_ALIGNMENT>> AlignedFloatBuffer;
+	typedef TArray<uint8, TAlignedHeapAllocator<AUDIO_BUFFER_ALIGNMENT>> AlignedByteBuffer;
+
+	/** Multiplies the input aligned float buffer with the given value. */
+	void BufferMultiplyByConstant(const AlignedFloatBuffer& InFloatBuffer, float InValue, AlignedFloatBuffer& OutFloatBuffer);
+
+	/* Takes a float buffer and quickly interpolates it's gain from StartValue to EndValue. */
+	/* This operation completely ignores channel counts, so avoid using this function on buffers that are not mono, stereo or quad */
+	/* if the buffer needs to fade all channels uniformly. */
+	void FadeBufferFast(AlignedFloatBuffer& OutFloatBuffer, const float StartValue, const float EndValue);
+	void FadeBufferFast(float* OutFloatBuffer, int32 NumSamples, const float StartValue, const float EndValue);
+
+	// Takes buffer InFloatBuffer, optionally multiplies it by Gain, and adds it to BufferToSumTo.
+	void MixInBufferFast(const AlignedFloatBuffer& InFloatBuffer, AlignedFloatBuffer& BufferToSumTo, const float Gain);
+	void MixInBufferFast(const float* InFloatBuffer, float* BufferToSumTo, int32 NumSamples, const float Gain);
+	void MixInBufferFast(const float* InFloatBuffer, float* BufferToSumTo, int32 NumSamples);
+
+	// Sums two buffers together and places the result in the resulting buffer.
+	void SumBuffers(const AlignedFloatBuffer& InFloatBuffer1, const AlignedFloatBuffer& InFloatBuffer2, AlignedFloatBuffer& OutputBuffer);
+	void SumBuffers(const float* InFloatBuffer1, const float* InFloatBuffer2, float* OutputBuffer, int32 NumSamples);
+
+	// Takes an audio buffer and returns the magnitude across that buffer.
+	float GetMagnitude(const AlignedFloatBuffer& Buffer);
+	float GetMagnitude(const float* Buffer, int32 NumSamples);
+
+	// Takes an audio buffer and gets the average absolute amplitude across that buffer.
+	float GetAverageAmplitude(const AlignedFloatBuffer& Buffer);
+	float GetAverageAmplitude(const float* Buffer, int32 NumSamples);
+}
diff --git a/Engine/Source/Runtime/Engine/Classes/Sound/AudioSettings.h b/Engine/Source/Runtime/Engine/Classes/Sound/AudioSettings.h
--- a/Engine/Source/Runtime/Engine/Classes/Sound/AudioSettings.h
+++ b/Engine/Source/Runtime/Engine/Classes/Sound/AudioSettings.h
@@ -46,6 +46,30 @@ enum class EVoiceSampleRate : int32
 	/* High48000Hz = 48000 //TODO: 48k VOIP requires serious performance optimizations on encoding and decoding. */
 };
 
+// Enumeration defines what method of panning to use (for non-binaural audio) with the audio-mixer.
+UENUM()
+enum class EPanningMethod : int8
+{
+	// Linear panning maintains linear amplitude when panning between speakers.
+	Linear,
+
+	// Equal power panning maintains equal power when panning between speakers.
+	EqualPower
+};
+
+// Enumeration defines how to treat mono 2D playback. Mono sounds need to upmixed to stereo when played back in 2D.
+UENUM()
+enum class EMonoChannelUpmixMethod : int8
+{
+	// The mono channel is split 0.5 left/right
+	Linear,
+	
+	// The mono channel is split 0.707 left/right
+	EqualPower,
+
+	// The mono channel is split 1.0 left/right
+	FullVolume
+};
 
 USTRUCT()
 struct ENGINE_API FAudioQualitySettings
@@ -141,6 +165,18 @@ class ENGINE_API UAudioSettings : public UDeveloperSettings
 	UPROPERTY(config, EditAnywhere, Category = "Quality", AdvancedDisplay)
 	uint32 NumStoppingSources;
 
+	/**
+	* The method to use when doing non-binaural or object-based panning.
+	*/
+	UPROPERTY(config, EditAnywhere, Category = "Quality", AdvancedDisplay)
+	EPanningMethod PanningMethod;
+
+	/**
+	* The upmixing method for mono sound sources. Defines up mono channels are up-mixed to stereo channels.
+	*/
+	UPROPERTY(config, EditAnywhere, Category = "Quality", AdvancedDisplay)
+	EMonoChannelUpmixMethod MonoChannelUpmixMethod;
+
 	/**
 	 * The format string to use when generating the filename for contexts within dialogue waves. This must generate unique names for your project.
 	 * Available format markers:
diff --git a/Engine/Source/Runtime/Engine/Classes/Sound/SoundWave.h b/Engine/Source/Runtime/Engine/Classes/Sound/SoundWave.h
--- a/Engine/Source/Runtime/Engine/Classes/Sound/SoundWave.h
+++ b/Engine/Source/Runtime/Engine/Classes/Sound/SoundWave.h
@@ -163,6 +163,12 @@ class ENGINE_API USoundWave : public USoundBase
 	/** Set to true for programmatically-generated, streamed audio. */
 	uint8 bProcedural:1;
 
+	/** Whether this sound wave is beginning to be destroyed by GC. */
+	uint8 bIsBeginDestroy:1;
+
+	/** Whether this sound is actively being used by the audio renderer (in audio mixer). Prevents GC issues with rendering audio. */
+	uint8 bIsSoundActive:1;
+
 	/** Set to true of this is a bus sound source. This will result in the sound wave not generating audio for itself, but generate audio through instances. Used only in audio mixer. */
 	uint8 bIsBus:1;
 
@@ -336,6 +342,10 @@ public:
 
 	FFormatContainer			CompressedFormatData;
 
+#if WITH_EDITORONLY_DATA
+	TMap<FName, uint32> AsyncLoadingDataFormats;
+#endif
+
 	/** Resource index to cross reference with buffers */
 	int32 ResourceID;
 
@@ -357,6 +367,7 @@ public:
 	virtual bool IsReadyForFinishDestroy() override;
 	virtual void FinishDestroy() override;
 	virtual void PostLoad() override;
+	virtual void BeginDestroy() override;
 #if WITH_EDITOR
 	virtual void PostEditChangeProperty(FPropertyChangedEvent& PropertyChangedEvent) override;	
 #endif // WITH_EDITOR
@@ -374,6 +385,12 @@ public:
 	virtual bool IsAllowedVirtual() const override;
 	//~ End USoundBase Interface.
 
+	// Called  when the procedural sound wave begins on the render thread. Only used in the audio mixer and when bProcedural is true.
+	virtual void OnBeginGenerate() {}
+
+	// Called when the procedural sound wave is done generating on the render thread. Only used in the audio mixer and when bProcedural is true..
+	virtual void OnEndGenerate() {};
+
 	/**
 	* Overwrite sample rate. Used for procedural soundwaves, as well as sound waves that are resampled on compress/decompress.
 	*/
@@ -462,6 +479,12 @@ public:
 
 	virtual bool HasCompressedData(FName Format, ITargetPlatform* TargetPlatform = GetRunningPlatform()) const;
 
+private:
+	FName GetPlatformSpecificFormat(FName Format, const FPlatformAudioCookOverrides* CompressionOverrides);
+
+public:
+	virtual void BeginGetCompressedData(FName Format, const FPlatformAudioCookOverrides* CompressionOverrides);
+
 	/** 
 	 * Gets the compressed data from derived data cache for the specified platform
 	 * Warning, the returned pointer isn't valid after we add new formats
@@ -470,7 +493,7 @@ public:
 	 * @param PlatformName optional name of platform we are getting compressed data for.
 	 * @param CompressionOverrides optional platform compression overrides
 	 * @return	compressed data, if it could be obtained
-	 */ 
+	 */
 	virtual FByteBulkData* GetCompressedData(FName Format, const FPlatformAudioCookOverrides* CompressionOverrides = GetPlatformCompressionOverridesForCurrentPlatform());
 
 	/** 
diff --git a/Engine/Source/Runtime/Engine/Classes/Sound/SoundWaveProcedural.h b/Engine/Source/Runtime/Engine/Classes/Sound/SoundWaveProcedural.h
--- a/Engine/Source/Runtime/Engine/Classes/Sound/SoundWaveProcedural.h
+++ b/Engine/Source/Runtime/Engine/Classes/Sound/SoundWaveProcedural.h
@@ -50,17 +50,16 @@ protected:
 
 public:
 	USoundWaveProcedural(const FObjectInitializer& ObjectInitializer);
-	USoundWaveProcedural(FVTableHelper& Helper);
 
 	//~ Begin UObject Interface. 
 	virtual void Serialize( FArchive& Ar ) override;
 	virtual void GetAssetRegistryTags(TArray<FAssetRegistryTag>& OutTags) const override;
-	virtual bool IsReadyForFinishDestroy() override;
 	//~ End UObject Interface. 
 
 	//~ Begin USoundWave Interface.
 	virtual int32 GeneratePCMData(uint8* PCMData, const int32 SamplesNeeded) override;
 	virtual bool HasCompressedData(FName Format, ITargetPlatform* TargetPlatform) const override;
+	virtual void BeginGetCompressedData(FName Format, const FPlatformAudioCookOverrides* CompressionOverrides) override;
 	virtual FByteBulkData* GetCompressedData(FName Format, const FPlatformAudioCookOverrides* CompressionOverrides = nullptr) override;
 	virtual void InitAudioResource( FByteBulkData& CompressedData ) override;
 	virtual bool InitAudioResource(FName Format) override;
@@ -71,12 +70,6 @@ public:
 	// Returns number of samples generated
 	virtual int32 OnGeneratePCMAudio(TArray<uint8>& OutAudio, int32 NumSamples) { return 0; }
 
-	// Called  when the procedural sound wave begins on the render thread. Only used in the audio mixer.
-	virtual void OnBeginGenerate() {}
-
-	// Called when the procedural sound wave is done generating on the render thread. Only used in the audio mixer.
-	virtual void OnEndGenerate() {};
-
 	/** Add data to the FIFO that feeds the audio device. */
 	void QueueAudio(const uint8* AudioData, const int32 BufferSize);
 
@@ -91,7 +84,4 @@ public:
 
 	/** Size in bytes of a single sample of audio in the procedural audio buffer. */
 	int32 SampleByteSize;
-
-	/** Whether or not this object is ready to be destroyed. Allows procedural sound wave generation to occur in async tasks without garbage collection deleting it from underneath. */
-	FThreadSafeBool bIsReadyForDestroy;
 };
diff --git a/Engine/Source/Runtime/Engine/Private/ADPCMAudioInfo.cpp b/Engine/Source/Runtime/Engine/Private/ADPCMAudioInfo.cpp
--- a/Engine/Source/Runtime/Engine/Private/ADPCMAudioInfo.cpp
+++ b/Engine/Source/Runtime/Engine/Private/ADPCMAudioInfo.cpp
@@ -386,7 +386,7 @@ bool FADPCMAudioInfo::StreamCompressedData(uint8* Destination, bool bLooping, ui
 					{
 						// If we did not get it then just bail, CurrentChunkIndex will not get incremented on the next callback so in effect another attempt will be made to fetch the chunk
 						// Since audio streaming depends on the general data streaming mechanism used by other parts of the engine and new data is prefectched on the game tick thread its possible a game hickup can cause this
-						UE_LOG(LogAudio, Warning, TEXT("Missed Deadline chunk %d"), CurrentChunkIndex);
+						UE_LOG(LogAudio, Log, TEXT("Missed Deadline chunk %d"), CurrentChunkIndex);
 						
 						// zero out remaining data and bail
 						memset(OutData, 0, BufferSize);
diff --git a/Engine/Source/Runtime/Engine/Private/ActiveSound.cpp b/Engine/Source/Runtime/Engine/Private/ActiveSound.cpp
--- a/Engine/Source/Runtime/Engine/Private/ActiveSound.cpp
+++ b/Engine/Source/Runtime/Engine/Private/ActiveSound.cpp
@@ -516,6 +516,15 @@ void FActiveSound::UpdateWaveInstances( TArray<FWaveInstance*> &InWaveInstances,
 	}
 	else if (ThisSoundsWaveInstances.Num() > 0)
 	{
+		// Let the wave instance know that this active sound is stopping. This will result in the wave instance getting a lower sort for voice prioritization
+		if (IsStopping())
+		{
+			for (FWaveInstance* WaveInstance : ThisSoundsWaveInstances)
+			{
+				WaveInstance->SetStopping(true);
+			}
+		}
+
 		// If this active sound is told to limit concurrency by the quietest sound
 		const FSoundConcurrencySettings* ConcurrencySettingsToApply = GetSoundConcurrencySettingsToApply();
 		if (ConcurrencySettingsToApply && ConcurrencySettingsToApply->ResolutionRule == EMaxConcurrentResolutionRule::StopQuietest)
@@ -592,9 +601,9 @@ void FActiveSound::Stop(bool bStopNow)
 		if (Source)
 		{
 			bool bStopped = false;
-			if (AudioDevice->IsAudioMixerEnabled())
+			if (AudioDevice->IsAudioMixerEnabled() && AudioDevice->IsStoppingVoicesEnabled())
 			{
-				if (bStopNow || !AudioDevice->GetNumFreeSources() || AudioDevice->CurrentStoppingVoiceCount == AudioDevice->NumStoppingVoices)
+				if (bStopNow || !AudioDevice->GetNumFreeSources())
 				{
 					Source->StopNow();
 					bStopped = true;
@@ -621,16 +630,14 @@ void FActiveSound::Stop(bool bStopNow)
 		if (Source)
 		{
 			if (!Source->IsStopping())
-			{
+			{		
+				Source->StopNow();
+
 				delete WaveInstance;
 				WaveInstance = nullptr;
 			}
 			else
 			{
-				// Increment the stopping voice count
-				AudioDevice->CurrentStoppingVoiceCount++;
-				check(AudioDevice->CurrentStoppingVoiceCount <= AudioDevice->NumStoppingVoices);
-
 				// This source is doing a fade out, so stopping. Can't remove the wave instance yet.
 				bIsStopping = true;
 			}
@@ -667,6 +674,8 @@ bool FActiveSound::UpdateStoppingSources(uint64 CurrentTick, bool bEnsureStopped
 		return true;
 	}
 
+	bIsStopping = false;
+
 	for (auto WaveInstanceIt(WaveInstances.CreateIterator()); WaveInstanceIt; ++WaveInstanceIt)
 	{
 		FWaveInstance*& WaveInstance = WaveInstanceIt.Value();
@@ -681,30 +690,23 @@ bool FActiveSound::UpdateStoppingSources(uint64 CurrentTick, bool bEnsureStopped
 				// We should have a stopping source here
 				check(Source->IsStopping());
 
-				// Wait until the source has finished stopping
-				if (bEnsureStopped)
-				{
-					Source->EnsureStopped();
-				}
-
 				// The source has finished (totally faded out)
 				if (Source->IsFinished() || bEnsureStopped)
 				{
-					// Release source resources now safely
-					Source->ReleaseStoppedSound();
+					Source->StopNow();
 
 					// Delete the wave instance
 					delete WaveInstance;
 					WaveInstance = nullptr;
-					AudioDevice->CurrentStoppingVoiceCount--;
-					check(AudioDevice->CurrentStoppingVoiceCount >= 0);
 				}
 				else
 				{
+					// We are not finished yet so touch it
 					Source->LastUpdate = CurrentTick;
 					Source->LastHeardUpdate = CurrentTick;
-					// One of sources has not finished fading out, just return false now
-					return false;
+
+					// flag that we're still stopping (return value)
+					bIsStopping = true;
 				}
 			}
 			else
@@ -718,8 +720,14 @@ bool FActiveSound::UpdateStoppingSources(uint64 CurrentTick, bool bEnsureStopped
 
 	// Return true to indicate this active sound can be cleaned up
 	// If we've reached this point, all sound waves have stopped so we can clear this wave instance out.
-	WaveInstances.Reset();
-	return true;
+	if (!bIsStopping)
+	{
+		WaveInstances.Reset();
+		return true;
+	}
+
+	// still stopping!
+	return false;
 }
 
 FWaveInstance* FActiveSound::FindWaveInstance( const UPTRINT WaveInstanceHash )
diff --git a/Engine/Source/Runtime/Engine/Private/Audio.cpp b/Engine/Source/Runtime/Engine/Private/Audio.cpp
--- a/Engine/Source/Runtime/Engine/Private/Audio.cpp
+++ b/Engine/Source/Runtime/Engine/Private/Audio.cpp
@@ -581,6 +581,7 @@ FSpatializationParams FSoundSource::GetSpatializationParams()
 void FSoundSource::InitCommon()
 {
 	PlaybackTime = 0.0f;
+	TickCount = 0;
 
 	// Reset pause state
 	bIsPausedByGame = false;
@@ -785,6 +786,7 @@ FWaveInstance::FWaveInstance( FActiveSound* InActiveSound )
 	, bCenterChannelOnly(false)
 	, bReportedSpatializationWarning(false)
 	, bIsAmbisonics(false)
+	, bIsStopping(false)
 	, SpatializationMethod(ESoundSpatializationAlgorithm::SPATIALIZATION_Default)
 	, OcclusionPluginSettings(nullptr)
 	, OutputTarget(EAudioOutputTarget::Speaker)
@@ -897,23 +899,29 @@ bool FWaveInstance::ShouldStopDueToMaxConcurrency() const
 
 float FWaveInstance::GetVolumeWeightedPriority() const
 {
-	if (WaveData && WaveData->bBypassVolumeScaleForPriority)
+	// This will result in zero-volume sounds still able to be sorted due to priority but give non-zero volumes higher priority than 0 volumes
+	float ActualVolume = GetVolumeWithDistanceAttenuation();
+	if (ActualVolume > 0.0f)
 	{
-		return Priority;
-	}
-	else
-	{
-		// This will result in zero-volume sounds still able to be sorted due to priority but give non-zero volumes higher priority than 0 volumes
-		float ActualVolume = GetVolumeWithDistanceAttenuation();
-		if (ActualVolume > 0.0f)
+		// Only check for bypass if the actual volume is greater than 0.0
+		if (WaveData && WaveData->bBypassVolumeScaleForPriority)
 		{
-			return ActualVolume * Priority;
+			return Priority;
 		}
 		else
 		{
-			return Priority - MAX_SOUND_PRIORITY - 1;
+			return ActualVolume * Priority;
 		}
 	}
+	else if (IsStopping())
+	{
+		// Stopping sounds will be sorted above 0-volume sounds
+		return ActualVolume * Priority - MAX_SOUND_PRIORITY - 1.0f;
+	}
+	else
+	{
+		return Priority - 2.0f * MAX_SOUND_PRIORITY - 1.0f;
+	}
 }
 
 bool FWaveInstance::IsStreaming() const
diff --git a/Engine/Source/Runtime/Engine/Private/AudioDevice.cpp b/Engine/Source/Runtime/Engine/Private/AudioDevice.cpp
--- a/Engine/Source/Runtime/Engine/Private/AudioDevice.cpp
+++ b/Engine/Source/Runtime/Engine/Private/AudioDevice.cpp
@@ -49,6 +49,14 @@ FAutoConsoleVariableRef CVarSetAudioChannelCount(
 	TEXT("0: Disable, >0: Enable"),
 	ECVF_Default);
 
+static int32 DisableStoppingVoicesCvar = 0;
+FAutoConsoleVariableRef CVarDisableStoppingVoices(
+	TEXT("au.DisableStoppingVoices"),
+	DisableStoppingVoicesCvar,
+	TEXT("Disables stopping voices feature.\n")
+	TEXT("0: Not Disabled, 1: Disabled"),
+	ECVF_Default);
+
 
 /*-----------------------------------------------------------------------------
 FDynamicParameter implementation.
@@ -112,7 +120,6 @@ void FDynamicParameter::Update(float DeltaTime)
 FAudioDevice::FAudioDevice()
 	: MaxChannels(0)
 	, NumStoppingVoices(32)
-	, CurrentStoppingVoiceCount(0)
 	, MaxWaveInstances(0)
 	, SampleRate(0)
 	, CommonAudioPoolSize(0)
@@ -141,6 +148,7 @@ FAudioDevice::FAudioDevice()
 	, bGameWasTicking(true)
 	, bDisableAudioCaching(false)
 	, bIsAudioDeviceHardwareInitialized(false)
+	, bIsStoppingVoicesEnabled(false)
 	, bAudioMixerModuleLoaded(false)
 	, bSpatializationIsExternalSend(false)
 	, bOcclusionIsExternalSend(false)
@@ -215,6 +223,8 @@ bool FAudioDevice::Init(int32 InMaxChannels)
 		PlatformAudioHeadroom = FMath::Pow(10.0f, Headroom / 20.0f);
 	}
 
+	bIsStoppingVoicesEnabled = !DisableStoppingVoicesCvar;
+
 	const UAudioSettings* AudioSettings = GetDefault<UAudioSettings>();
 
 	bAllowCenterChannel3DPanning = AudioSettings->bAllowCenterChannel3DPanning;
@@ -235,7 +245,7 @@ bool FAudioDevice::Init(int32 InMaxChannels)
 	InitSoundEffectPresets();
 
 	// Audio mixer needs to create effects manager before initializing the plugins.
-	if (IsAudioMixerEnabled())
+	if (IsAudioMixerEnabled() && IsStoppingVoicesEnabled())
 	{
 		// create a platform specific effects manager
 		Effects = CreateEffectsManager();
@@ -364,23 +374,23 @@ void FAudioDevice::PrecacheStartupSounds()
 
 void FAudioDevice::SetMaxChannels(int32 InMaxChannels)
 {
-	if (InMaxChannels + NumStoppingVoices > Sources.Num())
+	if (InMaxChannels > Sources.Num())
 	{
 		UE_LOG(LogAudio, Warning, TEXT("Can't increase channels past starting number!"));
 		return;
 	}
 
-	MaxChannels = InMaxChannels + NumStoppingVoices;
+	MaxChannels = InMaxChannels;
 }
 
 int32 FAudioDevice::GetMaxChannels() const
 {
 	if (AudioChannelCountCVar > 0 && AudioChannelCountCVar < Sources.Num())
 	{
-		return AudioChannelCountCVar + NumStoppingVoices;
+		return AudioChannelCountCVar;
 	}
 
-	return MaxChannels + NumStoppingVoices;
+	return MaxChannels;
 }
 
 void FAudioDevice::Teardown()
@@ -1595,7 +1605,7 @@ void FAudioDevice::InitSoundSources()
 	if (Sources.Num() == 0)
 	{
 		// now create platform specific sources
-		const int32 Channels = GetMaxChannels();
+		const int32 Channels = GetMaxChannels() + NumStoppingVoices;
 		for (int32 SourceIndex = 0; SourceIndex < Channels; SourceIndex++)
 		{
 			FSoundSource* Source = CreateSoundSource();
@@ -3077,6 +3087,87 @@ void FAudioDevice::UpdateActiveSoundPlaybackTime(bool bIsGameTicking)
 
 }
 
+void FAudioDevice::StopOldestStoppingSource()
+{
+	if (IsAudioMixerEnabled())
+	{
+		check(!FreeSources.Num());
+
+		FSoundSource* OldestStoppingSource = nullptr;
+		FSoundSource* LowestPriSource = nullptr;
+		FSoundSource* LowestPriNonLoopingSource = nullptr;
+
+		for (FSoundSource* Source : Sources)
+		{
+			// Find oldest stopping voice first
+			if (Source->IsStopping())
+			{
+				if (!OldestStoppingSource)
+				{
+					OldestStoppingSource = Source;
+				}
+				else
+				{
+					if (Source->TickCount > OldestStoppingSource->TickCount)
+					{
+						OldestStoppingSource = Source;
+					}
+				}
+			}
+			else if (Source->WaveInstance)
+			{
+				// Find lowest volume/priority non-looping source as fallback
+				if (Source->WaveInstance->LoopingMode != ELoopingMode::LOOP_Forever && !Source->WaveInstance->bIsUISound)
+				{
+					if (!LowestPriNonLoopingSource)
+					{
+						LowestPriNonLoopingSource = Source;
+					}
+					else
+					{
+						if (Source->WaveInstance->GetVolumeWeightedPriority() < LowestPriNonLoopingSource->WaveInstance->GetVolumeWeightedPriority())
+						{
+							LowestPriNonLoopingSource = Source;
+						}
+					}
+				}
+			
+
+				// Find lowest volume/priority source as final fallback
+				if (!LowestPriSource)
+				{
+					LowestPriSource = Source;
+				}
+				else
+				{
+					if (Source->WaveInstance->GetVolumeWeightedPriority() < LowestPriSource->WaveInstance->GetVolumeWeightedPriority())
+					{
+						LowestPriSource = Source;
+					}
+				}
+			}
+		}
+
+		// Stop oldest stopping source
+		if (OldestStoppingSource)
+		{
+			OldestStoppingSource->StopNow();
+		}
+		// If no oldest stopping source, stop oldest one-shot
+		else if (LowestPriNonLoopingSource)
+		{
+			LowestPriNonLoopingSource->StopNow();
+		}
+		// Otherwise stop oldest source.
+		else
+		{
+			check(LowestPriSource);
+			LowestPriSource->StopNow();
+		}
+		check(FreeSources.Num() > 0);
+	}
+}
+
 void FAudioDevice::StopSources(TArray<FWaveInstance*>& WaveInstances, int32 FirstActiveIndex)
 {
 	SCOPED_NAMED_EVENT(FAudioDevice_StopSources, FColor::Blue);
@@ -3110,17 +3201,20 @@ void FAudioDevice::StopSources(TArray<FWaveInstance*>& WaveInstances, int32 Firs
 	{
 		FSoundSource* Source = Sources[SourceIndex];
 
-		if (Source->WaveInstance && !Source->IsStopping())
+		if (Source->WaveInstance)
 		{
 			// If we need to stop this sound due to max concurrency (i.e. it was quietest in a concurrency group)
-			if (Source->WaveInstance->ShouldStopDueToMaxConcurrency())
-			{
-				Source->Stop();
-			}
-			// Source was not one of the active sounds this tick so needs to be stopped
-			else if (Source->LastUpdate != CurrentTick)
+			if (Source->WaveInstance->ShouldStopDueToMaxConcurrency() || Source->LastUpdate != CurrentTick)
 			{
-				Source->Stop();
+				if (!Source->IsStopping())
+				{
+					Source->Stop();
+				}
+				else
+				{
+					// Still do update even if stopping
+					Source->Update();
+				}
 			}
 			else
 			{
@@ -3176,7 +3270,7 @@ void FAudioDevice::StartSources(TArray<FWaveInstance*>& WaveInstances, int32 Fir
 		{
 			continue;
 		}
-		
+
 		// Editor uses bIsUISound for sounds played in the browser.
 		if (!WaveInstance->ShouldStopDueToMaxConcurrency() && (bGameTicking || WaveInstance->bIsUISound))
 		{
@@ -3185,12 +3279,13 @@ void FAudioDevice::StartSources(TArray<FWaveInstance*>& WaveInstances, int32 Fir
 				(!WaveInstance->IsStreaming() ||
 				IStreamingManager::Get().GetAudioStreamingManager().CanCreateSoundSource(WaveInstance)))
 			{
-
-				if (!FreeSources.Num())
+				// Check for full sources and stop the oldest stopping source
+				if (IsAudioMixerEnabled() && !FreeSources.Num())
 				{
-					continue;
+					StopOldestStoppingSource();
 				}
 
+				check(FreeSources.Num());
 				Source = FreeSources.Pop();
 				check(Source);
 
@@ -3256,7 +3351,7 @@ void FAudioDevice::StartSources(TArray<FWaveInstance*>& WaveInstances, int32 Fir
 						// If the source didn't get paused while initializing, then play it
 						if (!Source->IsPaused())
 						{
-						Source->Play();
+							Source->Play();
 						}
 					}
 					else
@@ -3328,6 +3423,8 @@ void FAudioDevice::Update(bool bGameTicking)
 		}
 	}
 
+	bIsStoppingVoicesEnabled = !DisableStoppingVoicesCvar;
+
 	UpdateAudioPluginSettingsObjectCache();
 
 	// Updates hardware timing logic. Only implemented in audio mixer.
@@ -3550,16 +3647,19 @@ void FAudioDevice::SendUpdateResultsToGameThread(const int32 FirstActiveIndex)
 			}
 		}
 
-		USoundMix* CurrentEQMix = Effects->GetCurrentEQMix();
-
-		for (const TPair<USoundMix*, FSoundMixState>& SoundMixPair : SoundMixModifiers)
+		if (Effects)
 		{
-			StatSoundMixes.AddDefaulted();
-			FAudioStats::FStatSoundMix& StatSoundMix = StatSoundMixes.Last();
-			StatSoundMix.MixName = SoundMixPair.Key->GetName();
-			StatSoundMix.InterpValue = SoundMixPair.Value.InterpValue;
-			StatSoundMix.RefCount = SoundMixPair.Value.ActiveRefCount + SoundMixPair.Value.PassiveRefCount;
-			StatSoundMix.bIsCurrentEQ = (SoundMixPair.Key == CurrentEQMix);
+			USoundMix* CurrentEQMix = Effects->GetCurrentEQMix();
+
+			for (const TPair<USoundMix*, FSoundMixState>& SoundMixPair : SoundMixModifiers)
+			{
+				StatSoundMixes.AddDefaulted();
+				FAudioStats::FStatSoundMix& StatSoundMix = StatSoundMixes.Last();
+				StatSoundMix.MixName = SoundMixPair.Key->GetName();
+				StatSoundMix.InterpValue = SoundMixPair.Value.InterpValue;
+				StatSoundMix.RefCount = SoundMixPair.Value.ActiveRefCount + SoundMixPair.Value.PassiveRefCount;
+				StatSoundMix.bIsCurrentEQ = (SoundMixPair.Key == CurrentEQMix);
+			}
 		}
 	}
 #endif
@@ -3713,6 +3813,15 @@ void FAudioDevice::AddNewActiveSound(const FActiveSound& NewActiveSound)
 		return;
 	}
 
+	USoundWave* SoundWave = Cast<USoundWave>(Sound);
+	if (SoundWave && SoundWave->bProcedural && SoundWave->bIsSoundActive)
+	{
+		FString SoundWaveName;
+		SoundWave->GetName(SoundWaveName);
+		UE_LOG(LogAudio, Warning, TEXT("Replaying a procedural sound '%s' without stopping the previous instance. Only one sound instance per procedural sound wave is supported."), *SoundWaveName)
+		return;
+	}
+
 	// Evaluate concurrency. This will create an ActiveSound ptr which is a copy of NewActiveSound if the sound can play.
 	FActiveSound* ActiveSound = nullptr;
 
@@ -3801,7 +3910,8 @@ void FAudioDevice::ProcessingPendingActiveSoundStops(bool bForceDelete)
 				if (ActiveSound->IsStopping())
 				{
 					// Make sure this sound finishes stopping if we're forcing all sounds to stop due to a flush, etc.
-					ActiveSound->UpdateStoppingSources(CurrentTick, true);
+					bool bIsNowStopped = ActiveSound->UpdateStoppingSources(CurrentTick, true);
+					check(bIsNowStopped);
 				}
 			}
 			else if (ActiveSound->IsStopping())
diff --git a/Engine/Source/Runtime/Engine/Private/AudioSettings.cpp b/Engine/Source/Runtime/Engine/Private/AudioSettings.cpp
--- a/Engine/Source/Runtime/Engine/Private/AudioSettings.cpp
+++ b/Engine/Source/Runtime/Engine/Private/AudioSettings.cpp
@@ -69,7 +69,7 @@ void UAudioSettings::AddDefaultSettings()
 	VoiPSampleRate = EVoiceSampleRate::Low16000Hz;
 	VoipBufferingDelay = 0.2f;
 	MaxWaveInstances = 100;
-	NumStoppingSources = 32;
+	NumStoppingSources = 8;
 }
 
 #if WITH_EDITOR
diff --git a/Engine/Source/Runtime/Engine/Private/SoundNodeWavePlayer.cpp b/Engine/Source/Runtime/Engine/Private/SoundNodeWavePlayer.cpp
--- a/Engine/Source/Runtime/Engine/Private/SoundNodeWavePlayer.cpp
+++ b/Engine/Source/Runtime/Engine/Private/SoundNodeWavePlayer.cpp
@@ -33,6 +33,13 @@ void USoundNodeWavePlayer::LoadAsset(bool bAddToRoot)
 	if (IsAsyncLoading())
 	{
 		SoundWave = SoundWaveAssetPtr.Get();
+		if (SoundWave && SoundWave->HasAnyFlags(RF_NeedLoad))
+		{			
+			// This can happen when the owning USoundCue's PostLoad gets called and the SoundWave hasn't been serialized yet
+			// In this case we need to make sure we don't pass the pointer to the SoundNodeWavePlayer too early as the SoundWave
+			// will be serialized on the AsyncLoadingThread shortly and this may lead to strange race conditions / thread safety issues
+			SoundWave = nullptr;
+		}
 		if (SoundWave == nullptr)
 		{
 			const FString LongPackageName = SoundWaveAssetPtr.GetLongPackageName();
@@ -48,7 +55,7 @@ void USoundNodeWavePlayer::LoadAsset(bool bAddToRoot)
 		}
 		if (SoundWave)
 		{
-			SoundWave->AddToCluster(this);
+			SoundWave->AddToCluster(this, true);
 		}
 	}
 	else
diff --git a/Engine/Source/Runtime/Engine/Private/SoundWave.cpp b/Engine/Source/Runtime/Engine/Private/SoundWave.cpp
--- a/Engine/Source/Runtime/Engine/Private/SoundWave.cpp
+++ b/Engine/Source/Runtime/Engine/Private/SoundWave.cpp
@@ -417,15 +417,9 @@ const FPlatformAudioCookOverrides* USoundWave::GetPlatformCompressionOverridesFo
 	return FPlatformCompressionUtilities::GetCookOverridesForCurrentPlatform();
 }
 
-FByteBulkData* USoundWave::GetCompressedData(FName Format, const FPlatformAudioCookOverrides* CompressionOverrides)
+FName USoundWave::GetPlatformSpecificFormat(FName Format, const FPlatformAudioCookOverrides* CompressionOverrides)
 {
-	if (IsTemplate() || IsRunningDedicatedServer())
-	{
-		return nullptr;
-	}
-
 	// Platforms that require compression overrides get concatenated formats.
-	
 #if WITH_EDITOR
 	FName PlatformSpecificFormat;
 	if (CompressionOverrides)
@@ -465,19 +459,73 @@ FByteBulkData* USoundWave::GetCompressedData(FName Format, const FPlatformAudioC
 
 #endif // WITH_EDITOR
 
+	return PlatformSpecificFormat;
+}
+
+void USoundWave::BeginGetCompressedData(FName Format, const FPlatformAudioCookOverrides* CompressionOverrides)
+{
+#if WITH_EDITOR
+	if (IsTemplate() || IsRunningDedicatedServer())
+	{
+		return;
+	}
+
+	FName PlatformSpecificFormat = GetPlatformSpecificFormat(Format, CompressionOverrides);
+
+	if (!CompressedFormatData.Contains(PlatformSpecificFormat) && !AsyncLoadingDataFormats.Contains(PlatformSpecificFormat))
+	{
+		if (GetDerivedDataCache())
+		{
+			FDerivedAudioDataCompressor* DeriveAudioData = new FDerivedAudioDataCompressor(this, Format, PlatformSpecificFormat, CompressionOverrides);
+			uint32 GetHandle = GetDerivedDataCacheRef().GetAsynchronous(DeriveAudioData);
+			AsyncLoadingDataFormats.Add(PlatformSpecificFormat, GetHandle);
+		}
+		else
+		{
+			UE_LOG(LogAudio, Error, TEXT("Attempt to access the DDC when there is none available on sound '%s', format = %s."), *GetFullName(), *PlatformSpecificFormat.ToString());
+		}
+	}
+#else
+	// No async DDC read in non-editor, nothing to precache
+#endif
+}
+
+FByteBulkData* USoundWave::GetCompressedData(FName Format, const FPlatformAudioCookOverrides* CompressionOverrides)
+{
+	if (IsTemplate() || IsRunningDedicatedServer())
+	{
+		return nullptr;
+	}
+	
+	FName PlatformSpecificFormat = GetPlatformSpecificFormat(Format, CompressionOverrides);
+
 	bool bContainedValidData = CompressedFormatData.Contains(PlatformSpecificFormat);
 	FByteBulkData* Result = &CompressedFormatData.GetFormat(PlatformSpecificFormat);
 	if (!bContainedValidData)
 	{
 		if (!FPlatformProperties::RequiresCookedData() && GetDerivedDataCache())
 		{
-
 			TArray<uint8> OutData;
-			FDerivedAudioDataCompressor* DeriveAudioData = new FDerivedAudioDataCompressor(this, Format, PlatformSpecificFormat, CompressionOverrides);
+			bool bDataWasBuilt = false;
+			bool bGetSuccessful = false;
 
 			COOK_STAT(auto Timer = SoundWaveCookStats::UsageStats.TimeSyncWork());
-			bool bDataWasBuilt = false;
-			if (GetDerivedDataCacheRef().GetSynchronous(DeriveAudioData, OutData, &bDataWasBuilt))
+#if WITH_EDITOR
+			uint32* AsyncHandle = AsyncLoadingDataFormats.Find(PlatformSpecificFormat);
+			if (AsyncHandle)
+			{
+				GetDerivedDataCacheRef().WaitAsynchronousCompletion(*AsyncHandle);
+				bGetSuccessful = GetDerivedDataCacheRef().GetAsynchronousResults(*AsyncHandle, OutData, &bDataWasBuilt);
+				AsyncLoadingDataFormats.Remove(PlatformSpecificFormat);
+			}
+			else
+#endif
+			{
+				FDerivedAudioDataCompressor* DeriveAudioData = new FDerivedAudioDataCompressor(this, Format, PlatformSpecificFormat, CompressionOverrides);
+				bGetSuccessful = GetDerivedDataCacheRef().GetSynchronous(DeriveAudioData, OutData, &bDataWasBuilt);
+			}
+
+			if (bGetSuccessful)
 			{
 				COOK_STAT(Timer.AddHitOrMiss(bDataWasBuilt ? FCookStats::CallStats::EHitOrMiss::Miss : FCookStats::CallStats::EHitOrMiss::Hit, OutData.Num()));
 				Result->Lock(LOCK_READ_WRITE);
@@ -536,7 +584,7 @@ void USoundWave::PostLoad()
 
 		for (int32 Index = 0; Index < Platforms.Num(); Index++)
 		{
-			GetCompressedData(Platforms[Index]->GetWaveFormat(this), Platforms[Index]->GetAudioCompressionSettings());
+			BeginGetCompressedData(Platforms[Index]->GetWaveFormat(this), Platforms[Index]->GetAudioCompressionSettings());
 		}
 	}
 
@@ -581,6 +629,30 @@ void USoundWave::PostLoad()
 	INC_FLOAT_STAT_BY( STAT_AudioBufferTimeChannels, NumChannels * Duration );
 }
 
+void USoundWave::BeginDestroy()
+{
+	Super::BeginDestroy();
+
+	// Flag that this sound wave is beginning destroying. For procedural sound waves, this will ensure
+	// the audio render thread stops the sound before GC hits.
+	bIsBeginDestroy = true;
+
+#if WITH_EDITOR
+	// Flush any async results so we dont leak them in the DDC
+	if (GetDerivedDataCache() && AsyncLoadingDataFormats.Num() > 0)
+	{
+		TArray<uint8> OutData;
+		for (auto AsyncLoadIt = AsyncLoadingDataFormats.CreateConstIterator(); AsyncLoadIt; ++AsyncLoadIt)
+		{
+			uint32 AsyncHandle = AsyncLoadIt.Value();
+			GetDerivedDataCacheRef().WaitAsynchronousCompletion(AsyncHandle);
+			GetDerivedDataCacheRef().GetAsynchronousResults(AsyncHandle, OutData);
+		}
+
+		AsyncLoadingDataFormats.Empty();
+	}
+#endif
+}
 
 void USoundWave::InitAudioResource( FByteBulkData& CompressedData )
 {
@@ -712,7 +784,8 @@ void USoundWave::FreeResources()
 	ResourceID = 0;
 	bDynamicResource = false;
 	DecompressionType = DTYPE_Setup;
-	bDecompressedFromOgg = 0;
+	bDecompressedFromOgg = false;
+	bIsBeginDestroy = false;
 
 	USoundWave* SoundWave = this;
 	FAudioThread::RunCommandOnGameThread([SoundWave]()
@@ -801,7 +874,8 @@ bool USoundWave::IsReadyForFinishDestroy()
 		}, GET_STATID(STAT_AudioFreeResources));
 	}
 	
-	return ResourceState == ESoundWaveResourceState::Freed;
+	// bIsSoundActive is set in audio mixer when decoding sound waves or generating PCM data from procedural sound waves
+	return ResourceState == ESoundWaveResourceState::Freed && !bIsSoundActive;
 }
 
 
diff --git a/Engine/Source/Runtime/Engine/Private/SoundWaveProcedural.cpp b/Engine/Source/Runtime/Engine/Private/SoundWaveProcedural.cpp
--- a/Engine/Source/Runtime/Engine/Private/SoundWaveProcedural.cpp
+++ b/Engine/Source/Runtime/Engine/Private/SoundWaveProcedural.cpp
@@ -35,19 +35,9 @@ USoundWaveProcedural::USoundWaveProcedural(const FObjectInitializer& ObjectIniti
 
 	SampleByteSize = 2;
 
-	// This is set to true to default to old behavior in old audio engine
-	// Audio mixer uses sound wave procedural in async tasks and sets this to false when using it.
-	bIsReadyForDestroy = true;
-
 	checkf(NumSamplesToGeneratePerCallback >= NumBufferUnderrunSamples, TEXT("Should generate more samples than this per callback."));
 }
 
-USoundWaveProcedural::USoundWaveProcedural(FVTableHelper& Helper)
-	: Super(Helper)
-{
-	bIsReadyForDestroy = true;
-}
-
 void USoundWaveProcedural::QueueAudio(const uint8* AudioData, const int32 BufferSize)
 {
 	Audio::EAudioMixerStreamDataFormat::Type Format = GetGeneratedPCMDataFormat();
@@ -128,7 +118,7 @@ int32 USoundWaveProcedural::GeneratePCMData(uint8* PCMData, const int32 SamplesN
 	SamplesAvailable = AudioBuffer.Num() / SampleByteSize;
 
 	// Wait until we have enough samples that are requested before starting.
-	if (SamplesAvailable >= SamplesToGenerate)
+	if (SamplesAvailable > 0)
 	{
 		const int32 SamplesToCopy = FMath::Min<int32>(SamplesToGenerate, SamplesAvailable);
 		const int32 BytesToCopy = SamplesToCopy * SampleByteSize;
@@ -175,14 +165,14 @@ void USoundWaveProcedural::GetAssetRegistryTags(TArray<FAssetRegistryTag>& OutTa
 	Super::GetAssetRegistryTags(OutTags);
 }
 
-bool USoundWaveProcedural::IsReadyForFinishDestroy()
+bool USoundWaveProcedural::HasCompressedData(FName Format, ITargetPlatform* TargetPlatform) const
 {
-	return bIsReadyForDestroy;
+	return false;
 }
 
-bool USoundWaveProcedural::HasCompressedData(FName Format, ITargetPlatform* TargetPlatform) const
+void USoundWaveProcedural::BeginGetCompressedData(FName Format, const FPlatformAudioCookOverrides* CompressionOverrides)
 {
-	return false;
+	// SoundWaveProcedural does not have compressed data and should generally not be asked about it
 }
 
 FByteBulkData* USoundWaveProcedural::GetCompressedData(FName Format, const FPlatformAudioCookOverrides* CompressionOverrides)
diff --git a/Engine/Source/Runtime/Engine/Public/Audio.h b/Engine/Source/Runtime/Engine/Public/Audio.h
--- a/Engine/Source/Runtime/Engine/Public/Audio.h
+++ b/Engine/Source/Runtime/Engine/Public/Audio.h
@@ -373,6 +373,9 @@ public:
 	/** Whether or not this wave instance is ambisonics. */
 	uint32 bIsAmbisonics:1;
 
+	/** Whether or not this wave instance is stopping. */
+	uint32 bIsStopping:1;
+
 	/** Which spatialization method to use to spatialize 3d sounds. */
 	ESoundSpatializationAlgorithm SpatializationMethod;
 
@@ -472,6 +475,9 @@ public:
 	void SetVolumeApp(const float InVolumeApp) { VolumeApp = InVolumeApp; }
 	void SetVolumeMultiplier(const float InVolumeMultiplier) { VolumeMultiplier = InVolumeMultiplier; }
 
+	void SetStopping(const bool bInIsStopping) { bIsStopping = bInIsStopping; }
+	bool IsStopping() const { return bIsStopping; }
+
 	/** Returns the volume multiplier on the wave instance. */
 	float GetVolumeMultiplier() const { return VolumeMultiplier; }
 
@@ -607,6 +613,7 @@ public:
 		, Pitch(1.0f)
 		, LastUpdate(0)
 		, LastHeardUpdate(0)
+		, TickCount(0)
 		, LeftChannelSourceLocation(0)
 		, RightChannelSourceLocation(0)
 		, NumFramesPlayed(0)
@@ -653,12 +660,6 @@ public:
 	/** Whether or not the source is stopping. Only implemented in audio mixer. */
 	virtual bool IsStopping() { return false; }
 
-	/** Makes sure the source stops. Only implemented in audio mixer. */
-	virtual void EnsureStopped() {}
-
-	/** Called after a sound has been stopped and needs cleaning up. */
-	virtual void ReleaseStoppedSound() {}
-
 	/** Returns true if the sound source has finished playing. */
 	virtual	bool IsFinished() = 0;
 	
@@ -793,6 +794,9 @@ protected:
 	/** Last tick when this source was active *and* had a hearable volume */
 	int32 LastHeardUpdate;
 
+	/** Update tick count. Used to stop oldest stopping sound source. */
+	int32 TickCount;
+
 	/** The location of the left-channel source for stereo spatialization. */
 	FVector LeftChannelSourceLocation;
 
diff --git a/Engine/Source/Runtime/Engine/Public/AudioDevice.h b/Engine/Source/Runtime/Engine/Public/AudioDevice.h
--- a/Engine/Source/Runtime/Engine/Public/AudioDevice.h
+++ b/Engine/Source/Runtime/Engine/Public/AudioDevice.h
@@ -1139,6 +1139,12 @@ public:
 		return bAudioMixerModuleLoaded;
 	}
 
+	/** Returns if stopping voices is enabled. */
+	bool IsStoppingVoicesEnabled() const
+	{
+		return bIsStoppingVoicesEnabled;
+	}
+
 	/** Updates the source effect chain. Only implemented in audio mixer. */
 	virtual void UpdateSourceEffectChain(const uint32 SourceEffectChainId, const TArray<FSourceEffectChainEntry>& SourceEffectChain, const bool bPlayEffectChainTails) {}
 
@@ -1397,6 +1403,9 @@ private:
 	/** Processes the set of pending sounds that need to be stopped */ 
 	void ProcessingPendingActiveSoundStops(bool bForceDelete = false);
 
+	/** Stops oldest sound source. */
+	void StopOldestStoppingSource();
+
 	/** Processes any pending active sounds. */
 	void ProcessPendingNewActiveSounds();
 
@@ -1479,9 +1488,6 @@ public:
 	/** The number of sources to reserve for stopping sounds. */
 	int32 NumStoppingVoices;
 
-	/** The number of sources currently stopping. */
-	int32 CurrentStoppingVoiceCount;
-
 	/** The maximum number of wave instances allowed. */
 	int32 MaxWaveInstances;
 
@@ -1613,7 +1619,9 @@ public:
 	uint8 bDisableAudioCaching:1;
 
 	/** Whether or not the lower-level audio device hardware initialized. */
-	uint32 bIsAudioDeviceHardwareInitialized : 1;
+	uint8 bIsAudioDeviceHardwareInitialized : 1;
+
+	uint8 bIsStoppingVoicesEnabled : 1;
 
 	/** Whether or not the audio mixer module is being used by this device. */
 	uint8 bAudioMixerModuleLoaded : 1;
diff --git a/Engine/Source/Runtime/Engine/Public/IAudioExtensionPlugin.h b/Engine/Source/Runtime/Engine/Public/IAudioExtensionPlugin.h
--- a/Engine/Source/Runtime/Engine/Public/IAudioExtensionPlugin.h
+++ b/Engine/Source/Runtime/Engine/Public/IAudioExtensionPlugin.h
@@ -133,7 +133,7 @@ struct FAudioPluginSourceInputData
 	uint64 AudioComponentId;
 
 	// The audio input buffer
-	TArray<float>* AudioBuffer;
+	Audio::AlignedFloatBuffer* AudioBuffer;
 
 	// Number of channels of the source audio buffer.
 	int32 NumChannels;
@@ -148,7 +148,7 @@ struct FAudioPluginSourceInputData
 struct FAudioPluginSourceOutputData
 {
 	// The audio output buffer
-	TArray<float> AudioBuffer;
+	Audio::AlignedFloatBuffer AudioBuffer;
 };
 
 /** This is a class which should be overridden to provide users with settings to use for individual sounds */
diff --git a/Engine/Source/Runtime/Windows/AudioMixerXAudio2/Private/AudioMixerPlatformWindows.cpp b/Engine/Source/Runtime/Windows/AudioMixerXAudio2/Private/AudioMixerPlatformWindows.cpp
--- a/Engine/Source/Runtime/Windows/AudioMixerXAudio2/Private/AudioMixerPlatformWindows.cpp
+++ b/Engine/Source/Runtime/Windows/AudioMixerXAudio2/Private/AudioMixerPlatformWindows.cpp
@@ -204,36 +204,43 @@ namespace Audio
 
 	void FMixerPlatformXAudio2::OnDefaultRenderDeviceChanged(const EAudioDeviceRole InAudioDeviceRole, const FString& DeviceId)
 	{
-		FScopeLock Lock(&AudioDeviceSwapCriticalSection);
+		if (AudioDeviceSwapCriticalSection.TryLock())
+		{
+			NewAudioDeviceId = "";
+			bMoveAudioStreamToNewAudioDevice = true;
 
-		NewAudioDeviceId = "";
-		bMoveAudioStreamToNewAudioDevice = true;
+			AudioDeviceSwapCriticalSection.Unlock();
+		}
 	}
 
 	void FMixerPlatformXAudio2::OnDeviceAdded(const FString& DeviceId)
 	{
-		FScopeLock Lock(&AudioDeviceSwapCriticalSection);
-
-		// If the device that was added is our original device and our current device is NOT our original device, 
-		// move our audio stream to this newly added device.
-		if (AudioStreamInfo.DeviceInfo.DeviceId != OriginalAudioDeviceId && DeviceId == OriginalAudioDeviceId)
+		if (AudioDeviceSwapCriticalSection.TryLock())
 		{
-			NewAudioDeviceId = OriginalAudioDeviceId;
-			bMoveAudioStreamToNewAudioDevice = true;
+			// If the device that was added is our original device and our current device is NOT our original device, 
+			// move our audio stream to this newly added device.
+			if (AudioStreamInfo.DeviceInfo.DeviceId != OriginalAudioDeviceId && DeviceId == OriginalAudioDeviceId)
+			{
+				NewAudioDeviceId = OriginalAudioDeviceId;
+				bMoveAudioStreamToNewAudioDevice = true;
+			}
+
+			AudioDeviceSwapCriticalSection.Unlock();
 		}
 	}
 
 	void FMixerPlatformXAudio2::OnDeviceRemoved(const FString& DeviceId)
 	{
-		FScopeLock Lock(&AudioDeviceSwapCriticalSection);
-
-		// If the device we're currently using was removed... then switch to the new default audio device.
-		if (AudioStreamInfo.DeviceInfo.DeviceId == DeviceId)
+		if (AudioDeviceSwapCriticalSection.TryLock())
 		{
-			NewAudioDeviceId = "";
-			bMoveAudioStreamToNewAudioDevice = true;
+			// If the device we're currently using was removed... then switch to the new default audio device.
+			if (AudioStreamInfo.DeviceInfo.DeviceId == DeviceId)
+			{
+				NewAudioDeviceId = "";
+				bMoveAudioStreamToNewAudioDevice = true;
+			}
+			AudioDeviceSwapCriticalSection.Unlock();
 		}
-
 	}
 
 	void FMixerPlatformXAudio2::OnDeviceStateChanged(const FString& DeviceId, const EAudioDeviceState InState)
diff --git a/Engine/Source/Runtime/Windows/AudioMixerXAudio2/Private/AudioMixerPlatformXAudio2.cpp b/Engine/Source/Runtime/Windows/AudioMixerXAudio2/Private/AudioMixerPlatformXAudio2.cpp
--- a/Engine/Source/Runtime/Windows/AudioMixerXAudio2/Private/AudioMixerPlatformXAudio2.cpp
+++ b/Engine/Source/Runtime/Windows/AudioMixerXAudio2/Private/AudioMixerPlatformXAudio2.cpp
@@ -511,9 +511,6 @@ namespace Audio
 			OutputAudioStreamSourceVoice = nullptr;
 		}
 
-		// Don't let the audio stream process while switching to new audio device!
-		//FScopeLock Lock(&AudioRenderCritSect); XXX - Audio should be stopped then the device switched then audio started again
-
 		// Now destroy the mastering voice
 		if (OutputAudioStreamMasteringVoice)
 		{
