[Release 4.20] SA health: Fix too-late nullptr check. #jira none #rb none #fyi aaron.mcleran

diff --git a/Engine/Source/Runtime/AudioMixer/Private/AudioMixerSource.cpp b/Engine/Source/Runtime/AudioMixer/Private/AudioMixerSource.cpp
--- a/Engine/Source/Runtime/AudioMixer/Private/AudioMixerSource.cpp
+++ b/Engine/Source/Runtime/AudioMixer/Private/AudioMixerSource.cpp
@@ -246,6 +246,9 @@ namespace Audio
 			}
 #endif
 
+			AUDIO_MIXER_CHECK(WaveInstance);
+			AUDIO_MIXER_CHECK(InWaveInstance);
+
 			// Whether or not we're 3D
 			bIs3D = !UseObjectBasedSpatialization() && WaveInstance->bUseSpatialization && SoundBuffer->NumChannels < 3;
 
@@ -260,6 +263,7 @@ namespace Audio
 
 			// We support reverb
 			SetReverbApplied(true);
+			AUDIO_MIXER_CHECK(InWaveInstance->WaveData);
 
 			// Update the buffer sample rate to the wave instance sample rate in case it was serialized incorrectly
 			MixerBuffer->InitSampleRate(InWaveInstance->WaveData->GetSampleRateForCurrentPlatform());
@@ -273,8 +277,6 @@ namespace Audio
 
 			if (MixerSourceVoice->Init(InitParams))
 			{
-				AUDIO_MIXER_CHECK(WaveInstance);
-
 				bInitialized = true;
 
 				Update();
