[Release 4.20] #jira UE-61196 SA health: Check for null pointers earlier in FMixerSource::Init. #rb none.

diff --git a/Engine/Source/Runtime/AudioMixer/Private/AudioMixerSource.cpp b/Engine/Source/Runtime/AudioMixer/Private/AudioMixerSource.cpp
--- a/Engine/Source/Runtime/AudioMixer/Private/AudioMixerSource.cpp
+++ b/Engine/Source/Runtime/AudioMixer/Private/AudioMixerSource.cpp
@@ -55,6 +55,9 @@ namespace Audio
 		// Get the number of frames before creating the buffer
 		int32 NumFrames = INDEX_NONE;
 
+		AUDIO_MIXER_CHECK(InWaveInstance);
+		AUDIO_MIXER_CHECK(InWaveInstance->WaveData);
+
 		if (InWaveInstance->WaveData->DecompressionType != DTYPE_Procedural)
 		{
 			const int32 NumBytes = InWaveInstance->WaveData->RawPCMDataSize;
@@ -232,6 +235,8 @@ namespace Audio
 				ChannelMaps[(int32)SubmixChannelType].ChannelMap.Reset();
 			}
 
+			AUDIO_MIXER_CHECK(WaveInstance);
+
 			// Check to see if this sound has been flagged to be in debug mode
 #if AUDIO_MIXER_ENABLE_DEBUG_MODE
 			InitParams.DebugName = InWaveInstance->GetName();
@@ -246,9 +251,6 @@ namespace Audio
 			}
 #endif
 
-			AUDIO_MIXER_CHECK(WaveInstance);
-			AUDIO_MIXER_CHECK(InWaveInstance);
-
 			// Whether or not we're 3D
 			bIs3D = !UseObjectBasedSpatialization() && WaveInstance->bUseSpatialization && SoundBuffer->NumChannels < 3;
 
@@ -263,7 +265,6 @@ namespace Audio
 
 			// We support reverb
 			SetReverbApplied(true);
-			AUDIO_MIXER_CHECK(InWaveInstance->WaveData);
 
 			// Update the buffer sample rate to the wave instance sample rate in case it was serialized incorrectly
 			MixerBuffer->InitSampleRate(InWaveInstance->WaveData->GetSampleRateForCurrentPlatform());
