Fix for memory leak in audio mixer.

Tested with -llm -llmcsv with replay A/B.

diff --git a/Engine/Source/Runtime/AudioMixer/Private/AudioMixerSourceBuffer.cpp b/Engine/Source/Runtime/AudioMixer/Private/AudioMixerSourceBuffer.cpp
--- a/Engine/Source/Runtime/AudioMixer/Private/AudioMixerSourceBuffer.cpp
+++ b/Engine/Source/Runtime/AudioMixer/Private/AudioMixerSourceBuffer.cpp
@@ -414,6 +414,14 @@ namespace Audio
 
 			SoundWave = nullptr;
 		}
+
+		EBufferType::Type BufferType = MixerBuffer->GetType();
+		if (BufferType == EBufferType::PCMRealTime || BufferType == EBufferType::Streaming)
+		{
+			delete MixerBuffer;
+		}
+
+		MixerBuffer = nullptr;
 	}
 
 }
