Decorating audio mixer code with new audio mixer LLM markup.

diff --git a/Engine/Source/Runtime/AudioMixer/Private/AudioMixer.cpp b/Engine/Source/Runtime/AudioMixer/Private/AudioMixer.cpp
--- a/Engine/Source/Runtime/AudioMixer/Private/AudioMixer.cpp
+++ b/Engine/Source/Runtime/AudioMixer/Private/AudioMixer.cpp
@@ -300,6 +300,8 @@ namespace Audio
 
 	void IAudioMixerPlatformInterface::ReadNextBuffer()
 	{
+		LLM_SCOPE(ELLMTag::AudioMixer);
+
 		// Don't read any more audio if we're not running or changing device
 		if (AudioStreamInfo.StreamState != EAudioOutputStreamState::Running || bAudioDeviceChanging)
 		{
@@ -434,6 +436,8 @@ namespace Audio
 
 	void IAudioMixerPlatformInterface::Tick()
 	{
+		LLM_SCOPE(ELLMTag::AudioMixer);
+
 		// In single-threaded mode, we simply render buffers until we run out of space
 		// The single-thread audio backend will consume these rendered buffers when they need to
 		if (AudioStreamInfo.StreamState == EAudioOutputStreamState::Running && bIsDeviceInitialized)
@@ -493,6 +497,8 @@ namespace Audio
 
 	uint32 IAudioMixerPlatformInterface::Run()
 	{	
+		LLM_SCOPE(ELLMTag::AudioMixer);
+
 		// Call different functions depending on if it's the "main" audio mixer instance. Helps debugging callstacks.
 		if (AudioStreamInfo.AudioMixer->IsMainAudioMixer())
 		{
diff --git a/Engine/Source/Runtime/AudioMixer/Private/AudioMixerDevice.cpp b/Engine/Source/Runtime/AudioMixer/Private/AudioMixerDevice.cpp
--- a/Engine/Source/Runtime/AudioMixer/Private/AudioMixerDevice.cpp
+++ b/Engine/Source/Runtime/AudioMixer/Private/AudioMixerDevice.cpp
@@ -292,6 +292,8 @@ namespace Audio
 
 	void FMixerDevice::UpdateHardware()
 	{
+		LLM_SCOPE(ELLMTag::AudioMixer);
+
 		// If we're in editor, re-query these in case they changed. 
 		if (GIsEditor)
 		{
@@ -919,6 +921,8 @@ namespace Audio
 
 	FMixerSubmixPtr FMixerDevice::GetSubmixInstance(USoundSubmix* SoundSubmix)
 	{
+		LLM_SCOPE(ELLMTag::AudioMixer);
+
 		check(SoundSubmix);
 		FMixerSubmixPtr* MixerSubmix = Submixes.Find(SoundSubmix);
 
@@ -936,6 +940,8 @@ namespace Audio
 
 	FMixerSourceVoice* FMixerDevice::GetMixerSourceVoice()
 	{
+		LLM_SCOPE(ELLMTag::AudioMixer);
+
 		FMixerSourceVoice* Voice = nullptr;
 		if (!SourceVoices.Dequeue(Voice))
 		{
diff --git a/Engine/Source/Runtime/AudioMixer/Private/AudioMixerSource.cpp b/Engine/Source/Runtime/AudioMixer/Private/AudioMixerSource.cpp
--- a/Engine/Source/Runtime/AudioMixer/Private/AudioMixerSource.cpp
+++ b/Engine/Source/Runtime/AudioMixer/Private/AudioMixerSource.cpp
@@ -50,6 +50,8 @@ namespace Audio
 		AUDIO_MIXER_CHECK(MixerBuffer);
 		AUDIO_MIXER_CHECK(MixerBuffer->IsRealTimeSourceReady());
 
+		LLM_SCOPE(ELLMTag::AudioMixer);
+
 		FSoundSource::InitCommon();
 
 		// Get the number of frames before creating the buffer
@@ -292,6 +294,8 @@ namespace Audio
 	{
 		SCOPE_CYCLE_COUNTER(STAT_AudioUpdateSources);
 
+		LLM_SCOPE(ELLMTag::AudioMixer);
+
 		if (!WaveInstance || !MixerSourceVoice || Paused || !bInitialized)
 		{
 			return;
@@ -314,6 +318,8 @@ namespace Audio
 
 	bool FMixerSource::PrepareForInitialization(FWaveInstance* InWaveInstance)
 	{
+		LLM_SCOPE(ELLMTag::AudioMixer);
+
 		// We are currently not supporting playing audio on a controller
 		if (InWaveInstance->OutputTarget == EAudioOutputTarget::Controller)
 		{
@@ -380,6 +386,8 @@ namespace Audio
 
 	bool FMixerSource::IsPreparedToInit()
 	{
+		LLM_SCOPE(ELLMTag::AudioMixer);
+			
 		if (MixerBuffer && MixerBuffer->IsRealTimeSourceReady() && !bIsReleasing)
 		{
 			check(MixerSourceBuffer.IsValid());
@@ -451,6 +459,8 @@ namespace Audio
 
 	void FMixerSource::Stop()
 	{
+		LLM_SCOPE(ELLMTag::AudioMixer);
+	
 		if (bIsFinished)
 		{
 			StopNow();
@@ -486,6 +496,8 @@ namespace Audio
 
 	void FMixerSource::StopNow()
 	{
+		LLM_SCOPE(ELLMTag::AudioMixer);
+
 		// Immediately stop the sound source
 
 		bInitialized = false;
@@ -620,6 +632,8 @@ namespace Audio
 
 	void FMixerSource::FreeResources()
 	{
+		LLM_SCOPE(ELLMTag::AudioMixer);
+
 		if (MixerBuffer)
 		{
 			MixerBuffer->EnsureHeaderParseTaskFinished();
diff --git a/Engine/Source/Runtime/AudioMixer/Private/AudioMixerSourceBuffer.cpp b/Engine/Source/Runtime/AudioMixer/Private/AudioMixerSourceBuffer.cpp
--- a/Engine/Source/Runtime/AudioMixer/Private/AudioMixerSourceBuffer.cpp
+++ b/Engine/Source/Runtime/AudioMixer/Private/AudioMixerSourceBuffer.cpp
@@ -73,6 +73,9 @@ namespace Audio
 
 	bool FMixerSourceBuffer::PreInit(FMixerBuffer* InBuffer, USoundWave* InWave, ELoopingMode InLoopingMode, bool bInIsSeeking)
 	{
+		LLM_SCOPE(ELLMTag::AudioMixer);
+
+		// Mixer source buffer now owns this buffer
 		MixerBuffer = InBuffer;
 		check(MixerBuffer);
 
