- Issue was due to a discontinuity introduced in the source manager when fade time is less than the render buffer length due to math error.
- Other issue was encountered while fixing where sounds would not always play when using stopping voices. That issue is also fixed.

diff --git a/Engine/Source/Runtime/AudioMixer/Private/AudioMixerSource.cpp b/Engine/Source/Runtime/AudioMixer/Private/AudioMixerSource.cpp
--- a/Engine/Source/Runtime/AudioMixer/Private/AudioMixerSource.cpp
+++ b/Engine/Source/Runtime/AudioMixer/Private/AudioMixerSource.cpp
@@ -28,9 +28,9 @@ namespace Audio
 		, bPlayedCachedBuffer(false)
 		, bPlaying(false)
 		, bLoopCallback(false)
-		, bIsFinished(false)
+		, bIsDone(false)
+		, bIsEffectTailsDone(false)
 		, bIsPlayingEffectTails(false)
-		, bResourcesNeedFreeing(false)
 		, bEditorWarnedChangedSpatialization(false)
 		, bUsingHRTFSpatialization(false)
 		, bIs3D(false)
@@ -69,12 +69,11 @@ namespace Audio
 		// Unfortunately, we need to know if this is a vorbis source since channel maps are different for 5.1 vorbis files
 		bIsVorbis = InWaveInstance->WaveData->bDecompressedFromOgg;
 
-		// Reset our releasing bool
-		bIsReleasing = false;
-
 		bIsStoppingVoicesEnabled = ((FAudioDevice*)MixerDevice)->IsStoppingVoicesEnabled();
-
+		
 		bIsStopping = false;
+		bIsEffectTailsDone = true;
+		bIsDone = false;
 
 		FSoundBuffer* SoundBuffer = static_cast<FSoundBuffer*>(MixerBuffer);
 		if (SoundBuffer->NumChannels > 0)
@@ -125,6 +124,12 @@ namespace Audio
 					}
 				}
 
+				// Only need to care about effect chain tails finishing if we're told to play them
+				if (InitParams.bPlayEffectChainTails)
+				{
+					bIsEffectTailsDone = false;
+				}
+
 				// Setup the bus Id if this source is a bus
 				if (InWaveInstance->WaveData->bIsBus)
 				{
@@ -349,10 +354,7 @@ namespace Audio
 			HPFFrequency = 0.0;
 			LastHPFFrequency = FLT_MAX;
 
-			bIsFinished = false;
-
-			EBufferType::Type BufferType = MixerBuffer->GetType();
-			bResourcesNeedFreeing = (BufferType == EBufferType::PCMRealTime || BufferType == EBufferType::Streaming);
+			bIsDone = false;
 
 			// Not all wave data types have a non-zero duration
 			if (InWaveInstance->WaveData->Duration > 0)
@@ -387,8 +389,8 @@ namespace Audio
 	bool FMixerSource::IsPreparedToInit()
 	{
 		LLM_SCOPE(ELLMTag::AudioMixer);
-			
-		if (MixerBuffer && MixerBuffer->IsRealTimeSourceReady() && !bIsReleasing)
+
+		if (MixerBuffer && MixerBuffer->IsRealTimeSourceReady())
 		{
 			check(MixerSourceBuffer.IsValid());
 
@@ -454,14 +456,20 @@ namespace Audio
 		Paused = false;
 		Playing = true;
 		bLoopCallback = false;
-		bIsFinished = false;
+		bIsDone = false;
 	}
 
 	void FMixerSource::Stop()
 	{
 		LLM_SCOPE(ELLMTag::AudioMixer);
-	
-		if (bIsFinished)
+
+		if (!MixerSourceVoice)
+		{
+			StopNow();
+			return;
+		}
+
+		if (bIsDone)
 		{
 			StopNow();
 		}
@@ -551,26 +559,13 @@ namespace Audio
 
 		if (WaveInstance && MixerSourceVoice)
 		{
-			if (bIsFinished && MixerSourceVoice->IsSourceEffectTailsDone())
+			if (bIsDone && bIsEffectTailsDone)
 			{
 				WaveInstance->NotifyFinished();
 				bIsStopping = false;
 				return true;
 			}
-			// Buses don't do buffer end callbacks, so we need to directly query doneness. 
-			// Also need to query if we're currently stopping
-			else if (WaveInstance->WaveData->bIsBus || bIsStopping)
-			{
-				if (MixerSourceVoice->IsSourceEffectTailsDone() && MixerSourceVoice->IsDone())
-				{
-					bIsFinished = true;
-					WaveInstance->NotifyFinished();
-					bIsStopping = false;
-					return true;
-				}
-			}
-
-			if (bLoopCallback && WaveInstance->LoopingMode == LOOP_WithNotification)
+			else if (bLoopCallback && WaveInstance->LoopingMode == LOOP_WithNotification)
 			{
 				WaveInstance->NotifyFinished();
 				bLoopCallback = false;
@@ -620,14 +615,15 @@ namespace Audio
 	{
 	}
 
-	void FMixerSource::OnRelease()
+	void FMixerSource::OnDone()
 	{
-		bIsReleasing = false;
+		bIsDone = true;
 	}
 
-	void FMixerSource::OnDone()
+
+	void FMixerSource::OnEffectTailsDone()
 	{
-		bIsFinished = true;
+		bIsEffectTailsDone = true;
 	}
 
 	void FMixerSource::FreeResources()
@@ -657,7 +653,6 @@ namespace Audio
 		MixerBuffer = nullptr;
 		Buffer = nullptr;
 		bLoopCallback = false;
-		bResourcesNeedFreeing = false;
 		NumTotalFrames = 0;
 
 		// Reset the source's channel maps
diff --git a/Engine/Source/Runtime/AudioMixer/Private/AudioMixerSource.h b/Engine/Source/Runtime/AudioMixer/Private/AudioMixerSource.h
--- a/Engine/Source/Runtime/AudioMixer/Private/AudioMixerSource.h
+++ b/Engine/Source/Runtime/AudioMixer/Private/AudioMixerSource.h
@@ -44,12 +44,12 @@ namespace Audio
 		virtual float GetEnvelopeValue() const override;
 		//~ End FSoundSource Interface
 
-		//~ Begin ISourceListneer
+		//~ Begin ISourceListener
 		virtual void OnBeginGenerate() override;
-		virtual void OnRelease() override;
 		virtual void OnDone() override;
+		virtual void OnEffectTailsDone() override;
 		virtual void OnLoopEnd() override { bLoopCallback = true; };
-		//~ End ISourceListneer
+		//~ End ISourceListener
 
 	private:
 
@@ -120,14 +120,14 @@ namespace Audio
 		FThreadSafeBool bPlaying;
 		FThreadSafeBool bIsStopping;
 		FThreadSafeBool bLoopCallback;
-		FThreadSafeBool bIsFinished;
+		FThreadSafeBool bIsDone;
+		FThreadSafeBool bIsEffectTailsDone;
 		FThreadSafeBool bIsPlayingEffectTails;
 		FThreadSafeBool bFreeAsyncTask;
 
 		// Whether or not we're currently releasing our resources. Prevents recycling the source until release is finished.
 		FThreadSafeBool bIsReleasing;
 
-		uint32 bResourcesNeedFreeing : 1;
 		uint32 bEditorWarnedChangedSpatialization : 1;
 		uint32 bUsingHRTFSpatialization : 1;
 		uint32 bIs3D : 1;
diff --git a/Engine/Source/Runtime/AudioMixer/Private/AudioMixerSourceManager.cpp b/Engine/Source/Runtime/AudioMixer/Private/AudioMixerSourceManager.cpp
--- a/Engine/Source/Runtime/AudioMixer/Private/AudioMixerSourceManager.cpp
+++ b/Engine/Source/Runtime/AudioMixer/Private/AudioMixerSourceManager.cpp
@@ -205,8 +205,6 @@ namespace Audio
 		}
 		
 		GameThreadInfo.bIsBusy.AddDefaulted(NumTotalSources);
-		GameThreadInfo.bIsDone.AddDefaulted(NumTotalSources);
-		GameThreadInfo.bEffectTailsDone.AddDefaulted(NumTotalSources);
 		GameThreadInfo.bNeedsSpeakerMap.AddDefaulted(NumTotalSources);
 		GameThreadInfo.bIsDebugMode.AddDefaulted(NumTotalSources);
 		GameThreadInfo.FreeSourceIndices.Reset(NumTotalSources);
@@ -346,13 +344,7 @@ namespace Audio
 			SourceInfo.MixerSourceBuffer = nullptr;
 		}
 
-		// Call OnRelease on the BufferQueueListener to give it a chance 
-		// to release any resources it owns on the audio render thread
-		if (SourceInfo.SourceListener)
-		{
-			SourceInfo.SourceListener->OnRelease();
-			SourceInfo.SourceListener = nullptr;
-		}
+		SourceInfo.SourceListener = nullptr;
 
 		// Remove the mixer source from its submix sends
 		for (FMixerSourceSubmixSend& SubmixSendItem : SourceInfo.SubmixSends)
@@ -1093,18 +1085,6 @@ namespace Audio
 		return SourceInfos[SourceId].SourceEnvelopeValue;
 	}
 
-	bool FMixerSourceManager::IsDone(const int32 SourceId) const
-	{
-		AUDIO_MIXER_CHECK_GAME_THREAD(MixerDevice);
-		return GameThreadInfo.bIsDone[SourceId];
-	}
-
-	bool FMixerSourceManager::IsEffectTailsDone(const int32 SourceId) const
-	{
-		AUDIO_MIXER_CHECK_GAME_THREAD(MixerDevice);
-		return GameThreadInfo.bEffectTailsDone[SourceId];
-	}
-
 	bool FMixerSourceManager::NeedsSpeakerMap(const int32 SourceId) const
 	{
 		AUDIO_MIXER_CHECK_GAME_THREAD(MixerDevice);
@@ -1589,7 +1569,7 @@ namespace Audio
 				if (NumFadeFrames < NumOutputFrames)
 				{
 					int32 SamplesLeft = NumSamples - NumFadeSamples;
-					FMemory::Memzero(&PreDistanceAttenBufferPtr[NumFadeFrames], sizeof(float)*SamplesLeft);
+					FMemory::Memzero(&PreDistanceAttenBufferPtr[NumFadeSamples], sizeof(float)*SamplesLeft);
 				}
 
 				SourceInfo.VolumeSourceStart = SourceInfo.VolumeSourceDestination;
@@ -1646,6 +1626,9 @@ namespace Audio
 					}
 				}
 			}
+
+			const bool bWasEffectTailsDone = SourceInfo.bEffectTailsDone;
+
 			if (!DisableEnvelopeFollowingCvar)
 			{
 				// Compute the source envelope using pre-distance attenuation buffer
@@ -1662,6 +1645,11 @@ namespace Audio
 				SourceInfo.bEffectTailsDone = true;
 			}
 
+			if (!bWasEffectTailsDone && SourceInfo.bEffectTailsDone)
+			{
+				SourceInfo.SourceListener->OnEffectTailsDone();
+			}
+
 			// Only scale with distance attenuation and send to source audio to plugins if we're not in output-to-bus only mode
 			if (!SourceInfo.bOutputToBusOnly && !DisableFilteringCvar)
 			{
@@ -2100,14 +2088,14 @@ namespace Audio
 			// Check for the stopping condition to "turn the sound off"
 			if (SourceInfo.bIsLastBuffer)
 			{
-				SourceInfo.bIsDone = true;
+				if (!SourceInfo.bIsDone)
+				{
+					SourceInfo.bIsDone = true;
 
-				// Notify that we're now done with this source
-				SourceInfo.SourceListener->OnDone();
+					// Notify that we're now done with this source
+					SourceInfo.SourceListener->OnDone();
+				}
 			}
-
-			GameThreadInfo.bIsDone[SourceId] = SourceInfo.bIsDone;
-			GameThreadInfo.bEffectTailsDone[SourceId] = SourceInfo.bEffectTailsDone;
 		}
 	}
 
@@ -2117,21 +2105,13 @@ namespace Audio
 		{
 			FSourceInfo& SourceInfo = SourceInfos[SourceId];
 
-			// Check for the stopping condition to "turn the sound off"
-			if (!SourceInfo.bIsDone && SourceInfo.bIsStopping)
-			{
-				if (SourceInfo.VolumeSourceDestination == 0.0f)
-				{
-					SourceInfo.bIsStopping = false;
-					SourceInfo.bIsDone = true;
-					GameThreadInfo.bIsDone[SourceId] = true;
-				}
-			}
-			else if (SourceInfo.bIsDone)
+			if (!SourceInfo.bIsDone && SourceInfo.bIsStopping && SourceInfo.VolumeSourceDestination == 0.0f)
 			{
 				SourceInfo.bIsStopping = false;
-				GameThreadInfo.bIsDone[SourceId] = true;
+				SourceInfo.bIsDone = true;
+				SourceInfo.SourceListener->OnDone();
 			}
+
 		}
 	}
 
@@ -2173,12 +2153,15 @@ namespace Audio
 				{
 					SourceInfo.MixerSourceBuffer->ClearSoundWave();
 						
+					if (!SourceInfo.bIsDone)
+					{
+						SourceInfo.bIsDone = true;
+						SourceInfo.SourceListener->OnDone();
+					}
+
 					// Clear out the mixer source buffer
 					SourceInfo.MixerSourceBuffer = nullptr;
 
-					// Signal that this source is done
-					SourceInfo.SourceListener->OnDone();
-
 					// Set the sound to be done playing
 					// This will flag the sound to be released
 					SourceInfo.bIsPlaying = false;
diff --git a/Engine/Source/Runtime/AudioMixer/Private/AudioMixerSourceManager.h b/Engine/Source/Runtime/AudioMixer/Private/AudioMixerSourceManager.h
--- a/Engine/Source/Runtime/AudioMixer/Private/AudioMixerSourceManager.h
+++ b/Engine/Source/Runtime/AudioMixer/Private/AudioMixerSourceManager.h
@@ -59,12 +59,12 @@ namespace Audio
 		// Called when a loop point is hit
 		virtual void OnLoopEnd() = 0;
 
-		// Called the source is released on the render thread
-		virtual void OnRelease() = 0;
-
-		// Called when the source finishes on the render thread
+		// Called when the source finishes on the audio render thread
 		virtual void OnDone() = 0;
 
+		// Called when the source's effect tails finish on the audio render thread.
+		virtual void OnEffectTailsDone() = 0;
+
 	};
 
 	struct FMixerSourceSubmixSend
@@ -249,8 +249,6 @@ namespace Audio
 
 		int64 GetNumFramesPlayed(const int32 SourceId) const;
 		float GetEnvelopeValue(const int32 SourceId) const;
-		bool IsDone(const int32 SourceId) const;
-		bool IsEffectTailsDone(const int32 SourceId) const;
 		bool NeedsSpeakerMap(const int32 SourceId) const;
 		void ComputeNextBlockOfSamples();
 		void ClearStoppingSounds();
@@ -495,8 +493,6 @@ namespace Audio
 		{
 			TArray<int32> FreeSourceIndices;
 			TArray<bool> bIsBusy;
-			TArray<FThreadSafeBool> bIsDone;
-			TArray<FThreadSafeBool> bEffectTailsDone;
 			TArray <bool> bNeedsSpeakerMap;
 			TArray<bool> bIsDebugMode;
 		} GameThreadInfo;
diff --git a/Engine/Source/Runtime/AudioMixer/Private/AudioMixerSourceVoice.cpp b/Engine/Source/Runtime/AudioMixer/Private/AudioMixerSourceVoice.cpp
--- a/Engine/Source/Runtime/AudioMixer/Private/AudioMixerSourceVoice.cpp
+++ b/Engine/Source/Runtime/AudioMixer/Private/AudioMixerSourceVoice.cpp
@@ -212,20 +212,6 @@ namespace Audio
 		return bIsActive;
 	}
 
-	bool FMixerSourceVoice::IsDone() const
-	{
-		AUDIO_MIXER_CHECK_GAME_THREAD(MixerDevice);
-
-		return SourceManager->IsDone(SourceId);
-	}
-
-	bool FMixerSourceVoice::IsSourceEffectTailsDone() const
-	{
-		AUDIO_MIXER_CHECK_GAME_THREAD(MixerDevice);
-
-		return SourceManager->IsEffectTailsDone(SourceId);
-	}
-
 	bool FMixerSourceVoice::NeedsSpeakerMap() const
 	{
 		AUDIO_MIXER_CHECK_GAME_THREAD(MixerDevice);
diff --git a/Engine/Source/Runtime/AudioMixer/Private/AudioMixerSourceVoice.h b/Engine/Source/Runtime/AudioMixer/Private/AudioMixerSourceVoice.h
--- a/Engine/Source/Runtime/AudioMixer/Private/AudioMixerSourceVoice.h
+++ b/Engine/Source/Runtime/AudioMixer/Private/AudioMixerSourceVoice.h
@@ -74,15 +74,9 @@ namespace Audio
 		// Queries if the source voice is active.
 		bool IsActive() const;
 
-		// Queries if the source voice has finished playing all its audio.
-		bool IsDone() const;
-
 		// Queries if the source has finished its fade out.
 		bool IsStopFadedOut() const { return bStopFadedOut; }
 
-		// Queries if the source ffect tails have finished
-		bool IsSourceEffectTailsDone() const;
-
 		// Whether or not the device changed and needs another speaker map sent
 		bool NeedsSpeakerMap() const;
 
diff --git a/Engine/Source/Runtime/Engine/Private/AudioDevice.cpp b/Engine/Source/Runtime/Engine/Private/AudioDevice.cpp
--- a/Engine/Source/Runtime/Engine/Private/AudioDevice.cpp
+++ b/Engine/Source/Runtime/Engine/Private/AudioDevice.cpp
@@ -3041,7 +3041,7 @@ int32 FAudioDevice::GetSortedActiveWaveInstances(TArray<FWaveInstance*>& WaveIns
 	int32 FirstActiveIndex = 0;
 
 	// Only need to do the wave instance sort if we have any waves and if our wave instances are greater than our max channels.
-	if (WaveInstances.Num() >= GetMaxChannels())
+	if (WaveInstances.Num() >= 0)
 	{
 		// Helper function for "Sort" (higher priority sorts last).
 		struct FCompareFWaveInstanceByPlayPriority
@@ -3086,14 +3086,14 @@ void FAudioDevice::UpdateActiveSoundPlaybackTime(bool bIsGameTicking)
 	}
 
 }
-
+ 
 void FAudioDevice::StopOldestStoppingSource()
 {
 	if (IsAudioMixerEnabled())
 	{
 		check(!FreeSources.Num());
 
-		FSoundSource* OldestStoppingSource = nullptr;
+		FSoundSource* LowestPriStoppingSource = nullptr;
 		FSoundSource* LowestPriSource = nullptr;
 		FSoundSource* LowestPriNonLoopingSource = nullptr;
 
@@ -3102,15 +3102,15 @@ void FAudioDevice::StopOldestStoppingSource()
 			// Find oldest stopping voice first
 			if (Source->IsStopping())
 			{
-				if (!OldestStoppingSource)
+				if (!LowestPriStoppingSource)
 				{
-					OldestStoppingSource = Source;
+					LowestPriStoppingSource = Source;
 				}
 				else
 				{
-					if (Source->TickCount > OldestStoppingSource->TickCount)
+					if (Source->WaveInstance->GetVolumeWeightedPriority() < LowestPriStoppingSource->WaveInstance->GetVolumeWeightedPriority())
 					{
-						OldestStoppingSource = Source;
+						LowestPriStoppingSource = Source;
 					}
 				}
 			}
@@ -3149,9 +3149,9 @@ void FAudioDevice::StopOldestStoppingSource()
 		}
 
 		// Stop oldest stopping source
-		if (OldestStoppingSource)
+		if (LowestPriStoppingSource)
 		{
-			OldestStoppingSource->StopNow();
+			LowestPriStoppingSource->StopNow();
 		}
 		// If no oldest stopping source, stop oldest one-shot
 		else if (LowestPriNonLoopingSource)
