We've already passed InWaveInstance to member var WaveInstance. Changed code to reflect this, added check at top of function. Should be non-null for whole function.

diff --git a/Engine/Source/Runtime/AudioMixer/Private/AudioMixerSource.cpp b/Engine/Source/Runtime/AudioMixer/Private/AudioMixerSource.cpp
--- a/Engine/Source/Runtime/AudioMixer/Private/AudioMixerSource.cpp
+++ b/Engine/Source/Runtime/AudioMixer/Private/AudioMixerSource.cpp
@@ -50,6 +50,9 @@ namespace Audio
 		AUDIO_MIXER_CHECK(MixerBuffer);
 		AUDIO_MIXER_CHECK(MixerBuffer->IsRealTimeSourceReady());
 
+		// We've already been passed the wave instance in PrepareForInitialization, make sure we have the same one
+		AUDIO_MIXER_CHECK(WaveInstance && WaveInstance == InWaveInstance);
+
 		LLM_SCOPE(ELLMTag::AudioMixer);
 
 		FSoundSource::InitCommon();
@@ -57,17 +60,16 @@ namespace Audio
 		// Get the number of frames before creating the buffer
 		int32 NumFrames = INDEX_NONE;
 
-		AUDIO_MIXER_CHECK(InWaveInstance);
-		AUDIO_MIXER_CHECK(InWaveInstance->WaveData);
+		AUDIO_MIXER_CHECK(WaveInstance->WaveData);
 
-		if (InWaveInstance->WaveData->DecompressionType != DTYPE_Procedural)
+		if (WaveInstance->WaveData->DecompressionType != DTYPE_Procedural)
 		{
-			const int32 NumBytes = InWaveInstance->WaveData->RawPCMDataSize;
-			NumFrames = NumBytes / (InWaveInstance->WaveData->NumChannels * sizeof(int16));
+			const int32 NumBytes = WaveInstance->WaveData->RawPCMDataSize;
+			NumFrames = NumBytes / (WaveInstance->WaveData->NumChannels * sizeof(int16));
 		}
 
 		// Unfortunately, we need to know if this is a vorbis source since channel maps are different for 5.1 vorbis files
-		bIsVorbis = InWaveInstance->WaveData->bDecompressedFromOgg;
+		bIsVorbis = WaveInstance->WaveData->bDecompressedFromOgg;
 
 		bIsStoppingVoicesEnabled = ((FAudioDevice*)MixerDevice)->IsStoppingVoicesEnabled();
 		
@@ -90,21 +92,21 @@ namespace Audio
 			// Initialize the source voice with the necessary format information
 			FMixerSourceVoiceInitParams InitParams;
 			InitParams.SourceListener = this;
-			InitParams.NumInputChannels = InWaveInstance->WaveData->NumChannels;
+			InitParams.NumInputChannels = WaveInstance->WaveData->NumChannels;
 			InitParams.NumInputFrames = NumFrames;
 			InitParams.SourceVoice = MixerSourceVoice;
 			InitParams.bUseHRTFSpatialization = UseObjectBasedSpatialization();
-			InitParams.bIsAmbisonics = InWaveInstance->bIsAmbisonics;
+			InitParams.bIsAmbisonics = WaveInstance->bIsAmbisonics;
 			if (InitParams.bIsAmbisonics)
 			{
 				checkf(InitParams.NumInputChannels == 4, TEXT("Only allow 4 channel source if file is ambisonics format."));
 			}
-			InitParams.AudioComponentUserID = InWaveInstance->ActiveSound->GetAudioComponentUserID();
+			InitParams.AudioComponentUserID = WaveInstance->ActiveSound->GetAudioComponentUserID();
 
-			InitParams.AudioComponentID = InWaveInstance->ActiveSound->GetAudioComponentID();
+			InitParams.AudioComponentID = WaveInstance->ActiveSound->GetAudioComponentID();
 
-			InitParams.EnvelopeFollowerAttackTime = InWaveInstance->EnvelopeFollowerAttackTime;
-			InitParams.EnvelopeFollowerReleaseTime = InWaveInstance->EnvelopeFollowerReleaseTime;
+			InitParams.EnvelopeFollowerAttackTime = WaveInstance->EnvelopeFollowerAttackTime;
+			InitParams.EnvelopeFollowerReleaseTime = WaveInstance->EnvelopeFollowerReleaseTime;
 
 			InitParams.SourceEffectChainId = 0;
 
@@ -113,14 +115,14 @@ namespace Audio
 
 			if (InitParams.NumInputChannels <= 2)
 			{
-				if (InWaveInstance->SourceEffectChain)
+				if (WaveInstance->SourceEffectChain)
 				{
-					InitParams.SourceEffectChainId = InWaveInstance->SourceEffectChain->GetUniqueID();
+					InitParams.SourceEffectChainId = WaveInstance->SourceEffectChain->GetUniqueID();
 
-					for (int32 i = 0; i < InWaveInstance->SourceEffectChain->Chain.Num(); ++i)
+					for (int32 i = 0; i < WaveInstance->SourceEffectChain->Chain.Num(); ++i)
 					{
-						InitParams.SourceEffectChain.Add(InWaveInstance->SourceEffectChain->Chain[i]);
-						InitParams.bPlayEffectChainTails = InWaveInstance->SourceEffectChain->bPlayEffectChainTails;
+						InitParams.SourceEffectChain.Add(WaveInstance->SourceEffectChain->Chain[i]);
+						InitParams.bPlayEffectChainTails = WaveInstance->SourceEffectChain->bPlayEffectChainTails;
 					}
 				}
 
@@ -131,24 +133,24 @@ namespace Audio
 				}
 
 				// Setup the bus Id if this source is a bus
-				if (InWaveInstance->WaveData->bIsBus)
+				if (WaveInstance->WaveData->bIsBus)
 				{
-					InitParams.BusId = InWaveInstance->WaveData->GetUniqueID();
-					if (!InWaveInstance->WaveData->IsLooping())
+					InitParams.BusId = WaveInstance->WaveData->GetUniqueID();
+					if (!WaveInstance->WaveData->IsLooping())
 					{
-						InitParams.BusDuration = InWaveInstance->WaveData->GetDuration();
+						InitParams.BusDuration = WaveInstance->WaveData->GetDuration();
 					}
 				}
 
 				// Toggle muting the source if sending only to output bus. 
 				// This can get set even if the source doesn't have bus sends since bus sends can be dynamically enabled.
-				InitParams.bOutputToBusOnly = InWaveInstance->bOutputToBusOnly;
+				InitParams.bOutputToBusOnly = WaveInstance->bOutputToBusOnly;
 
 				// If this source is sending its audio to a bus
 				for (int32 BusSendType = 0; BusSendType < (int32)EBusSendType::Count; ++BusSendType)
 				{
 					// And add all the source bus sends
-					for (FSoundSourceBusSendInfo& SendInfo : InWaveInstance->SoundSourceBusSends[BusSendType])
+					for (FSoundSourceBusSendInfo& SendInfo : WaveInstance->SoundSourceBusSends[BusSendType])
 					{
 						if (SendInfo.SoundSourceBus != nullptr)
 						{
@@ -166,10 +168,10 @@ namespace Audio
 			{
 				// If we're spatializing using HRTF and its an external send, don't need to setup a default/base submix send to master or EQ submix
 				// We'll only be using non-default submix sends (e.g. reverb).
-				if (!(InWaveInstance->SpatializationMethod == ESoundSpatializationAlgorithm::SPATIALIZATION_HRTF && MixerDevice->bSpatializationIsExternalSend))
+				if (!(WaveInstance->SpatializationMethod == ESoundSpatializationAlgorithm::SPATIALIZATION_HRTF && MixerDevice->bSpatializationIsExternalSend))
 				{
 					// If this sound is an ambisonics file, we preempt the normal base submix routing and only send to master ambisonics submix
-					if (InWaveInstance->bIsAmbisonics)
+					if (WaveInstance->bIsAmbisonics)
 					{
 						FMixerSourceSubmixSend SubmixSend;
 						SubmixSend.Submix = MixerDevice->GetMasterAmbisonicsSubmix();
@@ -180,10 +182,10 @@ namespace Audio
 					else
 					{
 						// If we've overridden which submix we're sending the sound, then add that as the first send
-						if (InWaveInstance->SoundSubmix != nullptr)
+						if (WaveInstance->SoundSubmix != nullptr)
 						{
 							FMixerSourceSubmixSend SubmixSend;
-							SubmixSend.Submix = MixerDevice->GetSubmixInstance(InWaveInstance->SoundSubmix);
+							SubmixSend.Submix = MixerDevice->GetSubmixInstance(WaveInstance->SoundSubmix);
 							SubmixSend.SendLevel = 1.0f;
 							SubmixSend.bIsMainSend = true;
 							InitParams.SubmixSends.Add(SubmixSend);
@@ -221,7 +223,7 @@ namespace Audio
 				}
 
 				// Now add any addition submix sends for this source
-				for (FSoundSubmixSendInfo& SendInfo : InWaveInstance->SoundSubmixSends)
+				for (FSoundSubmixSendInfo& SendInfo : WaveInstance->SoundSubmixSends)
 				{
 					if (SendInfo.SoundSubmix != nullptr)
 					{
@@ -246,7 +248,7 @@ namespace Audio
 
 			// Check to see if this sound has been flagged to be in debug mode
 #if AUDIO_MIXER_ENABLE_DEBUG_MODE
-			InitParams.DebugName = InWaveInstance->GetName();
+			InitParams.DebugName = WaveInstance->GetName();
 
 			bool bIsDebug = false;
 			FString WaveInstanceName = WaveInstance->GetName(); //-V595
@@ -262,19 +264,19 @@ namespace Audio
 			bIs3D = !UseObjectBasedSpatialization() && WaveInstance->bUseSpatialization && SoundBuffer->NumChannels < 3;
 
 			// Grab the source's reverb plugin settings 
-			InitParams.SpatializationPluginSettings = UseSpatializationPlugin() ? InWaveInstance->SpatializationPluginSettings : nullptr;
+			InitParams.SpatializationPluginSettings = UseSpatializationPlugin() ? WaveInstance->SpatializationPluginSettings : nullptr;
 
 			// Grab the source's occlusion plugin settings 
-			InitParams.OcclusionPluginSettings = UseOcclusionPlugin() ? InWaveInstance->OcclusionPluginSettings : nullptr;
+			InitParams.OcclusionPluginSettings = UseOcclusionPlugin() ? WaveInstance->OcclusionPluginSettings : nullptr;
 
 			// Grab the source's reverb plugin settings 
-			InitParams.ReverbPluginSettings = UseReverbPlugin() ? InWaveInstance->ReverbPluginSettings : nullptr;
+			InitParams.ReverbPluginSettings = UseReverbPlugin() ? WaveInstance->ReverbPluginSettings : nullptr;
 
 			// We support reverb
 			SetReverbApplied(true);
 
 			// Update the buffer sample rate to the wave instance sample rate in case it was serialized incorrectly
-			MixerBuffer->InitSampleRate(InWaveInstance->WaveData->GetSampleRateForCurrentPlatform());
+			MixerBuffer->InitSampleRate(WaveInstance->WaveData->GetSampleRateForCurrentPlatform());
 
 			// Now we init the mixer source buffer
 			MixerSourceBuffer->Init();
