Removed reduntant check for WaveInstance

diff --git a/Engine/Source/Runtime/AudioMixer/Private/AudioMixerSource.cpp b/Engine/Source/Runtime/AudioMixer/Private/AudioMixerSource.cpp
--- a/Engine/Source/Runtime/AudioMixer/Private/AudioMixerSource.cpp
+++ b/Engine/Source/Runtime/AudioMixer/Private/AudioMixerSource.cpp
@@ -244,8 +244,6 @@ namespace Audio
 				ChannelMaps[(int32)SubmixChannelType].ChannelMap.Reset();
 			}
 
-			AUDIO_MIXER_CHECK(WaveInstance);
-
 			// Check to see if this sound has been flagged to be in debug mode
 #if AUDIO_MIXER_ENABLE_DEBUG_MODE
 			InitParams.DebugName = WaveInstance->GetName();
