[NO NON-PII COMMIT MESSAGE]

diff --git a/Engine/Source/Runtime/AudioMixer/Private/AudioMixerSourceBuffer.cpp b/Engine/Source/Runtime/AudioMixer/Private/AudioMixerSourceBuffer.cpp
--- a/Engine/Source/Runtime/AudioMixer/Private/AudioMixerSourceBuffer.cpp
+++ b/Engine/Source/Runtime/AudioMixer/Private/AudioMixerSourceBuffer.cpp
@@ -418,13 +418,16 @@ namespace Audio
 			SoundWave = nullptr;
 		}
 
-		EBufferType::Type BufferType = MixerBuffer->GetType();
-		if (BufferType == EBufferType::PCMRealTime || BufferType == EBufferType::Streaming)
+		if (MixerBuffer)
 		{
-			delete MixerBuffer;
+			EBufferType::Type BufferType = MixerBuffer->GetType();
+			if (BufferType == EBufferType::PCMRealTime || BufferType == EBufferType::Streaming)
+			{
+				delete MixerBuffer;
+			}
+	
+			MixerBuffer = nullptr;
 		}
-
-		MixerBuffer = nullptr;
 	}
 
 }
