[NO NON-PII COMMIT MESSAGE]

diff --git a/Engine/Source/Runtime/AudioMixer/Private/AudioMixerSource.cpp b/Engine/Source/Runtime/AudioMixer/Private/AudioMixerSource.cpp
--- a/Engine/Source/Runtime/AudioMixer/Private/AudioMixerSource.cpp
+++ b/Engine/Source/Runtime/AudioMixer/Private/AudioMixerSource.cpp
@@ -25,6 +25,7 @@ namespace Audio
 		, MixerSourceBuffer(nullptr)
 		, MixerSourceVoice(nullptr)		
 		, PreviousAzimuth(-1.0f)
+		, InitializationState(EMixerSourceInitializationState::NotInitialized)
 		, bPlayedCachedBuffer(false)
 		, bPlaying(false)
 		, bLoopCallback(false)
@@ -285,12 +286,16 @@ namespace Audio
 
 			if (MixerSourceVoice->Init(InitParams))
 			{
-				bInitialized = true;
+				InitializationState = EMixerSourceInitializationState::Initialized;
 
 				Update();
 
 				return true;
 			}
+			else
+			{
+				InitializationState = EMixerSourceInitializationState::NotInitialized;
+			}
 		}
 		return false;
 	}
@@ -301,7 +306,7 @@ namespace Audio
 
 		LLM_SCOPE(ELLMTag::AudioMixer);
 
-		if (!WaveInstance || !MixerSourceVoice || Paused || !bInitialized)
+		if (!WaveInstance || !MixerSourceVoice || Paused || InitializationState == EMixerSourceInitializationState::NotInitialized)
 		{
 			return;
 		}
@@ -332,7 +337,7 @@ namespace Audio
 		}
 
 		// We are not initialized yet. We won't be until the sound file finishes loading and parsing the header.
-		bInitialized = false;
+		InitializationState = EMixerSourceInitializationState::Initializing;
 
 		//  Reset so next instance will warn if algorithm changes in-flight
 		bEditorWarnedChangedSpatialization = false;
@@ -438,6 +443,11 @@ namespace Audio
 		return false;
 	}
 
+	bool FMixerSource::IsInitialized() const
+	{
+		return InitializationState == EMixerSourceInitializationState::Initialized;
+	}
+
 	void FMixerSource::Play()
 	{
 		if (!WaveInstance)
@@ -447,7 +457,7 @@ namespace Audio
 
 		// It's possible if Pause and Play are called while a sound is async initializing. In this case
 		// we'll just not actually play the source here. Instead we'll call play when the sound finishes loading.
-		if (MixerSourceVoice && bInitialized)
+		if (MixerSourceVoice && InitializationState == EMixerSourceInitializationState::Initialized)
 		{
 			MixerSourceVoice->Play();
 		}
@@ -508,7 +518,8 @@ namespace Audio
 
 		// Immediately stop the sound source
 
-		bInitialized = false;
+		InitializationState = EMixerSourceInitializationState::NotInitialized;
+		
 		IStreamingManager::Get().GetAudioStreamingManager().RemoveStreamingSoundSource(this);
 
 		bIsStopping = false;
@@ -552,11 +563,16 @@ namespace Audio
 			return false;
 		}
 
-		if (!bInitialized)
+		if (InitializationState == EMixerSourceInitializationState::NotInitialized)
 		{
 			return true;
 		}
 
+		if (InitializationState == EMixerSourceInitializationState::Initializing)
+		{
+			return false;
+		}
+
 		if (WaveInstance && MixerSourceVoice)
 		{
 			if (bIsDone && bIsEffectTailsDone)
diff --git a/Engine/Source/Runtime/AudioMixer/Private/AudioMixerSource.h b/Engine/Source/Runtime/AudioMixer/Private/AudioMixerSource.h
--- a/Engine/Source/Runtime/AudioMixer/Private/AudioMixerSource.h
+++ b/Engine/Source/Runtime/AudioMixer/Private/AudioMixerSource.h
@@ -14,6 +14,14 @@ namespace Audio
 	class FMixerBuffer;
 	class ISourceListener;
 
+	/** State to track initialization stages. */
+	enum class EMixerSourceInitializationState : uint8
+	{
+		NotInitialized,
+		Initializing,
+		Initialized
+	};
+
 	/** 
 	 * FMixerSource
 	 * Class which implements a sound source object for the audio mixer module.
@@ -33,6 +41,7 @@ namespace Audio
 		virtual void Update() override;
 		virtual bool PrepareForInitialization(FWaveInstance* InWaveInstance) override;
 		virtual bool IsPreparedToInit() override;
+		virtual bool IsInitialized() const override;
 		virtual void Play() override;
 		virtual void Stop() override;
 		virtual void StopNow() override;
@@ -116,6 +125,8 @@ namespace Audio
 
 		FSpatializationParams SpatializationParams;
 
+		EMixerSourceInitializationState InitializationState;
+
 		FThreadSafeBool bPlayedCachedBuffer;
 		FThreadSafeBool bPlaying;
 		FThreadSafeBool bIsStopping;
diff --git a/Engine/Source/Runtime/Engine/Public/Audio.h b/Engine/Source/Runtime/Engine/Public/Audio.h
--- a/Engine/Source/Runtime/Engine/Public/Audio.h
+++ b/Engine/Source/Runtime/Engine/Public/Audio.h
@@ -644,7 +644,7 @@ public:
 	virtual bool Init(FWaveInstance* InWaveInstance) = 0;
 
 	/** Returns whether or not the sound source has initialized. */
-	bool IsInitialized() const { return bInitialized; };
+	virtual bool IsInitialized() const { return bInitialized; };
 
 	/** Updates the sound source. */
 	virtual void Update() = 0;
