[NO NON-PII COMMIT MESSAGE]

diff --git a/Engine/Source/Runtime/AudioMixer/Private/AudioMixerSourceBuffer.cpp b/Engine/Source/Runtime/AudioMixer/Private/AudioMixerSourceBuffer.cpp
--- a/Engine/Source/Runtime/AudioMixer/Private/AudioMixerSourceBuffer.cpp
+++ b/Engine/Source/Runtime/AudioMixer/Private/AudioMixerSourceBuffer.cpp
@@ -394,11 +394,15 @@ namespace Audio
 
 	void FMixerSourceBuffer::OnBeginGenerate()
 	{
-		if (SoundWave && SoundWave->bProcedural)
+		if (SoundWave)
 		{
 			SoundWave->bIsSoundActive = true;
 
-			SoundWave->OnBeginGenerate();
+			if (SoundWave->bProcedural)
+			{
+				SoundWave->OnBeginGenerate();
+			}
+		
 		}
 	}
 
@@ -412,9 +416,9 @@ namespace Audio
 			if (SoundWave->bProcedural)
 			{
 				SoundWave->OnEndGenerate();
-				SoundWave->bIsSoundActive = false;
 			}
 
+			SoundWave->bIsSoundActive = false;
 			SoundWave = nullptr;
 		}
 
diff --git a/Engine/Source/Runtime/Engine/Private/SoundWave.cpp b/Engine/Source/Runtime/Engine/Private/SoundWave.cpp
--- a/Engine/Source/Runtime/Engine/Private/SoundWave.cpp
+++ b/Engine/Source/Runtime/Engine/Private/SoundWave.cpp
@@ -874,7 +874,7 @@ bool USoundWave::IsReadyForFinishDestroy()
 		}, GET_STATID(STAT_AudioFreeResources));
 	}
 	
-	// bIsSoundActive is set in audio mixer when decoding sound waves or generating PCM data from procedural sound waves
+	// bIsSoundActive is set in audio mixer when decoding sound waves or generating PCM data
 	return ResourceState == ESoundWaveResourceState::Freed && !bIsSoundActive;
 }
 
