- Rather than a bool for flagging sounds as being used for GC, it's a ref count (multiple sound sources can play the same USoundWave).
- Removed a spot that was setting the flag to false causing a race condition.
- Tested with Shipping and Development client builds with shooter game and with very short match durations to exacerbate the GC flush.

diff --git a/Engine/Source/Runtime/AudioMixer/Private/AudioMixerSourceBuffer.cpp b/Engine/Source/Runtime/AudioMixer/Private/AudioMixerSourceBuffer.cpp
--- a/Engine/Source/Runtime/AudioMixer/Private/AudioMixerSourceBuffer.cpp
+++ b/Engine/Source/Runtime/AudioMixer/Private/AudioMixerSourceBuffer.cpp
@@ -59,6 +59,7 @@ namespace Audio
 		, SoundWave(nullptr)
 		, AsyncRealtimeAudioTask(nullptr)
 		, LoopingMode(ELoopingMode::LOOP_Never)
+		, bInitialized(false)
 		, bBufferFinished(false)
 		, bPlayedCachedBuffer(false)
 		, bIsSeeking(false)
@@ -82,18 +83,6 @@ namespace Audio
 		// May or may not be nullptr
 		SoundWave = InWave;
 
-		if (SoundWave->bProcedural)
-		{
-			if (SoundWave->bIsSoundActive)
-			{
-				UE_LOG(LogAudioMixer, Warning, TEXT("Procedural sound wave is reinitializing even though it is currently actively generating audio. Please stop sound before trying to play it again."));
-				return false;
-			}
-	
-			// We flag that this sound wave is active for the lifetime of this object since we use it for decoding, etc.
-			SoundWave->bIsSoundActive = true;
-		}
-
 		LoopingMode = InLoopingMode;
 		bIsSeeking = bInIsSeeking;
 		bLoopCallback = false;
@@ -113,8 +102,23 @@ namespace Audio
 		return true;
 	}
 
-	void FMixerSourceBuffer::Init()
+	bool FMixerSourceBuffer::Init()
 	{
+		check(SoundWave);
+		if (SoundWave->bProcedural && SoundWave->GetNumSoundsActive() > 0)
+		{
+			UE_LOG(LogAudioMixer, Warning, TEXT("Procedural sound wave is reinitializing even though it is currently actively generating audio. Please stop sound before trying to play it again."));
+			return false;
+		}
+
+		// We flag that this sound wave is active for the lifetime of this object since we use it for decoding, etc.
+		SoundWave->IncrementNumSounds();
+
+		// We have successfully initialized which means our SoundWave has been flagged as bIsActive
+		// GC can run between PreInit and Init so when cleaning up FMixerSourceBuffer, we don't want to touch SoundWave unless bInitailized is true.
+		// SoundWave->bIsSoundActive will prevent GC until it is released in audio render thread
+		bInitialized = true;
+
 		const EBufferType::Type BufferType = MixerBuffer->GetType();
 		switch (BufferType)
 		{
@@ -131,6 +135,8 @@ namespace Audio
 			case EBufferType::Invalid:
 				break;
 		}
+
+		return true;
 	}
 
 	void FMixerSourceBuffer::OnBufferEnd()
@@ -396,8 +402,6 @@ namespace Audio
 	{
 		if (SoundWave)
 		{
-			SoundWave->bIsSoundActive = true;
-
 			if (SoundWave->bProcedural)
 			{
 				SoundWave->OnBeginGenerate();
@@ -411,14 +415,15 @@ namespace Audio
 		// Make sure the async task finishes!
 		EnsureAsyncTaskFinishes();
 
-		if (SoundWave)
+		// Only need to call OnEndGenerate and access SoundWave here if we successfully initialized
+		if (bInitialized && SoundWave)
 		{
 			if (SoundWave->bProcedural)
 			{
 				SoundWave->OnEndGenerate();
 			}
 
-			SoundWave->bIsSoundActive = false;
+			SoundWave->DecrementNumSounds();
 			SoundWave = nullptr;
 		}
 
diff --git a/Engine/Source/Runtime/AudioMixer/Private/AudioMixerSourceBuffer.h b/Engine/Source/Runtime/AudioMixer/Private/AudioMixerSourceBuffer.h
--- a/Engine/Source/Runtime/AudioMixer/Private/AudioMixerSourceBuffer.h
+++ b/Engine/Source/Runtime/AudioMixer/Private/AudioMixerSourceBuffer.h
@@ -50,7 +50,7 @@ namespace Audio
 		~FMixerSourceBuffer();
 
 		bool PreInit(FMixerBuffer* InBuffer, USoundWave* InWave, ELoopingMode InLoopingMode, bool bInIsSeeking);
-		void Init();
+		bool Init();
 
 		// Called by source manager when needing more buffers
 		void OnBufferEnd();
@@ -108,6 +108,7 @@ namespace Audio
 		USoundWave* SoundWave;
 		IAudioTask* AsyncRealtimeAudioTask;
 		ELoopingMode LoopingMode;
+		bool bInitialized;
 		bool bBufferFinished;
 		bool bPlayedCachedBuffer;
 		bool bIsSeeking;
diff --git a/Engine/Source/Runtime/AudioMixer/Private/AudioMixerSourceDecode.cpp b/Engine/Source/Runtime/AudioMixer/Private/AudioMixerSourceDecode.cpp
--- a/Engine/Source/Runtime/AudioMixer/Private/AudioMixerSourceDecode.cpp
+++ b/Engine/Source/Runtime/AudioMixer/Private/AudioMixerSourceDecode.cpp
@@ -47,7 +47,7 @@ public:
 			case EAudioTaskType::Procedural:
 			{
 				// Make sure we've been flagged as active
-				if (!ProceduralTaskData.ProceduralSoundWave->bIsSoundActive)
+				if (ProceduralTaskData.ProceduralSoundWave->GetNumSoundsActive() == 0)
 				{
 					UE_LOG(LogAudioMixer, Warning, TEXT("Warning: Attempted to run async procedural task on inactive sound. Aborting."));
 					return;
diff --git a/Engine/Source/Runtime/AudioMixer/Private/Components/SynthComponent.cpp b/Engine/Source/Runtime/AudioMixer/Private/Components/SynthComponent.cpp
--- a/Engine/Source/Runtime/AudioMixer/Private/Components/SynthComponent.cpp
+++ b/Engine/Source/Runtime/AudioMixer/Private/Components/SynthComponent.cpp
@@ -283,7 +283,7 @@ void USynthComponent::OnUnregister()
 bool USynthComponent::IsReadyForOwnerToAutoDestroy() const
 {
 	const bool bIsAudioComponentReadyForDestroy = !AudioComponent || (AudioComponent && !AudioComponent->IsPlaying());
-	const bool bIsSynthSoundReadyForDestroy = !Synth || !Synth->bIsSoundActive;
+	const bool bIsSynthSoundReadyForDestroy = !Synth || !Synth->GetNumSoundsActive();
 	return bIsAudioComponentReadyForDestroy && bIsSynthSoundReadyForDestroy;
 }
 
diff --git a/Engine/Source/Runtime/Engine/Classes/Sound/SoundWave.h b/Engine/Source/Runtime/Engine/Classes/Sound/SoundWave.h
--- a/Engine/Source/Runtime/Engine/Classes/Sound/SoundWave.h
+++ b/Engine/Source/Runtime/Engine/Classes/Sound/SoundWave.h
@@ -166,9 +166,6 @@ class ENGINE_API USoundWave : public USoundBase
 	/** Whether this sound wave is beginning to be destroyed by GC. */
 	uint8 bIsBeginDestroy:1;
 
-	/** Whether this sound is actively being used by the audio renderer (in audio mixer). Prevents GC issues with rendering audio. */
-	uint8 bIsSoundActive:1;
-
 	/** Set to true of this is a bus sound source. This will result in the sound wave not generating audio for itself, but generate audio through instances. Used only in audio mixer. */
 	uint8 bIsBus:1;
 
@@ -223,6 +220,9 @@ private:
 	/** What state the precache decompressor is in. */
 	FThreadSafeCounter PrecacheState;
 
+	/** Number of sounds actively using this sound wave by the audio renderer (in audio mixer). Prevents GC issues with rendering realtime audio. */
+	FThreadSafeCounter NumSoundsActive;
+
 #if !WITH_EDITOR
 	// This is the sample rate gotten from platform settings.
 	float CachedSampleRateOverride;
@@ -391,6 +391,13 @@ public:
 	// Called when the procedural sound wave is done generating on the render thread. Only used in the audio mixer and when bProcedural is true..
 	virtual void OnEndGenerate() {};
 
+	// Returns number of sounds using this sound wave (audio mixer only)
+	int32 GetNumSoundsActive();
+
+	// Increment and decrement num sounds (used in audio mixer)
+	void IncrementNumSounds();
+	void DecrementNumSounds();
+
 	/**
 	* Overwrite sample rate. Used for procedural soundwaves, as well as sound waves that are resampled on compress/decompress.
 	*/
diff --git a/Engine/Source/Runtime/Engine/Private/AudioDevice.cpp b/Engine/Source/Runtime/Engine/Private/AudioDevice.cpp
--- a/Engine/Source/Runtime/Engine/Private/AudioDevice.cpp
+++ b/Engine/Source/Runtime/Engine/Private/AudioDevice.cpp
@@ -3814,7 +3814,7 @@ void FAudioDevice::AddNewActiveSound(const FActiveSound& NewActiveSound)
 	}
 
 	USoundWave* SoundWave = Cast<USoundWave>(Sound);
-	if (SoundWave && SoundWave->bProcedural && SoundWave->bIsSoundActive)
+	if (SoundWave && SoundWave->bProcedural && SoundWave->GetNumSoundsActive() > 0)
 	{
 		FString SoundWaveName;
 		SoundWave->GetName(SoundWaveName);
diff --git a/Engine/Source/Runtime/Engine/Private/SoundWave.cpp b/Engine/Source/Runtime/Engine/Private/SoundWave.cpp
--- a/Engine/Source/Runtime/Engine/Private/SoundWave.cpp
+++ b/Engine/Source/Runtime/Engine/Private/SoundWave.cpp
@@ -785,7 +785,6 @@ void USoundWave::FreeResources()
 	bDynamicResource = false;
 	DecompressionType = DTYPE_Setup;
 	bDecompressedFromOgg = false;
-	bIsBeginDestroy = false;
 
 	USoundWave* SoundWave = this;
 	FAudioThread::RunCommandOnGameThread([SoundWave]()
@@ -855,6 +854,22 @@ FWaveInstance* USoundWave::HandleStart( FActiveSound& ActiveSound, const UPTRINT
 	return WaveInstance;
 }
 
+int32 USoundWave::GetNumSoundsActive()
+{
+	return NumSoundsActive.GetValue();
+}
+
+void USoundWave::IncrementNumSounds()
+{
+	NumSoundsActive.Increment();
+}
+
+void USoundWave::DecrementNumSounds()
+{
+	int32 NewValue = NumSoundsActive.Decrement();
+	check(NewValue >= 0);
+}
+
 bool USoundWave::IsReadyForFinishDestroy()
 {
 	const bool bIsStreamingInProgress = IStreamingManager::Get().GetAudioStreamingManager().IsStreamingInProgress(this);
@@ -875,7 +890,7 @@ bool USoundWave::IsReadyForFinishDestroy()
 	}
 	
 	// bIsSoundActive is set in audio mixer when decoding sound waves or generating PCM data
-	return ResourceState == ESoundWaveResourceState::Freed && !bIsSoundActive;
+	return ResourceState == ESoundWaveResourceState::Freed && NumSoundsActive.GetValue() == 0;
 }
 
 
