Copying //UE4/Dev-Audio to //UE4/Dev-Main (Source: //UE4/Dev-Audio @ 4332921)

diff --git a/Engine/Plugins/Online/OnlineSubsystemUtils/Source/OnlineSubsystemUtils/Public/VoiceEngineImpl.h b/Engine/Plugins/Online/OnlineSubsystemUtils/Source/OnlineSubsystemUtils/Public/VoiceEngineImpl.h
--- a/Engine/Plugins/Online/OnlineSubsystemUtils/Source/OnlineSubsystemUtils/Public/VoiceEngineImpl.h
+++ b/Engine/Plugins/Online/OnlineSubsystemUtils/Source/OnlineSubsystemUtils/Public/VoiceEngineImpl.h
@@ -13,9 +13,12 @@
 #include "OnlineSubsystemBPCallHelper.h"
 #include "OnlineSubsystemUtilsPackage.h"
 #include "VoipListenerSynthComponent.h"
+#include "OnlineSubsystemUtilsPackage.h"
+
+#include "VoicePacketImpl.h"
+#include "UObject/CoreOnline.h"
 
 class IOnlineSubsystem;
-class FUniqueNetIdString;
 class IVoiceDecoder;
 class IVoiceEncoder;
 class IVoiceCapture;
@@ -117,7 +120,7 @@ class FVoiceEngineImpl : public IVoiceEngine, public FSelfRegisteringExec
 	friend class FVoiceSerializeHelper;
 
 	/** Mapping of UniqueIds to the incoming voice data and their audio component */
-	typedef TMap<FUniqueNetIdString, FRemoteTalkerDataImpl> FRemoteTalkerData;
+	typedef TMap<FUniqueNetIdWrapper, FRemoteTalkerDataImpl> FRemoteTalkerData;
 
 	/** Reference to the main online subsystem */
 	IOnlineSubsystem* OnlineSubsystem;
@@ -187,7 +190,7 @@ class FVoiceEngineImpl : public IVoiceEngine, public FSelfRegisteringExec
 	 * @param SamplesRequired number of samples needed for immediate playback
 	 * @param TalkerId id of the remote talker to allocate voice data for
 	 */
-	void GenerateVoiceData(USoundWaveProcedural* InProceduralWave, int32 SamplesRequired, FUniqueNetIdString TalkerId);
+	void GenerateVoiceData(USoundWaveProcedural* InProceduralWave, int32 SamplesRequired, const FUniqueNetId& TalkerId);
 
 PACKAGE_SCOPE:
 
@@ -271,7 +274,7 @@ public:
 
 	virtual bool IsRemotePlayerTalking(const FUniqueNetId& UniqueId) override
 	{
-		return RemoteTalkerBuffers.Find((const FUniqueNetIdString&)UniqueId) != nullptr;
+		return RemoteTalkerBuffers.Find(FUniqueNetIdWrapper(UniqueId.AsShared())) != nullptr;
 	}
 
 	virtual uint32 GetVoiceDataReadyFlags() const override;
diff --git a/Engine/Plugins/Online/OnlineSubsystemUtils/Source/OnlineSubsystemUtils/Public/VoiceInterfaceImpl.h b/Engine/Plugins/Online/OnlineSubsystemUtils/Source/OnlineSubsystemUtils/Public/VoiceInterfaceImpl.h
--- a/Engine/Plugins/Online/OnlineSubsystemUtils/Source/OnlineSubsystemUtils/Public/VoiceInterfaceImpl.h
+++ b/Engine/Plugins/Online/OnlineSubsystemUtils/Source/OnlineSubsystemUtils/Public/VoiceInterfaceImpl.h
@@ -3,6 +3,7 @@
 #pragma once
 
 #include "CoreMinimal.h"
+#include "UObject/CoreOnline.h"
 #include "OnlineSubsystemTypes.h"
 #include "Interfaces/VoiceInterface.h"
 #include "Net/VoiceDataCommon.h"
@@ -13,8 +14,8 @@
  * The generic implementation of the voice interface 
  */
 
-class ONLINESUBSYSTEMUTILS_API FOnlineVoiceImpl : public IOnlineVoice
-{
+class ONLINESUBSYSTEMUTILS_API FOnlineVoiceImpl : public IOnlineVoice {
+protected:
 	/** Reference to the main online subsystem */
 	class IOnlineSubsystem* OnlineSubsystem;
 	/** Reference to the sessions interface */
@@ -34,9 +35,10 @@ class ONLINESUBSYSTEMUTILS_API FOnlineVoiceImpl : public IOnlineVoice
 	/** State of all possible remote talkers */
 	TArray<FRemoteTalker> RemoteTalkers;
 	/** Remote players locally muted explicitly */
-	TArray<FUniqueNetIdString> SystemMuteList;
+
+	TArray<FUniqueNetIdWrapper> SystemMuteList;
 	/** All remote players locally muted (super set of SystemMuteList) */
-	TArray<FUniqueNetIdString> MuteList;
+	TArray<FUniqueNetIdWrapper> MuteList;
 
 	/** Time to wait for new data before triggering "not talking" */
 	float VoiceNotificationDelta;
@@ -51,7 +53,7 @@ class ONLINESUBSYSTEMUTILS_API FOnlineVoiceImpl : public IOnlineVoice
 	 *
 	 * @return pointer to the remote talker or NULL if not found
 	 */
-	struct FRemoteTalker* FindRemoteTalker(const FUniqueNetId& UniqueId);
+	virtual struct FRemoteTalker* FindRemoteTalker(const FUniqueNetId& UniqueId);
 
 	/**
 	 * Is a given id presently muted (either by system mute or game server)
@@ -60,7 +62,7 @@ class ONLINESUBSYSTEMUTILS_API FOnlineVoiceImpl : public IOnlineVoice
 	 *
 	 * @return true if the net id is muted at all, false otherwise
 	 */
-	bool IsLocallyMuted(const FUniqueNetId& UniqueId) const;
+	virtual bool IsLocallyMuted(const FUniqueNetId& UniqueId) const;
 
 	/**
 	 * Does a given id exist in the system wide mute list
@@ -69,7 +71,7 @@ class ONLINESUBSYSTEMUTILS_API FOnlineVoiceImpl : public IOnlineVoice
 	 *
 	 * @return true if the net id is on the system wide mute list, false otherwise
 	 */
-	bool IsSystemWideMuted(const FUniqueNetId& UniqueId) const;
+	virtual bool IsSystemWideMuted(const FUniqueNetId& UniqueId) const;
 
 PACKAGE_SCOPE:
 
@@ -85,29 +87,31 @@ PACKAGE_SCOPE:
 
 	// IOnlineVoice
 	virtual bool Init() override;
-	void ProcessMuteChangeNotification() override;
+	virtual void ProcessMuteChangeNotification() override;
+
+	virtual IVoiceEnginePtr CreateVoiceEngine() override;
 
 	/**
 	 * Cleanup voice interface
 	 */
-	void Shutdown();
+	virtual void Shutdown();
 
 	/**
 	 * Processes any talking delegates that need to be fired off
 	 *
 	 * @param DeltaTime the amount of time that has elapsed since the last tick
 	 */
-	void ProcessTalkingDelegates(float DeltaTime);
+	virtual void ProcessTalkingDelegates(float DeltaTime);
 
 	/**
 	 * Reads any data that is currently queued
 	 */
-	void ProcessLocalVoicePackets();
+	virtual void ProcessLocalVoicePackets();
 
 	/**
 	 * Submits network packets to audio system for playback
 	 */
-	void ProcessRemoteVoicePackets();
+	virtual void ProcessRemoteVoicePackets();
 
 	/**
 	 * Figures out which remote talkers need to be muted for a given local talker
@@ -115,7 +119,7 @@ PACKAGE_SCOPE:
 	 * @param TalkerIndex the talker that needs the mute list checked for
 	 * @param PlayerController the player controller associated with this talker
 	 */
-	void UpdateMuteListForLocalTalker(int32 TalkerIndex, class APlayerController* PlayerController);
+	virtual void UpdateMuteListForLocalTalker(int32 TalkerIndex, class APlayerController* PlayerController);
 
 public:
 
diff --git a/Engine/Plugins/Online/OnlineSubsystemUtils/Source/OnlineSubsystemUtils/Public/VoipListenerSynthComponent.h b/Engine/Plugins/Online/OnlineSubsystemUtils/Source/OnlineSubsystemUtils/Public/VoipListenerSynthComponent.h
--- a/Engine/Plugins/Online/OnlineSubsystemUtils/Source/OnlineSubsystemUtils/Public/VoipListenerSynthComponent.h
+++ b/Engine/Plugins/Online/OnlineSubsystemUtils/Source/OnlineSubsystemUtils/Public/VoipListenerSynthComponent.h
@@ -33,7 +33,7 @@ private:
 #endif
 
 UCLASS()
-class UVoipListenerSynthComponent : public USynthComponent
+class ONLINESUBSYSTEMUTILS_API UVoipListenerSynthComponent : public USynthComponent
 {
 
 	GENERATED_BODY()
@@ -73,7 +73,7 @@ public:
 	 * @param InBuffer pointer to a single channel buffer of audio samples.
 	 * @param NumBytes the number of bytes referenced by InBuffer.
 	 * @param InStartSample the Starting Sample associated with this packet.
-	 * @param DataFormat the format of this sample (float, fixed16, etc.
+	 * @param DataFormat the format of this sample (float, fixed16, etc.)
 	*/
 	void SubmitPacket(void* InBuffer, int32 NumBytes, int64 InStartSample, EVoipStreamDataFormat DataFormat);
 
diff --git a/Engine/Plugins/Runtime/ResonanceAudio/Source/ResonanceAudio/Private/ResonanceAudioAmbisonics.cpp b/Engine/Plugins/Runtime/ResonanceAudio/Source/ResonanceAudio/Private/ResonanceAudioAmbisonics.cpp
--- a/Engine/Plugins/Runtime/ResonanceAudio/Source/ResonanceAudio/Private/ResonanceAudioAmbisonics.cpp
+++ b/Engine/Plugins/Runtime/ResonanceAudio/Source/ResonanceAudio/Private/ResonanceAudioAmbisonics.cpp
@@ -104,7 +104,7 @@ namespace ResonanceAudio
 		}
 		else
 		{
-			checkf(false, TEXT("Tried to close a stream we did not open."));
+			UE_LOG(LogResonanceAudio, Warning, TEXT("Tried to close a stream id (%d) we did not open."), StreamId);
 		}
 	}
 
diff --git a/Engine/Plugins/Runtime/Synthesis/Source/Synthesis/Classes/SourceEffects/SourceEffectBitCrusher.h b/Engine/Plugins/Runtime/Synthesis/Source/Synthesis/Classes/SourceEffects/SourceEffectBitCrusher.h
--- a/Engine/Plugins/Runtime/Synthesis/Source/Synthesis/Classes/SourceEffects/SourceEffectBitCrusher.h
+++ b/Engine/Plugins/Runtime/Synthesis/Source/Synthesis/Classes/SourceEffects/SourceEffectBitCrusher.h
@@ -29,13 +29,13 @@ class SYNTHESIS_API FSourceEffectBitCrusher : public FSoundEffectSource
 {
 public:
 	// Called on an audio effect at initialization on main thread before audio processing begins.
-	virtual void Init(const FSoundEffectSourceInitData& InSampleRate) override;
+	virtual void Init(const FSoundEffectSourceInitData& InitData) override;
 	
 	// Called when an audio effect preset is changed
 	virtual void OnPresetChanged() override;
 
 	// Process the input block of audio. Called on audio thread.
-	virtual void ProcessAudio(const FSoundEffectSourceInputData& InData, FSoundEffectSourceOutputData& OutData) override;
+	virtual void ProcessAudio(const FSoundEffectSourceInputData& InData, float* OutAudioBufferData) override;
 
 protected:
 	Audio::FBitCrusher BitCrusher;
diff --git a/Engine/Plugins/Runtime/Synthesis/Source/Synthesis/Classes/SourceEffects/SourceEffectChorus.h b/Engine/Plugins/Runtime/Synthesis/Source/Synthesis/Classes/SourceEffects/SourceEffectChorus.h
--- a/Engine/Plugins/Runtime/Synthesis/Source/Synthesis/Classes/SourceEffects/SourceEffectChorus.h
+++ b/Engine/Plugins/Runtime/Synthesis/Source/Synthesis/Classes/SourceEffects/SourceEffectChorus.h
@@ -23,10 +23,14 @@ struct SYNTHESIS_API FSourceEffectChorusSettings
 	UPROPERTY(EditAnywhere, BlueprintReadWrite, Category = "SourceEffect|Preset", meta = (ClampMin = "0.0", ClampMax = "1.0", UIMin = "0.0", UIMax = "1.0"))
 	float Feedback;
 
-	// The wetlevel of the chorus effect
+	// The wet level of the chorus effect
 	UPROPERTY(EditAnywhere, BlueprintReadWrite, Category = "SourceEffect|Preset", meta = (ClampMin = "0.0", ClampMax = "1.0", UIMin = "0.0", UIMax = "1.0"))
 	float WetLevel;
 
+	// The dry level of the chorus effect
+	UPROPERTY(EditAnywhere, BlueprintReadWrite, Category = "SourceEffect|Preset", meta = (ClampMin = "0.0", ClampMax = "1.0", UIMin = "0.0", UIMax = "1.0"))
+	float DryLevel;
+
 	// The spread of the effect (larger means greater difference between left and right delay lines)
 	UPROPERTY(EditAnywhere, BlueprintReadWrite, Category = "SourceEffect|Preset", meta = (ClampMin = "0.0", ClampMax = "1.0", UIMin = "0.0", UIMax = "1.0"))
 	float Spread;
@@ -36,6 +40,7 @@ struct SYNTHESIS_API FSourceEffectChorusSettings
 		, Frequency(2.0f)
 		, Feedback(0.3f)
 		, WetLevel(0.5f)
+		, DryLevel(0.5f)
 		, Spread(0.0f)
 	{}
 };
@@ -44,13 +49,13 @@ class SYNTHESIS_API FSourceEffectChorus : public FSoundEffectSource
 {
 public:
 	// Called on an audio effect at initialization on main thread before audio processing begins.
-	virtual void Init(const FSoundEffectSourceInitData& InSampleRate) override;
+	virtual void Init(const FSoundEffectSourceInitData& InitData) override;
 	
 	// Called when an audio effect preset is changed
 	virtual void OnPresetChanged() override;
 
 	// Process the input block of audio. Called on audio thread.
-	virtual void ProcessAudio(const FSoundEffectSourceInputData& InData, FSoundEffectSourceOutputData& OutData) override;
+	virtual void ProcessAudio(const FSoundEffectSourceInputData& InData, float* OutAudioBufferData) override;
 
 protected:
 	Audio::FChorus Chorus;
diff --git a/Engine/Plugins/Runtime/Synthesis/Source/Synthesis/Classes/SourceEffects/SourceEffectDynamicsProcessor.h b/Engine/Plugins/Runtime/Synthesis/Source/Synthesis/Classes/SourceEffects/SourceEffectDynamicsProcessor.h
--- a/Engine/Plugins/Runtime/Synthesis/Source/Synthesis/Classes/SourceEffects/SourceEffectDynamicsProcessor.h
+++ b/Engine/Plugins/Runtime/Synthesis/Source/Synthesis/Classes/SourceEffects/SourceEffectDynamicsProcessor.h
@@ -89,13 +89,13 @@ public:
 	FSourceEffectDynamicsProcessor();
 
 	// Called on an audio effect at initialization on main thread before audio processing begins.
-	virtual void Init(const FSoundEffectSourceInitData& InSampleRate) override;
+	virtual void Init(const FSoundEffectSourceInitData& InitData) override;
 	
 	// Called when an audio effect preset is changed
 	virtual void OnPresetChanged() override;
 
 	// Process the input block of audio. Called on audio thread.
-	virtual void ProcessAudio(const FSoundEffectSourceInputData& InData, FSoundEffectSourceOutputData& OutData) override;
+	virtual void ProcessAudio(const FSoundEffectSourceInputData& InData, float* OutAudioBufferData) override;
 
 protected:
 
diff --git a/Engine/Plugins/Runtime/Synthesis/Source/Synthesis/Classes/SourceEffects/SourceEffectEQ.h b/Engine/Plugins/Runtime/Synthesis/Source/Synthesis/Classes/SourceEffects/SourceEffectEQ.h
--- a/Engine/Plugins/Runtime/Synthesis/Source/Synthesis/Classes/SourceEffects/SourceEffectEQ.h
+++ b/Engine/Plugins/Runtime/Synthesis/Source/Synthesis/Classes/SourceEffects/SourceEffectEQ.h
@@ -54,13 +54,13 @@ public:
 	FSourceEffectEQ();
 
 	// Called on an audio effect at initialization on main thread before audio processing begins.
-	virtual void Init(const FSoundEffectSourceInitData& InSampleRate) override;
+	virtual void Init(const FSoundEffectSourceInitData& InitData) override;
 	
 	// Called when an audio effect preset is changed
 	virtual void OnPresetChanged() override;
 
 	// Process the input block of audio. Called on audio thread.
-	virtual void ProcessAudio(const FSoundEffectSourceInputData& InData, FSoundEffectSourceOutputData& OutData) override;
+	virtual void ProcessAudio(const FSoundEffectSourceInputData& InData, float* OutAudioBufferData) override;
 
 protected:
 
@@ -69,6 +69,7 @@ protected:
 	float InAudioFrame[2];
 	float OutAudioFrame[2];
 	float SampleRate;
+	int32 NumChannels;
 };
 
 UCLASS(ClassGroup = AudioSourceEffect, meta = (BlueprintSpawnableComponent))
diff --git a/Engine/Plugins/Runtime/Synthesis/Source/Synthesis/Classes/SourceEffects/SourceEffectEnvelopeFollower.h b/Engine/Plugins/Runtime/Synthesis/Source/Synthesis/Classes/SourceEffects/SourceEffectEnvelopeFollower.h
--- a/Engine/Plugins/Runtime/Synthesis/Source/Synthesis/Classes/SourceEffects/SourceEffectEnvelopeFollower.h
+++ b/Engine/Plugins/Runtime/Synthesis/Source/Synthesis/Classes/SourceEffects/SourceEffectEnvelopeFollower.h
@@ -53,13 +53,13 @@ public:
 	~FSourceEffectEnvelopeFollower();
 
 	// Called on an audio effect at initialization on main thread before audio processing begins.
-	virtual void Init(const FSoundEffectSourceInitData& InSampleRate) override;
+	virtual void Init(const FSoundEffectSourceInitData& InitData) override;
 	
 	// Called when an audio effect preset is changed
 	virtual void OnPresetChanged() override;
 
 	// Process the input block of audio. Called on audio thread.
-	virtual void ProcessAudio(const FSoundEffectSourceInputData& InData, FSoundEffectSourceOutputData& OutData) override;
+	virtual void ProcessAudio(const FSoundEffectSourceInputData& InData, float* OutAudioBufferData) override;
 
 protected:
 
@@ -69,6 +69,7 @@ protected:
 	uint32 InstanceId;
 	int32 FrameCount;
 	int32 FramesToNotify;
+	int32 NumChannels;
 };
 
 DECLARE_DYNAMIC_MULTICAST_DELEGATE_OneParam(FOnEnvelopeFollowerUpdate, float, EnvelopeValue);
diff --git a/Engine/Plugins/Runtime/Synthesis/Source/Synthesis/Classes/SourceEffects/SourceEffectFoldbackDistortion.h b/Engine/Plugins/Runtime/Synthesis/Source/Synthesis/Classes/SourceEffects/SourceEffectFoldbackDistortion.h
--- a/Engine/Plugins/Runtime/Synthesis/Source/Synthesis/Classes/SourceEffects/SourceEffectFoldbackDistortion.h
+++ b/Engine/Plugins/Runtime/Synthesis/Source/Synthesis/Classes/SourceEffects/SourceEffectFoldbackDistortion.h
@@ -34,13 +34,13 @@ class SYNTHESIS_API FSourceEffectFoldbackDistortion : public FSoundEffectSource
 {
 public:
 	// Called on an audio effect at initialization on main thread before audio processing begins.
-	virtual void Init(const FSoundEffectSourceInitData& InSampleRate) override;
+	virtual void Init(const FSoundEffectSourceInitData& InitData) override;
 	
 	// Called when an audio effect preset is changed
 	virtual void OnPresetChanged() override;
 
 	// Process the input block of audio. Called on audio thread.
-	virtual void ProcessAudio(const FSoundEffectSourceInputData& InData, FSoundEffectSourceOutputData& OutData) override;
+	virtual void ProcessAudio(const FSoundEffectSourceInputData& InData, float* OutAudioBufferData);
 
 protected:
 	Audio::FFoldbackDistortion FoldbackDistortion;
diff --git a/Engine/Plugins/Runtime/Synthesis/Source/Synthesis/Classes/SourceEffects/SourceEffectPanner.h b/Engine/Plugins/Runtime/Synthesis/Source/Synthesis/Classes/SourceEffects/SourceEffectPanner.h
--- a/Engine/Plugins/Runtime/Synthesis/Source/Synthesis/Classes/SourceEffects/SourceEffectPanner.h
+++ b/Engine/Plugins/Runtime/Synthesis/Source/Synthesis/Classes/SourceEffects/SourceEffectPanner.h
@@ -24,17 +24,18 @@ class SYNTHESIS_API FSourceEffectPanner : public FSoundEffectSource
 {
 public:
 	// Called on an audio effect at initialization on main thread before audio processing begins.
-	virtual void Init(const FSoundEffectSourceInitData& InSampleRate) override;
+	virtual void Init(const FSoundEffectSourceInitData& InitData) override;
 	
 	// Called when an audio effect preset is changed
 	virtual void OnPresetChanged() override;
 
-	// Process the input block of audio. Called on audio thread.
-	virtual void ProcessAudio(const FSoundEffectSourceInputData& InData, FSoundEffectSourceOutputData& OutData) override;
+	// Process the input block of audio. Called on audio render thread.
+	virtual void ProcessAudio(const FSoundEffectSourceInputData& InData, float* OutAudioBufferData) override;
 
 protected:
 	// The pan value of the source effect
 	float PanValue;
+	int32 NumChannels;
 };
 
 UCLASS(ClassGroup = AudioSourceEffect, meta = (BlueprintSpawnableComponent))
diff --git a/Engine/Plugins/Runtime/Synthesis/Source/Synthesis/Classes/SourceEffects/SourceEffectRingModulation.h b/Engine/Plugins/Runtime/Synthesis/Source/Synthesis/Classes/SourceEffects/SourceEffectRingModulation.h
--- a/Engine/Plugins/Runtime/Synthesis/Source/Synthesis/Classes/SourceEffects/SourceEffectRingModulation.h
+++ b/Engine/Plugins/Runtime/Synthesis/Source/Synthesis/Classes/SourceEffects/SourceEffectRingModulation.h
@@ -33,10 +33,20 @@ struct SYNTHESIS_API FSourceEffectRingModulationSettings
 	UPROPERTY(EditAnywhere, BlueprintReadWrite, Category = "SourceEffect|Preset", meta = (ClampMin = "0.0", ClampMax = "1.0", UIMin = "0.0", UIMax = "1.0"))
 	float Depth;
 
+	// Gain for the dry signal (no ring mod)
+	UPROPERTY(EditAnywhere, BlueprintReadWrite, Category = "SourceEffect|Preset", meta = (ClampMin = "0.0", ClampMax = "1.0", UIMin = "0.0", UIMax = "1.0"))
+	float DryLevel;
+
+	// Gain for the wet signal (with ring mod)
+	UPROPERTY(EditAnywhere, BlueprintReadWrite, Category = "SourceEffect|Preset", meta = (ClampMin = "0.0", ClampMax = "1.0", UIMin = "0.0", UIMax = "1.0"))
+	float WetLevel;
+
 	FSourceEffectRingModulationSettings()
 		: ModulatorType(ERingModulatorTypeSourceEffect::Sine)
 		, Frequency(100.0f)
 		, Depth(0.5f)
+		, DryLevel(0.0f)
+		, WetLevel(1.0f)
 	{}
 };
 
@@ -44,13 +54,13 @@ class SYNTHESIS_API FSourceEffectRingModulation : public FSoundEffectSource
 {
 public:
 	// Called on an audio effect at initialization on main thread before audio processing begins.
-	virtual void Init(const FSoundEffectSourceInitData& InSampleRate) override;
+	virtual void Init(const FSoundEffectSourceInitData& InitData) override;
 	
 	// Called when an audio effect preset is changed
 	virtual void OnPresetChanged() override;
 
 	// Process the input block of audio. Called on audio thread.
-	virtual void ProcessAudio(const FSoundEffectSourceInputData& InData, FSoundEffectSourceOutputData& OutData) override;
+	virtual void ProcessAudio(const FSoundEffectSourceInputData& InData, float* OutAudioBufferData) override;
 
 protected:
 	Audio::FRingModulation RingModulation;
diff --git a/Engine/Plugins/Runtime/Synthesis/Source/Synthesis/Classes/SourceEffects/SourceEffectSimpleDelay.h b/Engine/Plugins/Runtime/Synthesis/Source/Synthesis/Classes/SourceEffects/SourceEffectSimpleDelay.h
--- a/Engine/Plugins/Runtime/Synthesis/Source/Synthesis/Classes/SourceEffects/SourceEffectSimpleDelay.h
+++ b/Engine/Plugins/Runtime/Synthesis/Source/Synthesis/Classes/SourceEffects/SourceEffectSimpleDelay.h
@@ -19,6 +19,14 @@ struct SYNTHESIS_API FSourceEffectSimpleDelaySettings
 	UPROPERTY(EditAnywhere, BlueprintReadWrite, Category = "SourceEffect|Preset", meta = (ClampMin = "0.0", ClampMax = "2.0", UIMin = "0.0", UIMax = "2.0"))
 	float DelayAmount;
 
+	// Gain stage on dry (non-delayed signal)
+	UPROPERTY(EditAnywhere, BlueprintReadWrite, Category = "SourceEffect|Preset", meta = (ClampMin = "0.0", ClampMax = "1.0", UIMin = "0.0", UIMax = "1.0"))
+	float DryAmount;
+
+	// Gain stage on wet (delayed) signal
+	UPROPERTY(EditAnywhere, BlueprintReadWrite, Category = "SourceEffect|Preset", meta = (ClampMin = "0.0", ClampMax = "1.0", UIMin = "0.0", UIMax = "1.0"))
+	float WetAmount;
+
 	// Amount to feedback into the delay line (because why not)
 	UPROPERTY(EditAnywhere, BlueprintReadWrite, Category = "SourceEffect|Preset", meta = (ClampMin = "0.0", ClampMax = "1.0", UIMin = "0.0", UIMax = "1.0"))
 	float Feedback;
@@ -30,6 +38,8 @@ struct SYNTHESIS_API FSourceEffectSimpleDelaySettings
 	FSourceEffectSimpleDelaySettings()
 		: SpeedOfSound(343.0f)
 		, DelayAmount(0.0f)
+		, DryAmount(0.0f)
+		, WetAmount(1.0f)
 		, Feedback(0.0f)
 		, bDelayBasedOnDistance(true)
 	{}
@@ -39,18 +49,18 @@ class SYNTHESIS_API FSourceEffectSimpleDelay : public FSoundEffectSource
 {
 public:
 	// Called on an audio effect at initialization on main thread before audio processing begins.
-	virtual void Init(const FSoundEffectSourceInitData& InSampleRate) override;
+	virtual void Init(const FSoundEffectSourceInitData& InitData) override;
 	
 	// Called when an audio effect preset is changed
 	virtual void OnPresetChanged() override;
 
 	// Process the input block of audio. Called on audio thread.
-	virtual void ProcessAudio(const FSoundEffectSourceInputData& InData, FSoundEffectSourceOutputData& OutData) override;
+	virtual void ProcessAudio(const FSoundEffectSourceInputData& InData, float* OutAudioBufferData) override;
 
 protected:
 	// 2 Delay lines, one for each channel
-	Audio::FDelay DelayLines[2];
-	float FeedbackSamples[2];
+	TArray<Audio::FDelay> Delays;
+	TArray<float> FeedbackSamples;
 	FSourceEffectSimpleDelaySettings SettingsCopy;
 };
 
diff --git a/Engine/Plugins/Runtime/Synthesis/Source/Synthesis/Classes/SubmixEffects/SubmixEffectFilter.h b/Engine/Plugins/Runtime/Synthesis/Source/Synthesis/Classes/SubmixEffects/SubmixEffectFilter.h
new file mode 100644
--- /dev/null
+++ b/Engine/Plugins/Runtime/Synthesis/Source/Synthesis/Classes/SubmixEffects/SubmixEffectFilter.h
@@ -0,0 +1,171 @@
+// Copyright 1998-2018 Epic Games, Inc. All Rights Reserved.
+
+#pragma once
+
+#include "CoreMinimal.h"
+#include "AudioEffect.h"
+#include "DSP/Filter.h"
+#include "Sound/SoundEffectSubmix.h"
+#include "SubmixEffectFilter.generated.h"
+
+UENUM(BlueprintType)
+enum class ESubmixFilterType : uint8
+{
+	LowPass = 0,
+	HighPass,
+	BandPass,
+	BandStop,
+	Count UMETA(Hidden)
+};
+
+UENUM(BlueprintType)
+enum class ESubmixFilterAlgorithm : uint8
+{
+	OnePole = 0,
+	StateVariable,
+	Ladder,
+	Count UMETA(Hidden)
+};
+
+// ========================================================================
+// FSubmixEffectFilterSettings
+// UStruct used to define user-exposed params for use with your effect.
+// ========================================================================
+
+USTRUCT(BlueprintType)
+struct SYNTHESIS_API FSubmixEffectFilterSettings
+{
+	GENERATED_USTRUCT_BODY()
+
+	// What type of filter to use for the submix filter effect
+	UPROPERTY(EditAnywhere, BlueprintReadWrite, Category = Filter)
+	ESubmixFilterType FilterType;
+
+	// What type of filter algorithm to use for the submix filter effect
+	UPROPERTY(EditAnywhere, BlueprintReadWrite, Category = Filter)
+	ESubmixFilterAlgorithm FilterAlgorithm;
+
+	// The output filter cutoff frequency (hz) [0.0, 20000.0]
+	UPROPERTY(EditAnywhere, BlueprintReadWrite, Category = Filter, meta = (ClampMin = "0.0", ClampMax = "20000.0", UIMin = "0.0", UIMax = "20000.0"))
+	float FilterFrequency;
+
+	// The output filter resonance (Q) [0.5, 10]
+	UPROPERTY(EditAnywhere, BlueprintReadWrite, Category = Filter, meta = (ClampMin = "0.5", ClampMax = "10.0", UIMin = "0.5", UIMax = "10.0"))
+	float FilterQ;
+
+	FSubmixEffectFilterSettings()
+		: FilterType(ESubmixFilterType::LowPass)
+		, FilterAlgorithm(ESubmixFilterAlgorithm::OnePole)
+		, FilterFrequency(20000.0f)
+		, FilterQ(2.0f)
+	{
+	}
+};
+
+class SYNTHESIS_API FSubmixEffectFilter : public FSoundEffectSubmix
+{
+public:
+	FSubmixEffectFilter();
+	~FSubmixEffectFilter();
+
+	//~ Begin FSoundEffectSubmix
+	virtual void Init(const FSoundEffectSubmixInitData& InData) override;
+	virtual void OnProcessAudio(const FSoundEffectSubmixInputData& InData, FSoundEffectSubmixOutputData& OutData) override;
+	//~ End FSoundEffectSubmix
+
+	//~ Begin FSoundEffectBase
+	virtual void OnPresetChanged() override;
+	//~End FSoundEffectBase
+
+	// Sets the filter type
+	void SetFilterType(ESubmixFilterType InType);
+
+	// Sets the filter algorithm
+	void SetFilterAlgorithm(ESubmixFilterAlgorithm InAlgorithm);
+
+	// Sets the base filter cutoff frequency
+	void SetFilterCutoffFrequency(float InFrequency);
+
+	// Sets the mod filter cutoff frequency
+	void SetFilterCutoffFrequencyMod(float InFrequency);
+
+	// Sets the filter Q
+	void SetFilterQ(float InQ);
+
+	// Sets the filter Q
+	void SetFilterQMod(float InQ);
+
+private:
+
+	void InitFilter();
+
+	// Sample rate of the submix effect
+	float SampleRate;
+
+	// Filters for each channel and type
+	Audio::FOnePoleFilter OnePoleFilter;
+	Audio::FStateVariableFilter StateVariableFilter;
+	Audio::FLadderFilter LadderFilter;
+
+	// The current filter selected
+	Audio::IFilter* CurrentFilter;
+
+	// Filter control data
+	ESubmixFilterAlgorithm FilterAlgorithm;
+	ESubmixFilterType FilterType;
+
+	float FilterFrequency;
+	float FilterFrequencyMod;
+
+	float FilterQ;
+	float FilterQMod;
+
+	int32 NumChannels;
+};
+
+// ========================================================================
+// USubmixEffectFilterPreset
+// Class which processes audio streams and uses parameters defined in the preset class.
+// ========================================================================
+
+UCLASS()
+class SYNTHESIS_API USubmixEffectFilterPreset : public USoundEffectSubmixPreset
+{
+	GENERATED_BODY()
+
+public:
+	EFFECT_PRESET_METHODS(SubmixEffectFilter)
+
+	// Set all filter effect settings
+	UFUNCTION(BlueprintCallable, Category = "Audio|Effects|Filter")
+	void SetSettings(const FSubmixEffectFilterSettings& InSettings);
+
+	// Sets the filter type
+	UFUNCTION(BlueprintCallable, Category = "Audio|Effects|Filter")
+	void SetFilterType(ESubmixFilterType InType);
+
+	// Sets the filter algorithm
+	UFUNCTION(BlueprintCallable, Category = "Audio|Effects|Filter")
+	void SetFilterAlgorithm(ESubmixFilterAlgorithm InAlgorithm);
+
+	// Sets the base filter cutoff frequency
+	UFUNCTION(BlueprintCallable, Category = "Audio|Effects|Filter")
+	void SetFilterCutoffFrequency(float InFrequency);
+
+	// Sets the mod filter cutoff frequency
+	UFUNCTION(BlueprintCallable, Category = "Audio|Effects|Filter")
+	void SetFilterCutoffFrequencyMod(float InFrequency);
+
+	// Sets the filter Q
+	UFUNCTION(BlueprintCallable, Category = "Audio|Effects|Filter")
+	void SetFilterQ(float InQ);
+
+	// Sets the filter Q
+	UFUNCTION(BlueprintCallable, Category = "Audio|Effects|Filter")
+	void SetFilterQMod(float InQ);
+
+public:
+
+	UPROPERTY(EditAnywhere, BlueprintReadWrite, Category = SubmixEffectPreset, Meta = (ShowOnlyInnerProperties))
+	FSubmixEffectFilterSettings Settings;
+};
diff --git a/Engine/Plugins/Runtime/Synthesis/Source/Synthesis/Private/EpicSynth1.cpp b/Engine/Plugins/Runtime/Synthesis/Source/Synthesis/Private/EpicSynth1.cpp
--- a/Engine/Plugins/Runtime/Synthesis/Source/Synthesis/Private/EpicSynth1.cpp
+++ b/Engine/Plugins/Runtime/Synthesis/Source/Synthesis/Private/EpicSynth1.cpp
@@ -692,7 +692,7 @@ namespace Audio
 
 		// Apply filters...
 #if !SYNTH_DEBUG_MODE
-		CurrentFilter->ProcessAudio(OutSamples, OutSamples);
+		CurrentFilter->ProcessAudioFrame(OutSamples, OutSamples);
 #endif
 		// Stop the oscillators if they're done
 		if (GainEnv.IsDone())
@@ -1455,11 +1455,10 @@ namespace Audio
 		return true;
 	}
 
-	void FEpicSynth1::Generate(float& OutLeft, float& OutRight)
+	void FEpicSynth1::GenerateFrame(float* OutFrame)
 	{
-		float OutputSamples[2];
-		OutputSamples[0] = 0.0f;
-		OutputSamples[1] = 0.0f;
+		OutFrame[0] = 0.0f;
+		OutFrame[1] = 0.0f;
 
 		for (int32 VoiceId = 0; VoiceId < Voices.Num(); ++VoiceId)
 		{
@@ -1497,27 +1496,19 @@ namespace Audio
 			Voices[VoiceId]->Generate(VoiceSamples);
 
 			// Mix the voices audio with the output
-			OutputSamples[0] += VoiceSamples[0];
-			OutputSamples[1] += VoiceSamples[1];
+			OutFrame[0] += VoiceSamples[0];
+			OutFrame[1] += VoiceSamples[1];
 		}
 
-#if 0
-		OutLeft = OutputSamples[0];
-		OutRight = OutputSamples[1];
-#else
-		OutLeft = OutputSamples[0];
-		OutRight = OutputSamples[1];
-
 		if (bIsChorusEnabled)
 		{
-			Chorus.ProcessAudio(OutLeft, OutRight, OutLeft, OutRight);
+			Chorus.ProcessAudioFrame(OutFrame, OutFrame);
 		}
 
 		if (bIsStereoEnabled)
 		{
-			StereoDelay.ProcessAudio(OutLeft, OutRight, OutLeft, OutRight);
+			StereoDelay.ProcessAudioFrame(OutFrame, OutFrame);
 		}
-#endif
 	}
 
 
diff --git a/Engine/Plugins/Runtime/Synthesis/Source/Synthesis/Private/EpicSynth1Component.cpp b/Engine/Plugins/Runtime/Synthesis/Source/Synthesis/Private/EpicSynth1Component.cpp
--- a/Engine/Plugins/Runtime/Synthesis/Source/Synthesis/Private/EpicSynth1Component.cpp
+++ b/Engine/Plugins/Runtime/Synthesis/Source/Synthesis/Private/EpicSynth1Component.cpp
@@ -607,18 +607,9 @@ bool UModularSynthComponent::Init(int32& SampleRate)
 
 int32 UModularSynthComponent::OnGenerateAudio(float* OutAudio, int32 NumSamples)
 {
-	const int32 NumFrames = NumSamples / NumChannels;
-
-	float LeftSample = 0.0f;
-	float RightSample = 0.0f;
-	int32 SampleIndex = 0;
-
-	for (int32 Frame = 0; Frame < NumFrames; ++Frame)
+	for (int32 SampleIndex = 0; SampleIndex < NumSamples; SampleIndex += NumChannels)
 	{
-		EpicSynth1.Generate(LeftSample, RightSample);
-
-		OutAudio[SampleIndex++] = LeftSample;
-		OutAudio[SampleIndex++] = RightSample;
+		EpicSynth1.GenerateFrame(&OutAudio[SampleIndex]);
 	}
 	return NumSamples;
 }
diff --git a/Engine/Plugins/Runtime/Synthesis/Source/Synthesis/Private/Flexiverb.cpp b/Engine/Plugins/Runtime/Synthesis/Source/Synthesis/Private/Flexiverb.cpp
--- a/Engine/Plugins/Runtime/Synthesis/Source/Synthesis/Private/Flexiverb.cpp
+++ b/Engine/Plugins/Runtime/Synthesis/Source/Synthesis/Private/Flexiverb.cpp
@@ -112,13 +112,13 @@ namespace Audio
 		DryInput /= InChannels;
 
 		// Process the input through our pre-delay:
-		PreDelayLine.ProcessAudio(&DryInput, &DryInput);
+		DryInput = PreDelayLine.ProcessAudioSample(DryInput);
 
 		// Process the pre-delay output through a series of APFs:
 		FDelayAPF* APFArrayPtr = APFArray.GetData();
 		for (int32 APFIndex = 0; APFIndex < APFArray.Num(); ++APFIndex)
 		{
-			APFArrayPtr[APFIndex].ProcessAudio(&DryInput, &DryInput);
+			DryInput = APFArrayPtr[APFIndex].ProcessAudioSample(DryInput);
 		}
 		
 		// Process the APF series output through our FDN.
@@ -146,10 +146,10 @@ namespace Audio
 			Temp *= DecayFactor;
 			
 			// Process our delay line input through it's corresponding LPF:
-			DampeningArrayPtr[DelayColumn].ProcessAudio(&Temp, &Temp);
+			Temp = DampeningArrayPtr[DelayColumn].ProcessAudioSample(Temp);
 
 			// Finally, process the delay line.
-			DelayLinesPtr[DelayColumn].ProcessAudio(&Temp, &Temp);
+			Temp = DelayLinesPtr[DelayColumn].ProcessAudioSample(Temp);
 
 			// Tap out of each delay line and sum into our OutputAccumulator:
 			OutputAccumulator += Temp;
diff --git a/Engine/Plugins/Runtime/Synthesis/Source/Synthesis/Private/SourceEffectBitCrusher.cpp b/Engine/Plugins/Runtime/Synthesis/Source/Synthesis/Private/SourceEffectBitCrusher.cpp
--- a/Engine/Plugins/Runtime/Synthesis/Source/Synthesis/Private/SourceEffectBitCrusher.cpp
+++ b/Engine/Plugins/Runtime/Synthesis/Source/Synthesis/Private/SourceEffectBitCrusher.cpp
@@ -5,7 +5,7 @@
 void FSourceEffectBitCrusher::Init(const FSoundEffectSourceInitData& InitData)
 {
 	bIsActive = true;
-	BitCrusher.Init(InitData.SampleRate);
+	BitCrusher.Init(InitData.SampleRate, InitData.NumSourceChannels);
 }
 
 void FSourceEffectBitCrusher::OnPresetChanged()
@@ -16,16 +16,9 @@ void FSourceEffectBitCrusher::OnPresetChanged()
 	BitCrusher.SetSampleRateCrush(Settings.CrushedSampleRate);
 }
 
-void FSourceEffectBitCrusher::ProcessAudio(const FSoundEffectSourceInputData& InData, FSoundEffectSourceOutputData& OutData)
+void FSourceEffectBitCrusher::ProcessAudio(const FSoundEffectSourceInputData& InData, float* OutAudioBufferData)
 {
-	if (InData.AudioFrame.Num() == 2)
-	{
-		BitCrusher.ProcessAudio(InData.AudioFrame[0], InData.AudioFrame[1], OutData.AudioFrame[0], OutData.AudioFrame[1]);
-	}
-	else
-	{
-		BitCrusher.ProcessAudio(InData.AudioFrame[0], OutData.AudioFrame[0]);
-	}
+	BitCrusher.ProcessAudio(InData.InputSourceEffectBufferPtr, InData.NumSamples, OutAudioBufferData);
 }
 
 void USourceEffectBitCrusherPreset::SetSettings(const FSourceEffectBitCrusherSettings& InSettings)
diff --git a/Engine/Plugins/Runtime/Synthesis/Source/Synthesis/Private/SourceEffectChorus.cpp b/Engine/Plugins/Runtime/Synthesis/Source/Synthesis/Private/SourceEffectChorus.cpp
--- a/Engine/Plugins/Runtime/Synthesis/Source/Synthesis/Private/SourceEffectChorus.cpp
+++ b/Engine/Plugins/Runtime/Synthesis/Source/Synthesis/Private/SourceEffectChorus.cpp
@@ -6,7 +6,7 @@ void FSourceEffectChorus::Init(const FSoundEffectSourceInitData& InitData)
 {
 	bIsActive = true;
 
-	Chorus.Init(InitData.SampleRate, 2.0f, 64);
+	Chorus.Init(InitData.SampleRate, InitData.NumSourceChannels, 2.0f, 64);
 }
 
 void FSourceEffectChorus::OnPresetChanged()
@@ -26,22 +26,13 @@ void FSourceEffectChorus::OnPresetChanged()
 	Chorus.SetFrequency(Audio::EChorusDelays::Right, Settings.Frequency);
 
 	Chorus.SetWetLevel(Settings.WetLevel);
+	Chorus.SetDryLevel(Settings.DryLevel);
 	Chorus.SetSpread(Settings.Spread);
 }
 
-void FSourceEffectChorus::ProcessAudio(const FSoundEffectSourceInputData& InData, FSoundEffectSourceOutputData& OutData)
+void FSourceEffectChorus::ProcessAudio(const FSoundEffectSourceInputData& InData, float* OutAudioBufferData)
 {
-	if (InData.AudioFrame.Num() == 2)
-	{
-		Chorus.ProcessAudio(InData.AudioFrame[0], InData.AudioFrame[1], OutData.AudioFrame[0], OutData.AudioFrame[1]);
-	}
-	else
-	{
-		float OutLeft = 0.0f;
-		float OutRight = 0.0f;
-		Chorus.ProcessAudio(InData.AudioFrame[0], InData.AudioFrame[0], OutLeft, OutRight);
-		OutData.AudioFrame[0] = 0.5f * (OutLeft + OutRight);
-	}
+	Chorus.ProcessAudio(InData.InputSourceEffectBufferPtr, InData.NumSamples, OutAudioBufferData);
 }
 
 void USourceEffectChorusPreset::SetSettings(const FSourceEffectChorusSettings& InSettings)
diff --git a/Engine/Plugins/Runtime/Synthesis/Source/Synthesis/Private/SourceEffectDynamicsProcessor.cpp b/Engine/Plugins/Runtime/Synthesis/Source/Synthesis/Private/SourceEffectDynamicsProcessor.cpp
--- a/Engine/Plugins/Runtime/Synthesis/Source/Synthesis/Private/SourceEffectDynamicsProcessor.cpp
+++ b/Engine/Plugins/Runtime/Synthesis/Source/Synthesis/Private/SourceEffectDynamicsProcessor.cpp
@@ -65,9 +65,9 @@ void FSourceEffectDynamicsProcessor::OnPresetChanged()
 	DynamicsProcessor.SetAnalogMode(Settings.bAnalogMode);
 }
 
-void FSourceEffectDynamicsProcessor::ProcessAudio(const FSoundEffectSourceInputData& InData, FSoundEffectSourceOutputData& OutData)
+void FSourceEffectDynamicsProcessor::ProcessAudio(const FSoundEffectSourceInputData& InData, float* OutAudioBufferData)
 {
-	DynamicsProcessor.ProcessAudio(InData.AudioFrame.GetData(), InData.AudioFrame.Num(), OutData.AudioFrame.GetData());
+	DynamicsProcessor.ProcessAudio(InData.InputSourceEffectBufferPtr, InData.NumSamples, OutAudioBufferData);
 }
 
 void USourceEffectDynamicsProcessorPreset::SetSettings(const FSourceEffectDynamicsProcessorSettings& InSettings)
diff --git a/Engine/Plugins/Runtime/Synthesis/Source/Synthesis/Private/SourceEffectEQ.cpp b/Engine/Plugins/Runtime/Synthesis/Source/Synthesis/Private/SourceEffectEQ.cpp
--- a/Engine/Plugins/Runtime/Synthesis/Source/Synthesis/Private/SourceEffectEQ.cpp
+++ b/Engine/Plugins/Runtime/Synthesis/Source/Synthesis/Private/SourceEffectEQ.cpp
@@ -6,14 +6,16 @@
 
 FSourceEffectEQ::FSourceEffectEQ()
 	: SampleRate(0)
+	, NumChannels(0)
 {
 	FMemory::Memzero((void*)InAudioFrame, sizeof(float)*2);
 	FMemory::Memzero((void*)OutAudioFrame, sizeof(float)*2);
 }
 
-void FSourceEffectEQ::Init(const FSoundEffectSourceInitData& InSampleRate)
+void FSourceEffectEQ::Init(const FSoundEffectSourceInitData& InitData)
 {
-	SampleRate = InSampleRate.SampleRate;
+	SampleRate = InitData.SampleRate;
+	NumChannels = InitData.NumSourceChannels;
 }
 
 void FSourceEffectEQ::OnPresetChanged() 
@@ -32,7 +34,7 @@ void FSourceEffectEQ::OnPresetChanged()
 		// Now initialize the new filters
 		for (; i < Filters.Num(); ++i)
 		{
-			Filters[i].Init(SampleRate, 2, Audio::EBiquadFilter::ParametricEQ);
+			Filters[i].Init(SampleRate, NumChannels, Audio::EBiquadFilter::ParametricEQ);
 		}
 	}
 	else
@@ -54,58 +56,24 @@ void FSourceEffectEQ::OnPresetChanged()
 		Filters[i].SetEnabled(EQBandSetting.bEnabled);
 		Filters[i].SetParams(Audio::EBiquadFilter::ParametricEQ, FMath::Max(20.0f, EQBandSetting.Frequency), EQBandSetting.Bandwidth, EQBandSetting.GainDb);
 	}
-
 }
 
-void FSourceEffectEQ::ProcessAudio(const FSoundEffectSourceInputData& InData, FSoundEffectSourceOutputData& OutData)
+void FSourceEffectEQ::ProcessAudio(const FSoundEffectSourceInputData& InData, float* OutAudioBufferData)
 {
-	const int32 NumChannels = InData.AudioFrame.Num();
+	float* InAudioBufferData = InData.InputSourceEffectBufferPtr;
+
 	if (Filters.Num() == 0)
 	{
-		for (int32 Channel = 0; Channel < NumChannels; ++Channel)
-		{
-			OutData.AudioFrame[Channel] = InData.AudioFrame[Channel];
-		}
+		FMemory::Memcpy(OutAudioBufferData, InAudioBufferData, sizeof(float)*InData.NumSamples);
+		return;
 	}
 
-	if (NumChannels == 2)
+	for (int32 FilterIndex = 0; FilterIndex < Filters.Num(); ++FilterIndex)
 	{
-		for (int32 Channel = 0; Channel < 2; ++Channel)
-		{
-			InAudioFrame[Channel] = InData.AudioFrame[Channel];
-		}
-
-		for (int32 i = 0; i < Filters.Num(); ++i)
-		{
-			Filters[i].ProcessAudioFrame(InAudioFrame, OutData.AudioFrame.GetData());
-
-			for (int32 Channel = 0; Channel < 2; ++Channel)
-			{
-				InAudioFrame[Channel] = OutData.AudioFrame[Channel];
-			}
-		}
+		Filters[FilterIndex].ProcessAudio(InAudioBufferData, InData.NumSamples, InAudioBufferData);
 	}
-	else
-	{
-		InAudioFrame[0] = 0.5f * InData.AudioFrame[0];
-		InAudioFrame[1] = InAudioFrame[0];
-		OutAudioFrame[0] = 0.0f;
-		OutAudioFrame[1] = 0.0f;
-
-		for (int32 i = 0; i < Filters.Num(); ++i)
-		{
-			Filters[i].ProcessAudioFrame(InAudioFrame, OutAudioFrame);
-
-			for (int32 Channel = 0; Channel < 2; ++Channel)
-			{
-				InAudioFrame[Channel] = OutAudioFrame[Channel];
-			}
-		}
-
-		OutData.AudioFrame[0] = OutAudioFrame[0];
-	}
-
 
+	FMemory::Memcpy(OutAudioBufferData, InAudioBufferData, sizeof(float)*InData.NumSamples);
 }
 
 void USourceEffectEQPreset::SetSettings(const FSourceEffectEQSettings& InSettings)
diff --git a/Engine/Plugins/Runtime/Synthesis/Source/Synthesis/Private/SourceEffectEnvelopeFollower.cpp b/Engine/Plugins/Runtime/Synthesis/Source/Synthesis/Private/SourceEffectEnvelopeFollower.cpp
--- a/Engine/Plugins/Runtime/Synthesis/Source/Synthesis/Private/SourceEffectEnvelopeFollower.cpp
+++ b/Engine/Plugins/Runtime/Synthesis/Source/Synthesis/Private/SourceEffectEnvelopeFollower.cpp
@@ -222,6 +222,7 @@ void FSourceEffectEnvelopeFollower::Init(const FSoundEffectSourceInitData& InitD
 	bIsActive = true;
 	EnvelopeFollower.Init(InitData.SampleRate);
 	CurrentEnvelopeValue = 0.0f;
+	NumChannels = InitData.NumSourceChannels;
 }
 
 FSourceEffectEnvelopeFollower::~FSourceEffectEnvelopeFollower()
@@ -239,28 +240,31 @@ void FSourceEffectEnvelopeFollower::OnPresetChanged()
 	EnvelopeFollower.SetMode((Audio::EPeakMode::Type)Settings.PeakMode);
 }
 
-void FSourceEffectEnvelopeFollower::ProcessAudio(const FSoundEffectSourceInputData& InData, FSoundEffectSourceOutputData& OutData)
+void FSourceEffectEnvelopeFollower::ProcessAudio(const FSoundEffectSourceInputData& InData, float* OutAudioBufferData)
 {
-	float SampleValue = 0.0f;
-	for (int32 i = 0; i < InData.AudioFrame.Num(); ++i)
-	{
-		SampleValue += InData.AudioFrame[i];
-		OutData.AudioFrame[i] = InData.AudioFrame[i];
-	}
+	// Memcpy the output (this is a "pass through" effect)
+	FMemory::Memcpy(OutAudioBufferData, InData.InputSourceEffectBufferPtr, sizeof(float)*InData.NumSamples);
 
-	if (InData.AudioFrame.Num() == 2)
+	for (int32 SampleIndex = 0; SampleIndex < InData.NumSamples; SampleIndex += NumChannels)
 	{
-		SampleValue *= 0.5f;
-	}
+		float SampleValue = 0.0f;
+		for (int32 ChannelIndex = 0; ChannelIndex < NumChannels; ++ChannelIndex)
+		{
+			SampleValue += InData.InputSourceEffectBufferPtr[SampleIndex + ChannelIndex];
+		}
 
-	CurrentEnvelopeValue = EnvelopeFollower.ProcessAudio(SampleValue);
+		if (NumChannels == 2)
+		{
+			SampleValue *= 0.5f;
+		}
 
-	FrameCount = FrameCount & (FramesToNotify - 1);
-	if (FrameCount == 0)
-	{
-		SourceEffectEnvFollowerNotifier->UpdateEnvFollowerInstance(OwningPresetUniqueId, InstanceId, CurrentEnvelopeValue);
+		CurrentEnvelopeValue = EnvelopeFollower.ProcessAudio(SampleValue);
+
+		if ((FrameCount++ & (FramesToNotify - 1)) == 0)
+		{
+			SourceEffectEnvFollowerNotifier->UpdateEnvFollowerInstance(OwningPresetUniqueId, InstanceId, CurrentEnvelopeValue);
+		}
 	}
-	++FrameCount;
 }
 
 void USourceEffectEnvelopeFollowerPreset::SetSettings(const FSourceEffectEnvelopeFollowerSettings& InSettings)
diff --git a/Engine/Plugins/Runtime/Synthesis/Source/Synthesis/Private/SourceEffectFilter.cpp b/Engine/Plugins/Runtime/Synthesis/Source/Synthesis/Private/SourceEffectFilter.cpp
--- a/Engine/Plugins/Runtime/Synthesis/Source/Synthesis/Private/SourceEffectFilter.cpp
+++ b/Engine/Plugins/Runtime/Synthesis/Source/Synthesis/Private/SourceEffectFilter.cpp
@@ -17,10 +17,10 @@ FSourceEffectFilter::FSourceEffectFilter()
 void FSourceEffectFilter::Init(const FSoundEffectSourceInitData& InitData)
 {
 	bIsActive = true;
-
-	StateVariableFilter.Init(InitData.SampleRate, 2);
-	LadderFilter.Init(InitData.SampleRate, 2);
-	OnePoleFilter.Init(InitData.SampleRate, 2);
+	NumChannels = InitData.NumSourceChannels;
+	StateVariableFilter.Init(InitData.SampleRate, NumChannels);
+	LadderFilter.Init(InitData.SampleRate, NumChannels);
+	OnePoleFilter.Init(InitData.SampleRate, NumChannels);
 
 	UpdateFilter();
 }
@@ -86,18 +86,9 @@ void FSourceEffectFilter::OnPresetChanged()
 	UpdateFilter();
 }
 
-void FSourceEffectFilter::ProcessAudio(const FSoundEffectSourceInputData& InData, FSoundEffectSourceOutputData& OutData)
+void FSourceEffectFilter::ProcessAudio(const FSoundEffectSourceInputData& InData, float* OutAudioBufferData)
 {
-	if (InData.AudioFrame.Num() == 2)
-	{
-		CurrentFilter->ProcessAudio(InData.AudioFrame.GetData(), OutData.AudioFrame.GetData());
-	}
-	else
-	{
-		AudioInput[0] = InData.AudioFrame[0];
-		CurrentFilter->ProcessAudio(AudioInput, AudioOutput);
-		OutData.AudioFrame[0] = AudioOutput[0];
-	}
+	CurrentFilter->ProcessAudio(InData.InputSourceEffectBufferPtr, InData.NumSamples, OutAudioBufferData);
 }
 
 void USourceEffectFilterPreset::SetSettings(const FSourceEffectFilterSettings& InSettings)
diff --git a/Engine/Plugins/Runtime/Synthesis/Source/Synthesis/Private/SourceEffectFoldbackDistortion.cpp b/Engine/Plugins/Runtime/Synthesis/Source/Synthesis/Private/SourceEffectFoldbackDistortion.cpp
--- a/Engine/Plugins/Runtime/Synthesis/Source/Synthesis/Private/SourceEffectFoldbackDistortion.cpp
+++ b/Engine/Plugins/Runtime/Synthesis/Source/Synthesis/Private/SourceEffectFoldbackDistortion.cpp
@@ -5,7 +5,7 @@
 void FSourceEffectFoldbackDistortion::Init(const FSoundEffectSourceInitData& InitData)
 {
 	bIsActive = true;
-	FoldbackDistortion.Init(InitData.SampleRate);
+	FoldbackDistortion.Init(InitData.SampleRate, InitData.NumSourceChannels);
 }
 
 void FSourceEffectFoldbackDistortion::OnPresetChanged()
@@ -17,16 +17,9 @@ void FSourceEffectFoldbackDistortion::OnPresetChanged()
 	FoldbackDistortion.SetOutputGainDb(Settings.OutputGainDb);
 }
 
-void FSourceEffectFoldbackDistortion::ProcessAudio(const FSoundEffectSourceInputData& InData, FSoundEffectSourceOutputData& OutData)
+void FSourceEffectFoldbackDistortion::ProcessAudio(const FSoundEffectSourceInputData& InData, float* OutAudioBufferData)
 {
-	if (InData.AudioFrame.Num() == 2)
-	{
-		FoldbackDistortion.ProcessAudio(InData.AudioFrame[0], InData.AudioFrame[1], OutData.AudioFrame[0], OutData.AudioFrame[1]);
-	}
-	else
-	{
-		FoldbackDistortion.ProcessAudio(InData.AudioFrame[0], OutData.AudioFrame[0]);
-	}
+	FoldbackDistortion.ProcessAudio(InData.InputSourceEffectBufferPtr, InData.NumSamples, OutAudioBufferData);
 }
 
 void USourceEffectFoldbackDistortionPreset::SetSettings(const FSourceEffectFoldbackDistortionSettings& InSettings)
diff --git a/Engine/Plugins/Runtime/Synthesis/Source/Synthesis/Private/SourceEffectPanner.cpp b/Engine/Plugins/Runtime/Synthesis/Source/Synthesis/Private/SourceEffectPanner.cpp
--- a/Engine/Plugins/Runtime/Synthesis/Source/Synthesis/Private/SourceEffectPanner.cpp
+++ b/Engine/Plugins/Runtime/Synthesis/Source/Synthesis/Private/SourceEffectPanner.cpp
@@ -7,6 +7,7 @@ void FSourceEffectPanner::Init(const FSoundEffectSourceInitData& InitData)
 {
 	bIsActive = true;
 	PanValue = 0.0f;
+	NumChannels = InitData.NumSourceChannels;
 }
 
 void FSourceEffectPanner::OnPresetChanged()
@@ -20,29 +21,33 @@ void FSourceEffectPanner::OnPresetChanged()
 	PanValue *= 0.5f * PI;
 }
 
-void FSourceEffectPanner::ProcessAudio(const FSoundEffectSourceInputData& InData, FSoundEffectSourceOutputData& OutData)
+void FSourceEffectPanner::ProcessAudio(const FSoundEffectSourceInputData& InData, float* OutAudioBufferData)
 {
-	const int32 NumChannels = InData.AudioFrame.Num();
-
-	// If only 1 channel, do a simple passthrough
-	if (NumChannels == 1)
+	if (NumChannels != 2)
 	{
-		OutData.AudioFrame[0] = InData.AudioFrame[0];
+		FMemory::Memcpy(OutAudioBufferData, InData.InputSourceEffectBufferPtr, sizeof(float)*InData.NumSamples);
 	}
 	else
 	{
-		// Use the "cosine" equal power panning law to compute a smooth pan based off our parameter
 		float PanGains[2];
+		// Use the "cosine" equal power panning law to compute a smooth pan based off our parameter
 		FMath::SinCos(&PanGains[0], &PanGains[1], PanValue);
 
-		// Then simply scale the pan value output with the channel inputs. 
-		for (int32 i = 0; i < NumChannels; ++i)
+		// Clamp this to be between 0.0 and 1.0 since SinCos is fast and may have values greater than 1.0 or less than 0.0
+		for (int32 i = 0; i < 2; ++i)
 		{
-			// Clamp this to be between 0.0 and 1.0 since SinCos is fast and may have values greater than 1.0 or less than 0.0
-			float PanGain = FMath::Clamp(PanGains[i], 0.0f, 1.0f);
+			PanGains[i] = FMath::Clamp(PanGains[i], 0.0f, 1.0f);
+		}
 
-			// Simply scale the input sample with the pan value 
-			OutData.AudioFrame[i] = PanGain * InData.AudioFrame[i];
+		for (int32 SampleIndex = 0; SampleIndex < InData.NumSamples; SampleIndex += NumChannels)
+		{
+			// Then simply scale the pan value output with the channel inputs. Note we have to clamp the 
+			// the output of the SinCos since 
+			for (int32 i = 0; i < 2; ++i)
+			{
+				// Simply scale the input sample with the pan value 
+				OutAudioBufferData[SampleIndex + i] = PanGains[i] * InData.InputSourceEffectBufferPtr[SampleIndex + i];
+			}
 		}
 	}
 }
diff --git a/Engine/Plugins/Runtime/Synthesis/Source/Synthesis/Private/SourceEffectPhaser.cpp b/Engine/Plugins/Runtime/Synthesis/Source/Synthesis/Private/SourceEffectPhaser.cpp
--- a/Engine/Plugins/Runtime/Synthesis/Source/Synthesis/Private/SourceEffectPhaser.cpp
+++ b/Engine/Plugins/Runtime/Synthesis/Source/Synthesis/Private/SourceEffectPhaser.cpp
@@ -6,7 +6,7 @@ void FSourceEffectPhaser::Init(const FSoundEffectSourceInitData& InitData)
 {
 	bIsActive = true;
 
-	Phaser.Init(InitData.SampleRate);
+	Phaser.Init(InitData.SampleRate, InitData.NumSourceChannels);
 }
 
 void FSourceEffectPhaser::OnPresetChanged()
@@ -21,22 +21,9 @@ void FSourceEffectPhaser::OnPresetChanged()
 	Phaser.SetLFOType((Audio::ELFO::Type)Settings.LFOType);
 }
 
-void FSourceEffectPhaser::ProcessAudio(const FSoundEffectSourceInputData& InData, FSoundEffectSourceOutputData& OutData)
+void FSourceEffectPhaser::ProcessAudio(const FSoundEffectSourceInputData& InData, float* OutAudioBufferData)
 {
-	if (InData.AudioFrame.Num() == 2)
-	{
-		Phaser.ProcessAudio(InData.AudioFrame.GetData(), OutData.AudioFrame.GetData());
-	}
-	else
-	{
-		float InFrame[2];
-		InFrame[0] = 0.5f * InData.AudioFrame[0];
-		InFrame[1] = InFrame[0];
-
-		float OutFrame[2];
-		Phaser.ProcessAudio(InFrame, OutFrame);
-		OutData.AudioFrame[0] = 0.5f * (OutFrame[0] + OutFrame[1]);
-	}
+	Phaser.ProcessAudio(InData.InputSourceEffectBufferPtr, InData.NumSamples, OutAudioBufferData);
 }
 
 void USourceEffectPhaserPreset::SetSettings(const FSourceEffectPhaserSettings& InSettings)
diff --git a/Engine/Plugins/Runtime/Synthesis/Source/Synthesis/Private/SourceEffectRingModulation.cpp b/Engine/Plugins/Runtime/Synthesis/Source/Synthesis/Private/SourceEffectRingModulation.cpp
--- a/Engine/Plugins/Runtime/Synthesis/Source/Synthesis/Private/SourceEffectRingModulation.cpp
+++ b/Engine/Plugins/Runtime/Synthesis/Source/Synthesis/Private/SourceEffectRingModulation.cpp
@@ -5,7 +5,7 @@
 void FSourceEffectRingModulation::Init(const FSoundEffectSourceInitData& InitData)
 {
 	bIsActive = true;
-	RingModulation.Init(InitData.SampleRate);
+	RingModulation.Init(InitData.SampleRate, InitData.NumSourceChannels);
 }
 
 void FSourceEffectRingModulation::OnPresetChanged()
@@ -34,21 +34,13 @@ void FSourceEffectRingModulation::OnPresetChanged()
 
 	RingModulation.SetModulationDepth(Settings.Depth);
 	RingModulation.SetModulationFrequency(Settings.Frequency);
+	RingModulation.SetDryLevel(Settings.DryLevel);
+	RingModulation.SetWetLevel(Settings.WetLevel);
 }
 
-void FSourceEffectRingModulation::ProcessAudio(const FSoundEffectSourceInputData& InData, FSoundEffectSourceOutputData& OutData)
+void FSourceEffectRingModulation::ProcessAudio(const FSoundEffectSourceInputData& InData, float* OutAudioBufferData)
 {
-	if (InData.AudioFrame.Num() == 2)
-	{
-		RingModulation.ProcessAudio(InData.AudioFrame[0], InData.AudioFrame[1], OutData.AudioFrame[0], OutData.AudioFrame[1]);
-	}
-	else
-	{
-		float OutLeft = 0.0f;
-		float OutRight = 0.0f;
-		RingModulation.ProcessAudio(InData.AudioFrame[0], 0.0f, OutLeft, OutRight);
-		OutData.AudioFrame[0] = OutLeft;
-	}
+	RingModulation.ProcessAudio(InData.InputSourceEffectBufferPtr, InData.NumSamples, OutAudioBufferData);
 }
 
 void USourceEffectRingModulationPreset::SetSettings(const FSourceEffectRingModulationSettings& InSettings)
diff --git a/Engine/Plugins/Runtime/Synthesis/Source/Synthesis/Private/SourceEffectSimpleDelay.cpp b/Engine/Plugins/Runtime/Synthesis/Source/Synthesis/Private/SourceEffectSimpleDelay.cpp
--- a/Engine/Plugins/Runtime/Synthesis/Source/Synthesis/Private/SourceEffectSimpleDelay.cpp
+++ b/Engine/Plugins/Runtime/Synthesis/Source/Synthesis/Private/SourceEffectSimpleDelay.cpp
@@ -6,10 +6,12 @@ void FSourceEffectSimpleDelay::Init(const FSoundEffectSourceInitData& InitData)
 {
 	bIsActive = true;
 
-	for (int32 i = 0; i < 2; ++i)
+	FeedbackSamples.AddZeroed(InitData.NumSourceChannels);
+	Delays.AddDefaulted(InitData.NumSourceChannels);
+
+	for (int32 i = 0; i < InitData.NumSourceChannels; ++i)
 	{
-		DelayLines[i].Init(InitData.SampleRate, 2.0f);
-		FeedbackSamples[i] = 0.0f;
+		Delays[i].Init(InitData.SampleRate, 2.0f);
 	}
 }
 
@@ -22,36 +24,41 @@ void FSourceEffectSimpleDelay::OnPresetChanged()
 	// If we are manually setting the delay, lets set it now on the delay line
 	if (!SettingsCopy.bDelayBasedOnDistance)
 	{
-		for (int32 i = 0; i < 2; ++i)
+		for (Audio::FDelay& Delay : Delays)
 		{
-			DelayLines[i].SetEasedDelayMsec(SettingsCopy.DelayAmount * 1000.0f);
+			Delay.SetEasedDelayMsec(SettingsCopy.DelayAmount * 1000.0f);
 		}
 	}
 }
 
-void FSourceEffectSimpleDelay::ProcessAudio(const FSoundEffectSourceInputData& InData, FSoundEffectSourceOutputData& OutData)
+void FSourceEffectSimpleDelay::ProcessAudio(const FSoundEffectSourceInputData& InData, float* OutAudioBufferData)
 {
-	const int32 NumChannels = InData.AudioFrame.Num();
-
 	if (SettingsCopy.bDelayBasedOnDistance)
 	{
-		const float DistanceMeters = InData.Distance * 100.0f;
+		const float DistanceMeters = InData.SpatParams.Distance * 0.01f;
 		const float DelayAmountMsec = 1000.0f * DistanceMeters / SettingsCopy.SpeedOfSound;
 
-		for (int32 i = 0; i < NumChannels; ++i)
+		for (Audio::FDelay& Delay : Delays)
 		{
-			DelayLines[i].SetDelayMsec(DelayAmountMsec * 1000.0f);
+			Delay.SetEasedDelayMsec(DelayAmountMsec);
 		}
 	}
 
-	const float* InAudioBufferFrame = InData.AudioFrame.GetData();
-	float* OutAudioBufferFrame = OutData.AudioFrame.GetData();
-
-	for (int32 i = 0; i < NumChannels; ++i)
+	int32 NumChannels = Delays.Num();
+	for (int32 SampleIndex = 0; SampleIndex < InData.NumSamples; SampleIndex += NumChannels)
 	{
-		OutAudioBufferFrame[i] = DelayLines[i].ProcessAudio(InAudioBufferFrame[i] + FeedbackSamples[i] * SettingsCopy.Feedback);
-		FeedbackSamples[i] = OutAudioBufferFrame[i];
+		for (int32 ChannelIndex = 0; ChannelIndex < NumChannels; ++ChannelIndex)
+		{
+			Audio::FDelay& Delay = Delays[ChannelIndex];
+
+			const float DrySample = InData.InputSourceEffectBufferPtr[SampleIndex + ChannelIndex] * SettingsCopy.DryAmount;
+			const float WetSample = SettingsCopy.WetAmount * Delay.ProcessAudioSample(InData.InputSourceEffectBufferPtr[SampleIndex + ChannelIndex] + FeedbackSamples[ChannelIndex] * SettingsCopy.Feedback);
+
+			OutAudioBufferData[SampleIndex + ChannelIndex] = DrySample + WetSample;
+			FeedbackSamples[ChannelIndex] = OutAudioBufferData[SampleIndex + ChannelIndex];
+		}
 	}
+
 }
 
 void USourceEffectSimpleDelayPreset::SetSettings(const FSourceEffectSimpleDelaySettings& InSettings)
diff --git a/Engine/Plugins/Runtime/Synthesis/Source/Synthesis/Private/SourceEffectStereoDelay.cpp b/Engine/Plugins/Runtime/Synthesis/Source/Synthesis/Private/SourceEffectStereoDelay.cpp
--- a/Engine/Plugins/Runtime/Synthesis/Source/Synthesis/Private/SourceEffectStereoDelay.cpp
+++ b/Engine/Plugins/Runtime/Synthesis/Source/Synthesis/Private/SourceEffectStereoDelay.cpp
@@ -6,7 +6,7 @@
 void FSourceEffectStereoDelay::Init(const FSoundEffectSourceInitData& InitData)
 {
 	bIsActive = true;
-	DelayStereo.Init(InitData.SampleRate);
+	DelayStereo.Init(InitData.SampleRate, InitData.NumSourceChannels);
 }
 
 void FSourceEffectStereoDelay::OnPresetChanged()
@@ -20,19 +20,9 @@ void FSourceEffectStereoDelay::OnPresetChanged()
 	DelayStereo.SetMode((Audio::EStereoDelayMode::Type)Settings.DelayMode);
 }
 
-void FSourceEffectStereoDelay::ProcessAudio(const FSoundEffectSourceInputData& InData, FSoundEffectSourceOutputData& OutData)
+void FSourceEffectStereoDelay::ProcessAudio(const FSoundEffectSourceInputData& InData, float* OutAudioBufferData)
 {
-	if (InData.AudioFrame.Num() == 2)
-	{
-		DelayStereo.ProcessAudio(InData.AudioFrame[0], InData.AudioFrame[1], OutData.AudioFrame[0], OutData.AudioFrame[1]);
-	}
-	else
-	{
-		float OutLeft = 0.0f;
-		float OutRight = 0.0f;
-		DelayStereo.ProcessAudio(InData.AudioFrame[0], InData.AudioFrame[0], OutLeft, OutRight);
-		OutData.AudioFrame[0] = 0.5f * (OutLeft + OutRight);
-	}
+	DelayStereo.ProcessAudio(InData.InputSourceEffectBufferPtr, InData.NumSamples, OutAudioBufferData);
 }
 
 void USourceEffectStereoDelayPreset::SetSettings(const FSourceEffectStereoDelaySettings& InSettings)
diff --git a/Engine/Plugins/Runtime/Synthesis/Source/Synthesis/Private/SourceEffectWaveShaper.cpp b/Engine/Plugins/Runtime/Synthesis/Source/Synthesis/Private/SourceEffectWaveShaper.cpp
--- a/Engine/Plugins/Runtime/Synthesis/Source/Synthesis/Private/SourceEffectWaveShaper.cpp
+++ b/Engine/Plugins/Runtime/Synthesis/Source/Synthesis/Private/SourceEffectWaveShaper.cpp
@@ -6,6 +6,7 @@ void FSourceEffectWaveShaper::Init(const FSoundEffectSourceInitData& InitData)
 {
 	bIsActive = true;
 	WaveShaper.Init(InitData.SampleRate);
+	NumChannels = InitData.NumSourceChannels;
 }
 
 void FSourceEffectWaveShaper::OnPresetChanged()
@@ -16,18 +17,14 @@ void FSourceEffectWaveShaper::OnPresetChanged()
 	WaveShaper.SetOutputGainDb(Settings.OutputGainDb);
 }
 
-void FSourceEffectWaveShaper::ProcessAudio(const FSoundEffectSourceInputData& InData, FSoundEffectSourceOutputData& OutData)
+void FSourceEffectWaveShaper::ProcessAudio(const FSoundEffectSourceInputData& InData, float* OutAudioBufferData)
 {
-	if (InData.AudioFrame.Num() == 2)
+	for (int32 SampleIndex = 0; SampleIndex < InData.NumSamples; SampleIndex += NumChannels)
 	{
-		check(OutData.AudioFrame.Num() == 2);
-
-		WaveShaper.ProcessAudio(InData.AudioFrame[0], OutData.AudioFrame[0]);
-		WaveShaper.ProcessAudio(InData.AudioFrame[1], OutData.AudioFrame[1]);
-	}
-	else
-	{
-		WaveShaper.ProcessAudio(InData.AudioFrame[0], OutData.AudioFrame[0]);
+		for (int32 ChannelIndex = 0; ChannelIndex < NumChannels; ++ChannelIndex)
+		{
+			WaveShaper.ProcessAudio(InData.InputSourceEffectBufferPtr[SampleIndex + ChannelIndex], OutAudioBufferData[SampleIndex + ChannelIndex]);
+		}
 	}
 }
 
diff --git a/Engine/Plugins/Runtime/Synthesis/Source/Synthesis/Private/SubmixEffectDelay.cpp b/Engine/Plugins/Runtime/Synthesis/Source/Synthesis/Private/SubmixEffectDelay.cpp
--- a/Engine/Plugins/Runtime/Synthesis/Source/Synthesis/Private/SubmixEffectDelay.cpp
+++ b/Engine/Plugins/Runtime/Synthesis/Source/Synthesis/Private/SubmixEffectDelay.cpp
@@ -47,7 +47,7 @@ void FSubmixEffectDelay::OnProcessAudio(const FSoundEffectSubmixInputData& InDat
 		for (int32 DelayIndex = 0; DelayIndex < NumDelays; DelayIndex++)
 		{
 			const int32 SampleIndex = OutputBufferIndex + DelayIndex;
-			DelaysPtr[DelayIndex].ProcessAudio(&InBuffer[SampleIndex], &OutBuffer[SampleIndex]);
+			OutBuffer[SampleIndex] = DelaysPtr[DelayIndex].ProcessAudioSample(InBuffer[SampleIndex]);
 			DelayLines[DelayIndex].SetDelayMsec(InterpolationInfo.GetValue());
 		}
 	}
diff --git a/Engine/Plugins/Runtime/Synthesis/Source/Synthesis/Public/EpicSynth1.h b/Engine/Plugins/Runtime/Synthesis/Source/Synthesis/Public/EpicSynth1.h
--- a/Engine/Plugins/Runtime/Synthesis/Source/Synthesis/Public/EpicSynth1.h
+++ b/Engine/Plugins/Runtime/Synthesis/Source/Synthesis/Public/EpicSynth1.h
@@ -200,7 +200,7 @@ namespace Audio
 		bool SetEnablePatch(const FPatchId PatchId, bool bIsEnabled);
 
 		// Generate the next frame of audio
-		void Generate(float& OutLeft, float& OutRight);
+		void GenerateFrame(float* OutFrame);
 
 	protected:
 
diff --git a/Engine/Source/Editor/AudioEditor/Private/AssetTypeActions/AssetTypeActions_SoundEffectPreset.h b/Engine/Source/Editor/AudioEditor/Private/AssetTypeActions/AssetTypeActions_SoundEffectPreset.h
--- a/Engine/Source/Editor/AudioEditor/Private/AssetTypeActions/AssetTypeActions_SoundEffectPreset.h
+++ b/Engine/Source/Editor/AudioEditor/Private/AssetTypeActions/AssetTypeActions_SoundEffectPreset.h
@@ -19,7 +19,7 @@ class FAssetTypeActions_SoundEffectSubmixPreset : public FAssetTypeActions_Base
 public:
 	// IAssetTypeActions Implementation
 	virtual FText GetName() const override { return NSLOCTEXT("AssetTypeActions", "AssetTypeActions_SoundSubmixPreset", "Submix Effect preset"); }
-	virtual FColor GetTypeColor() const override { return FColor(175, 255, 0); }
+	virtual FColor GetTypeColor() const override { return FColor(99, 63, 56); }
 	virtual UClass* GetSupportedClass() const override { return USoundEffectSubmixPreset::StaticClass(); }
 	virtual uint32 GetCategories() override { return EAssetTypeCategories::Sounds; }
 };
@@ -29,7 +29,7 @@ class FAssetTypeActions_SoundEffectSourcePreset : public FAssetTypeActions_Base
 public:
 	// IAssetTypeActions Implementation
 	virtual FText GetName() const override { return NSLOCTEXT("AssetTypeActions", "AssetTypeActions_SoundSourcePreset", "Source Effect Preset"); }
-	virtual FColor GetTypeColor() const override { return FColor(175, 255, 0); }
+	virtual FColor GetTypeColor() const override { return FColor(72, 185, 187); }
 	virtual UClass* GetSupportedClass() const override { return USoundEffectSourcePreset::StaticClass(); }
 	virtual uint32 GetCategories() override { return EAssetTypeCategories::Sounds; }
 };
@@ -39,7 +39,7 @@ class FAssetTypeActions_SoundEffectSourcePresetChain : public FAssetTypeActions_
 public:
 	// IAssetTypeActions Implementation
 	virtual FText GetName() const override { return NSLOCTEXT("AssetTypeActions", "AssetTypeActions_SoundSourcePresetChain", "Source Effect Preset Chain"); }
-	virtual FColor GetTypeColor() const override { return FColor(175, 200, 100); }
+	virtual FColor GetTypeColor() const override { return FColor(51, 107, 142); }
 	virtual UClass* GetSupportedClass() const override { return USoundEffectSourcePresetChain::StaticClass(); }
 	virtual uint32 GetCategories() override { return EAssetTypeCategories::Sounds; }
 };
diff --git a/Engine/Source/Editor/AudioEditor/Private/AssetTypeActions/AssetTypeActions_SoundSourceBus.h b/Engine/Source/Editor/AudioEditor/Private/AssetTypeActions/AssetTypeActions_SoundSourceBus.h
--- a/Engine/Source/Editor/AudioEditor/Private/AssetTypeActions/AssetTypeActions_SoundSourceBus.h
+++ b/Engine/Source/Editor/AudioEditor/Private/AssetTypeActions/AssetTypeActions_SoundSourceBus.h
@@ -19,5 +19,6 @@ public:
 	virtual UClass* GetSupportedClass() const override;
 	virtual bool CanFilter() override { return true; }
 	virtual bool IsImportedAsset() const override { return false; }
+	virtual TSharedPtr<class SWidget> GetThumbnailOverlay(const FAssetData& AssetData) const override { return nullptr; }
 	//~ End IAssetTypeActions Implementation
 };
diff --git a/Engine/Source/Editor/AudioEditor/Private/AudioEditorModule.cpp b/Engine/Source/Editor/AudioEditor/Private/AudioEditorModule.cpp
--- a/Engine/Source/Editor/AudioEditor/Private/AudioEditorModule.cpp
+++ b/Engine/Source/Editor/AudioEditor/Private/AudioEditorModule.cpp
@@ -34,11 +34,29 @@
 #include "AssetTypeActions/AssetTypeActions_SoundSourceBus.h"
 #include "Utils.h"
 #include "UObject/UObjectIterator.h"
+#include "Styling/SlateStyle.h"
+#include "Styling/SlateStyleRegistry.h"
 
 const FName AudioEditorAppIdentifier = FName(TEXT("AudioEditorApp"));
 
 DEFINE_LOG_CATEGORY(LogAudioEditor);
 
+class FSlateStyleSet;
+struct FGraphPanelPinConnectionFactory;
+
+
+// Preprocessor macro to make defining audio icons simple...
+
+// CLASS_NAME - name of the class to make the icon for
+// ICON_NAME - base-name of the icon to use. Not necessarily based off class name
+#define SET_AUDIO_ICON(CLASS_NAME, ICON_NAME) \
+		AudioStyleSet->Set( *FString::Printf(TEXT("ClassIcon.%s"), TEXT(#CLASS_NAME)), new FSlateImageBrush(FPaths::EngineContentDir() / FString::Printf(TEXT("Editor/Slate/Icons/AssetIcons/%s_16x.png"), TEXT(#ICON_NAME)), Icon16)); \
+		AudioStyleSet->Set( *FString::Printf(TEXT("ClassThumbnail.%s"), TEXT(#CLASS_NAME)), new FSlateImageBrush(FPaths::EngineContentDir() / FString::Printf(TEXT("Editor/Slate/Icons/AssetIcons/%s_64x.png"), TEXT(#ICON_NAME)), Icon64)); 
+
+// Simpler version of SET_AUDIO_ICON, assumes same name of icon png and class name
+#define SET_AUDIO_ICON_SIMPLE(CLASS_NAME) SET_AUDIO_ICON(CLASS_NAME, CLASS_NAME)
+
+
 class FAudioEditorModule : public IAudioEditorModule
 {
 public:
@@ -64,6 +82,9 @@ public:
 
 		// Create reimport handler for surround sound waves
 		UReimportSoundSurroundFactory::StaticClass();
+
+		SetupIcons();
+
 	}
 
 	virtual void ShutdownModule() override
@@ -253,6 +274,38 @@ public:
 
 private:
 
+	void SetupIcons()
+	{
+		// Setup icon sizes
+		FVector2D Icon16 = FVector2D(16.0f, 16.0f);
+		FVector2D Icon64 = FVector2D(64.0f, 64.0f);
+
+		// Setup the audio icon root
+		FString AudioIconRoot = TEXT("Icons/AssetIcons/");
+
+		// Create style set for audio asset icons
+		AudioStyleSet = MakeShareable(new FSlateStyleSet("AudioStyleSet"));
+
+		SET_AUDIO_ICON_SIMPLE(SoundAttenuation);
+		SET_AUDIO_ICON_SIMPLE(AmbientSound);
+		SET_AUDIO_ICON_SIMPLE(SoundClass);
+		SET_AUDIO_ICON_SIMPLE(SoundConcurrency);
+		SET_AUDIO_ICON_SIMPLE(SoundCue);
+		SET_AUDIO_ICON_SIMPLE(SoundMix);
+		SET_AUDIO_ICON_SIMPLE(AudioVolume);
+		SET_AUDIO_ICON_SIMPLE(SoundSourceBus);
+		SET_AUDIO_ICON_SIMPLE(SoundSubmix);
+		SET_AUDIO_ICON_SIMPLE(ReverbEffect);
+
+		SET_AUDIO_ICON(SoundEffectSubmixPreset, SubmixEffectPreset);
+		SET_AUDIO_ICON(SoundEffectSourcePreset, SourceEffectPreset);
+		SET_AUDIO_ICON(SoundEffectSourcePresetChain, SourceEffectPresetChain_1);
+		SET_AUDIO_ICON(ModularSynthPresetBank, SoundGenericIcon_2);
+		
+
+		FSlateStyleRegistry::RegisterSlateStyle(*AudioStyleSet.Get());
+	}
+
 	struct FExtensibilityManagers
 	{
 		TSharedPtr<FExtensibilityManager> MenuExtensibilityManager;
@@ -276,7 +329,8 @@ private:
 	FExtensibilityManagers SoundSubmixExtensibility;
 	TArray<TSharedPtr<ISoundWaveAssetActionExtensions>> SoundWaveAssetActionExtensions;
 	TSet<USoundEffectPreset*> RegisteredActions;
-	TSharedPtr<struct FGraphPanelPinConnectionFactory> SoundCueGraphConnectionFactory;
+	TSharedPtr<FGraphPanelPinConnectionFactory> SoundCueGraphConnectionFactory;
+	TSharedPtr<FSlateStyleSet> AudioStyleSet;
 
 };
 
diff --git a/Engine/Source/Editor/AudioEditor/Private/Factories/SoundSourceBusFactory.cpp b/Engine/Source/Editor/AudioEditor/Private/Factories/SoundSourceBusFactory.cpp
--- a/Engine/Source/Editor/AudioEditor/Private/Factories/SoundSourceBusFactory.cpp
+++ b/Engine/Source/Editor/AudioEditor/Private/Factories/SoundSourceBusFactory.cpp
@@ -19,12 +19,6 @@ UObject* USoundSourceBusFactory::FactoryCreateNew(UClass* InClass, UObject* InPa
 {
 	USoundSourceBus* SoundSourceBus = NewObject<USoundSourceBus>(InParent, InName, Flags);
 
-	FAudioDeviceManager* AudioDeviceManager = GEngine ? GEngine->GetAudioDeviceManager() : nullptr;
-	if (AudioDeviceManager)
-	{
-		AudioDeviceManager->InitSoundSubmixes();
-	}
-
 	return SoundSourceBus;
 }
 
diff --git a/Engine/Source/Editor/EditorStyle/Private/SlateEditorStyle.cpp b/Engine/Source/Editor/EditorStyle/Private/SlateEditorStyle.cpp
--- a/Engine/Source/Editor/EditorStyle/Private/SlateEditorStyle.cpp
+++ b/Engine/Source/Editor/EditorStyle/Private/SlateEditorStyle.cpp
@@ -6117,7 +6117,6 @@ void FSlateEditorStyle::FStyle::SetupClassIconsAndThumbnails()
 			TEXT("AimOffsetBlendSpace"),
 			TEXT("AimOffsetBlendSpace1D"),
 			TEXT("AIPerceptionComponent"),
-			TEXT("AmbientSound"),
 			TEXT("AnimationModifier"),		
 			TEXT("AnimBlueprint"),
 			TEXT("AnimComposite"),
@@ -6125,7 +6124,6 @@ void FSlateEditorStyle::FStyle::SetupClassIconsAndThumbnails()
 			TEXT("AnimSequence"),
 			TEXT("ApplicationLifecycleComponent"),
 			TEXT("AtmosphericFog"),
-			TEXT("AudioVolume"),
 			TEXT("BehaviorTree"),
 			TEXT("BlackboardData"),
 			TEXT("BlendSpace"),
@@ -6224,7 +6222,6 @@ void FSlateEditorStyle::FStyle::SetupClassIconsAndThumbnails()
 			TEXT("RadialForceActor"),
 			TEXT("RadialForceComponent"),
 			TEXT("ReflectionCapture"),
-			TEXT("ReverbEffect"),
 			TEXT("RotatingMovementComponent"),
 			TEXT("SceneCapture2D"),
 			TEXT("SceneCaptureCube"),
@@ -6236,11 +6233,6 @@ void FSlateEditorStyle::FStyle::SetupClassIconsAndThumbnails()
 			TEXT("SlateBrushAsset"),
 			TEXT("SlateWidgetStyleAsset"),
 			TEXT("StringTable"),
-			TEXT("SoundAttenuation"),
-			TEXT("SoundClass"),
-			TEXT("SoundConcurrency"),
-			TEXT("SoundCue"),
-			TEXT("SoundMix"),
 			TEXT("SphereReflectionCapture"),
 			TEXT("SpotLight"),
 			TEXT("SpotLightMovable"),
diff --git a/Engine/Source/Editor/UnrealEd/Classes/ThumbnailRendering/SoundWaveThumbnailRenderer.h b/Engine/Source/Editor/UnrealEd/Classes/ThumbnailRendering/SoundWaveThumbnailRenderer.h
--- a/Engine/Source/Editor/UnrealEd/Classes/ThumbnailRendering/SoundWaveThumbnailRenderer.h
+++ b/Engine/Source/Editor/UnrealEd/Classes/ThumbnailRendering/SoundWaveThumbnailRenderer.h
@@ -21,6 +21,7 @@ class USoundWaveThumbnailRenderer : public UThumbnailRenderer
 	GENERATED_UCLASS_BODY()
 
 	// Begin UThumbnailRenderer Object
+	virtual bool CanVisualizeAsset(UObject* Object);
 	virtual void Draw(UObject* Object, int32 X, int32 Y, uint32 Width, uint32 Height, FRenderTarget*, FCanvas* Canvas) override;
 	virtual bool AllowsRealtimeThumbnails(UObject* Object) const override;
 	// End UThumbnailRenderer Object
diff --git a/Engine/Source/Editor/UnrealEd/Private/ThumbnailRendering/SoundWaveThumbnailRenderer.cpp b/Engine/Source/Editor/UnrealEd/Private/ThumbnailRendering/SoundWaveThumbnailRenderer.cpp
--- a/Engine/Source/Editor/UnrealEd/Private/ThumbnailRendering/SoundWaveThumbnailRenderer.cpp
+++ b/Engine/Source/Editor/UnrealEd/Private/ThumbnailRendering/SoundWaveThumbnailRenderer.cpp
@@ -4,12 +4,23 @@
 #include "CanvasItem.h"
 #include "Audio.h"
 #include "Sound/SoundWave.h"
-
+#include "Sound/SoundSourceBus.h"
+ 
 USoundWaveThumbnailRenderer::USoundWaveThumbnailRenderer(const FObjectInitializer& ObjectInitializer)
 	: Super(ObjectInitializer)
 {
 }
 
+bool USoundWaveThumbnailRenderer::CanVisualizeAsset(UObject* Object)
+{
+	USoundSourceBus* SourceBus = Cast<USoundSourceBus>(Object);
+	if (SourceBus)
+	{
+		return false;
+	}
+	return true;
+}
+
 void USoundWaveThumbnailRenderer::Draw(UObject* Object, int32 X, int32 Y, uint32 Width, uint32 Height, FRenderTarget*, FCanvas* Canvas)
 {
 	USoundWave* SoundWave = Cast<USoundWave>(Object);
diff --git a/Engine/Source/Runtime/AudioMixer/Classes/SubmixEffects/AudioMixerSubmixEffectReverb.h b/Engine/Source/Runtime/AudioMixer/Classes/SubmixEffects/AudioMixerSubmixEffectReverb.h
--- a/Engine/Source/Runtime/AudioMixer/Classes/SubmixEffects/AudioMixerSubmixEffectReverb.h
+++ b/Engine/Source/Runtime/AudioMixer/Classes/SubmixEffects/AudioMixerSubmixEffectReverb.h
@@ -60,6 +60,10 @@ struct AUDIOMIXER_API FSubmixEffectReverbSettings
 	UPROPERTY(EditAnywhere, BlueprintReadWrite, Category = Reverb)
 	float WetLevel;
 
+	// Overall drylevel of the reverb effect
+	UPROPERTY(EditAnywhere, BlueprintReadWrite, Category = Reverb)
+	float DryLevel;
+
 	FSubmixEffectReverbSettings()
 	: Density(0.85f)
 	, Diffusion(0.85f)
@@ -73,6 +77,7 @@ struct AUDIOMIXER_API FSubmixEffectReverbSettings
 	, LateDelay(0.1f)
 	, AirAbsorptionGainHF(0.994f)
 	, WetLevel(0.3f)
+	, DryLevel(0.0f)
 	{
 	}
 };
@@ -109,6 +114,9 @@ private:
 	// Curve which maps old reverb times to new decay value
 	FRichCurve DecayCurve;
 
+	// Cache the dry level
+	float DryLevel;
+
 	bool bIsEnabled;
 };
 
@@ -124,7 +132,7 @@ public:
 	void SetSettings(const FSubmixEffectReverbSettings& InSettings);
 
 	UFUNCTION(BlueprintCallable, Category = "Audio|Effects")
-	void SetSettingsWithReverbEffect(const UReverbEffect* InReverbEffect, const float WetLevel);
+	void SetSettingsWithReverbEffect(const UReverbEffect* InReverbEffect, const float WetLevel, const float DryLevel = 0.0f);
 
 	UPROPERTY(EditAnywhere, BlueprintReadWrite, Category = SubmixEffectPreset)
 	FSubmixEffectReverbSettings Settings;
diff --git a/Engine/Source/Runtime/AudioMixer/Private/AudioMixerBus.cpp b/Engine/Source/Runtime/AudioMixer/Private/AudioMixerBus.cpp
--- a/Engine/Source/Runtime/AudioMixer/Private/AudioMixerBus.cpp
+++ b/Engine/Source/Runtime/AudioMixer/Private/AudioMixerBus.cpp
@@ -97,7 +97,8 @@ namespace Audio
 				}
 
 				const int32 NumSourceChannels = SourceManager->GetNumChannels(BusSend.SourceId);
-				const int32 NumSourceSamples = NumSourceChannels * SourceManager->GetNumOutputFrames();
+				const int32 NumOutputFrames = SourceManager->GetNumOutputFrames();
+				const int32 NumSourceSamples = NumSourceChannels * NumOutputFrames;
 
 				// If source channels are 1 but the bus is 2 channels, we need to up-mix
 				if (NumSourceChannels == 1 && NumChannels == 2)
@@ -119,7 +120,7 @@ namespace Audio
 				else if (NumSourceChannels == 2 && NumChannels == 1)
 				{
 					int32 SourceSampleIndex = 0;
-					for (int32 BusSampleIndex = 0; BusSampleIndex < NumSourceSamples; ++BusSampleIndex)
+					for (int32 BusSampleIndex = 0; BusSampleIndex < NumOutputFrames; ++BusSampleIndex)
 					{
 						// Downmix the stereo source to mono before summing to bus
 						float SourceSample = 0.0f;
diff --git a/Engine/Source/Runtime/AudioMixer/Private/AudioMixerDevice.cpp b/Engine/Source/Runtime/AudioMixer/Private/AudioMixerDevice.cpp
--- a/Engine/Source/Runtime/AudioMixer/Private/AudioMixerDevice.cpp
+++ b/Engine/Source/Runtime/AudioMixer/Private/AudioMixerDevice.cpp
@@ -266,6 +266,21 @@ namespace Audio
 
 	void FMixerDevice::TeardownHardware()
 	{
+		// Make sure all submixes are registered but not initialized
+		for (TObjectIterator<USoundSubmix> It; It; ++It)
+		{
+			UnregisterSoundSubmix(*It);
+		}
+		
+		// reset all the sound effect presets loaded
+#if WITH_EDITOR
+		for (TObjectIterator<USoundEffectPreset> It; It; ++It)
+		{
+			USoundEffectPreset* SoundEffectPreset = *It;
+			SoundEffectPreset->Init();
+		}
+#endif
+
 		if (AudioMixerPlatform)
 		{
 			SourceManager.Update();
@@ -276,6 +291,10 @@ namespace Audio
 			AudioMixerPlatform->TeardownHardware();
 		}
 
+		// Reset existing submixes if they exist
+		MasterSubmixInstances.Reset();
+		Submixes.Reset();
+
 		if (AmbisonicsMixer.IsValid())
 		{
 			AmbisonicsMixer->Shutdown();
@@ -290,6 +309,13 @@ namespace Audio
 		AudioThreadTimingData.AudioThreadTime = FPlatformTime::Seconds() - AudioThreadTimingData.StartTime + AudioThreadJitterDelta;
 	}
 
+	void FMixerDevice::UpdateGameThread()
+	{
+		LLM_SCOPE(ELLMTag::AudioMixer);
+
+
+	}
+
 	void FMixerDevice::UpdateHardware()
 	{
 		LLM_SCOPE(ELLMTag::AudioMixer);
@@ -329,6 +355,39 @@ namespace Audio
 
 		// Update listener transforms, some effects use the listener transform data
 		SourceManager.SetListenerTransforms(ListenerTransforms);
+
+		// Loop through any envelope-following submixes and perform any broadcasting of envelope data if needed
+		TArray<float> SubmixEnvelopeData;
+		for (USoundSubmix* SoundSubmix : EnvelopeFollowingSubmixes)
+		{
+			if (SoundSubmix)
+			{
+				// Retrieve the submix instance and the envelope data
+				Audio::FMixerSubmixWeakPtr SubmixPtr;
+
+				// First see list of submixes
+				Audio::FMixerSubmixPtr* FoundSubmix = Submixes.Find(SoundSubmix);
+				if (FoundSubmix)
+				{
+					SubmixPtr = *FoundSubmix;
+				}
+				// Fallback to master submix
+				else
+				{
+					SubmixPtr = GetMasterSubmix();
+				}
+
+				// On the audio thread, do the broadcast.
+				FAudioThread::RunCommandOnGameThread([this, SubmixPtr]()
+				{
+					Audio::FMixerSubmixPtr ThisSubmixPtr = SubmixPtr.Pin();
+					if (ThisSubmixPtr.IsValid())
+					{
+						ThisSubmixPtr->BroadcastEnvelope();
+					}
+				});
+			}
+		}
 	}
 
 	double FMixerDevice::GetAudioTime() const
@@ -438,13 +497,17 @@ namespace Audio
 		// Compute the next block of audio in the source manager
 		SourceManager.ComputeNextBlockOfSamples();
 
-		FMixerSubmixPtr MasterSubmix = GetMasterSubmix();
+		FMixerSubmixWeakPtr MasterSubmix = GetMasterSubmix();
 
 		{
 			SCOPE_CYCLE_COUNTER(STAT_AudioMixerSubmixes);
 
-			// Process the audio output from the master submix
-			MasterSubmix->ProcessAudio(ESubmixChannelFormat::Device, Output);
+			FMixerSubmixPtr MasterSubmixPtr = MasterSubmix.Pin();
+			if (MasterSubmixPtr.IsValid())
+			{
+				// Process the audio output from the master submix
+				MasterSubmixPtr->ProcessAudio(ESubmixChannelFormat::Device, Output);
+			}
 		}
 
 		// Reset stopping sounds and clear their state after submixes have been mixed
@@ -560,16 +623,15 @@ namespace Audio
 				}
 			}
 
-			FMixerSubmixPtr MasterSubmixInstance = MasterSubmixInstances[EMasterSubmixType::Master];
+			FMixerSubmixPtr& MasterSubmixInstance = MasterSubmixInstances[EMasterSubmixType::Master];
 			
 			FSoundEffectSubmixInitData InitData;
 			InitData.SampleRate = GetSampleRate();
-
 	
 			// Setup the master reverb plugin
 			if (ReverbPluginInterface.IsValid() && MasterSubmixInstances[EMasterSubmixType::ReverbPlugin].IsValid())
 			{
-				auto ReverbPluginEffectSubmix = ReverbPluginInterface->GetEffectSubmix(FMixerDevice::MasterSubmixes[EMasterSubmixType::ReverbPlugin]);
+				FSoundEffectSubmix* ReverbPluginEffectSubmix = ReverbPluginInterface->GetEffectSubmix(FMixerDevice::MasterSubmixes[EMasterSubmixType::ReverbPlugin]);
 
 				ReverbPluginEffectSubmix->Init(InitData);
 				ReverbPluginEffectSubmix->SetEnabled(true);
@@ -577,7 +639,7 @@ namespace Audio
 				const uint32 ReverbPluginId = FMixerDevice::MasterSubmixes[EMasterSubmixType::ReverbPlugin]->GetUniqueID();
 
 				FMixerSubmixPtr MasterReverbPluginSubmix = MasterSubmixInstances[EMasterSubmixType::ReverbPlugin];
-				MasterReverbPluginSubmix->AddSoundEffectSubmix(ReverbPluginId, MakeShareable(ReverbPluginEffectSubmix));
+				MasterReverbPluginSubmix->AddSoundEffectSubmix(ReverbPluginId, ReverbPluginEffectSubmix);
 				MasterReverbPluginSubmix->SetParentSubmix(MasterSubmixInstance);
 				MasterSubmixInstance->AddChildSubmix(MasterReverbPluginSubmix);
 			}
@@ -598,7 +660,7 @@ namespace Audio
 				const uint32 ReverbPresetId = ReverbPreset->GetUniqueID();
 
 				FMixerSubmixPtr MasterReverbSubmix = MasterSubmixInstances[EMasterSubmixType::Reverb];
-				MasterReverbSubmix->AddSoundEffectSubmix(ReverbPresetId, MakeShareable(ReverbEffectSubmix));
+				MasterReverbSubmix->AddSoundEffectSubmix(ReverbPresetId, ReverbEffectSubmix);
 				MasterReverbSubmix->SetParentSubmix(MasterSubmixInstance);
 				MasterSubmixInstance->AddChildSubmix(MasterReverbSubmix);
 			}
@@ -618,7 +680,7 @@ namespace Audio
 				const uint32 EQPresetId = EQPreset->GetUniqueID();
 
 				FMixerSubmixPtr MasterEQSubmix = MasterSubmixInstances[EMasterSubmixType::EQ];
-				MasterEQSubmix->AddSoundEffectSubmix(EQPresetId, MakeShareable(EQEffectSubmix));
+				MasterEQSubmix->AddSoundEffectSubmix(EQPresetId, EQEffectSubmix);
 				MasterEQSubmix->SetParentSubmix(MasterSubmixInstance);
 				MasterSubmixInstance->AddChildSubmix(MasterEQSubmix);
 
@@ -629,10 +691,17 @@ namespace Audio
 			}
 		}
 
-		// Now register all the non-core submixes
+		// Now register all the non-core submixes.
+
+#if WITH_EDITOR
+		// sanity check the Submixes map to ensure we don't leak any instances of FMixerSubmix.
+		for (auto& Submix : Submixes)
+		{
+			checkf(Submix.Value.IsUnique(), TEXT("Possible leak: please check FMixerSubmix for possible circular references."));
+		}
 
-		// Reset existing submixes if they exist
 		Submixes.Reset();
+#endif
 
 		// Make sure all submixes are registered but not initialized
 		for (TObjectIterator<USoundSubmix> It; It; ++It)
@@ -647,7 +716,7 @@ namespace Audio
 			FMixerSubmixPtr& SubmixInstance = Entry.Value;
 
 			// Setup up the submix instance's parent and add the submix instance as a child
-			FMixerSubmixPtr ParentSubmixInstance;
+			FMixerSubmixWeakPtr ParentSubmixInstance;
 			if (SoundSubmix->ParentSubmix)
 			{
 				ParentSubmixInstance = GetSubmixInstance(SoundSubmix->ParentSubmix);
@@ -656,8 +725,12 @@ namespace Audio
 			{
 				ParentSubmixInstance = GetMasterSubmix();
 			}
-			ParentSubmixInstance->AddChildSubmix(SubmixInstance);
-			SubmixInstance->ParentSubmix = ParentSubmixInstance;
+			FMixerSubmixPtr ParentSubmixInstancePtr = ParentSubmixInstance.Pin();
+			if (ParentSubmixInstancePtr.IsValid())
+			{
+				ParentSubmixInstancePtr->AddChildSubmix(SubmixInstance);
+				SubmixInstance->ParentSubmix = ParentSubmixInstance;
+			}
 
 			// Now add all the child submixes to this submix instance
 			for (USoundSubmix* ChildSubmix : SoundSubmix->ChildSubmixes)
@@ -668,7 +741,14 @@ namespace Audio
 					FChildSubmixInfo ChildSubmixInfo;
 					ChildSubmixInfo.SubmixPtr = GetSubmixInstance(ChildSubmix);
 					ChildSubmixInfo.bNeedsAmbisonicsEncoding = true;
-					SubmixInstance->ChildSubmixes.Add(ChildSubmixInfo.SubmixPtr->GetId(), ChildSubmixInfo);
+
+					FMixerSubmixPtr ChildSubmixInstancePtr = ChildSubmixInfo.SubmixPtr.Pin();
+
+					if (SubmixInstance.IsValid() && ChildSubmixInstancePtr.IsValid())
+					{
+						SubmixInstance->ChildSubmixes.Add(ChildSubmixInstancePtr->GetId(), ChildSubmixInfo);
+					}
+
 				}
 			}
 
@@ -735,27 +815,27 @@ namespace Audio
  		return Settings;
  	}
 
-	FMixerSubmixPtr FMixerDevice::GetMasterSubmix()
+	FMixerSubmixWeakPtr FMixerDevice::GetMasterSubmix()
 	{
 		return MasterSubmixInstances[EMasterSubmixType::Master];
 	}
 
-	FMixerSubmixPtr FMixerDevice::GetMasterReverbPluginSubmix()
+	FMixerSubmixWeakPtr FMixerDevice::GetMasterReverbPluginSubmix()
 	{
 		return MasterSubmixInstances[EMasterSubmixType::ReverbPlugin];
 	}
 
-	FMixerSubmixPtr FMixerDevice::GetMasterReverbSubmix()
+	FMixerSubmixWeakPtr FMixerDevice::GetMasterReverbSubmix()
 	{
 		return MasterSubmixInstances[EMasterSubmixType::Reverb];
 	}
 
-	FMixerSubmixPtr FMixerDevice::GetMasterEQSubmix()
+	FMixerSubmixWeakPtr FMixerDevice::GetMasterEQSubmix()
 	{
 		return MasterSubmixInstances[EMasterSubmixType::EQ];
 	}
 
-	FMixerSubmixPtr FMixerDevice::GetMasterAmbisonicsSubmix()
+	FMixerSubmixWeakPtr FMixerDevice::GetMasterAmbisonicsSubmix()
 	{
 		return MasterSubmixInstances[EMasterSubmixType::Ambisonics];
 	}
@@ -764,7 +844,7 @@ namespace Audio
 	{
 		AudioRenderThreadCommand([this, SubmixEffectId, SoundEffectSubmix]()
 		{
-			MasterSubmixInstances[EMasterSubmixType::Master]->AddSoundEffectSubmix(SubmixEffectId, MakeShareable(SoundEffectSubmix));
+			MasterSubmixInstances[EMasterSubmixType::Master]->AddSoundEffectSubmix(SubmixEffectId, SoundEffectSubmix);
 		});
 	}
 
@@ -865,7 +945,7 @@ namespace Audio
 					if (bInit)
 					{
 						// Setup the parent-child relationship
-						FMixerSubmixPtr ParentSubmixInstance;
+						FMixerSubmixWeakPtr ParentSubmixInstance;
 						if (InSoundSubmix->ParentSubmix)
 						{
 							ParentSubmixInstance = GetSubmixInstance(InSoundSubmix->ParentSubmix);
@@ -875,10 +955,13 @@ namespace Audio
 							ParentSubmixInstance = GetMasterSubmix();
 						}
 
-						ParentSubmixInstance->AddChildSubmix(MixerSubmix);
-						MixerSubmix->SetParentSubmix(ParentSubmixInstance);
-
-						MixerSubmix->Init(InSoundSubmix);
+						FMixerSubmixPtr ParentSubmixInstancePtr = ParentSubmixInstance.Pin();
+						if (ParentSubmixInstancePtr.IsValid())
+						{
+							ParentSubmixInstancePtr->AddChildSubmix(MixerSubmix);
+							MixerSubmix->SetParentSubmix(ParentSubmixInstance);
+							MixerSubmix->Init(InSoundSubmix);
+						}
 					}
 				}
 			}
@@ -919,7 +1002,7 @@ namespace Audio
 #endif
 	}
 
-	FMixerSubmixPtr FMixerDevice::GetSubmixInstance(USoundSubmix* SoundSubmix)
+	FMixerSubmixWeakPtr FMixerDevice::GetSubmixInstance(USoundSubmix* SoundSubmix)
 	{
 		LLM_SCOPE(ELLMTag::AudioMixer);
 
@@ -935,7 +1018,7 @@ namespace Audio
 
 		// At this point, this should exist
 		check(MixerSubmix);
-		return *MixerSubmix;		
+		return *MixerSubmix;
 	}
 
 	FMixerSourceVoice* FMixerDevice::GetMixerSourceVoice()
@@ -1135,7 +1218,11 @@ namespace Audio
 		}
 		else
 		{
-			GetMasterSubmix()->OnStartRecordingOutput(ExpectedRecordingDuration);
+			FMixerSubmixWeakPtr MasterSubmix = GetMasterSubmix();
+			FMixerSubmixPtr MasterSubmixPtr = MasterSubmix.Pin();
+			check(MasterSubmixPtr.IsValid());
+
+			MasterSubmixPtr->OnStartRecordingOutput(ExpectedRecordingDuration);
 		}
 	}
 
@@ -1149,7 +1236,11 @@ namespace Audio
 		}
 		else
 		{
-			return GetMasterSubmix()->OnStopRecordingOutput(OutNumChannels, OutSampleRate);
+			FMixerSubmixWeakPtr MasterSubmix = GetMasterSubmix();
+			FMixerSubmixPtr MasterSubmixPtr = MasterSubmix.Pin();
+			check(MasterSubmixPtr.IsValid());
+
+			return MasterSubmixPtr->OnStopRecordingOutput(OutNumChannels, OutSampleRate);
 		}
 	}
 
@@ -1163,7 +1254,11 @@ namespace Audio
 		}
 		else
 		{
-			GetMasterSubmix()->PauseRecordingOutput();
+			FMixerSubmixWeakPtr MasterSubmix = GetMasterSubmix();
+			FMixerSubmixPtr MasterSubmixPtr = MasterSubmix.Pin();
+			check(MasterSubmixPtr.IsValid());
+
+			MasterSubmixPtr->PauseRecordingOutput();
 		}
 	}
 
@@ -1177,10 +1272,73 @@ namespace Audio
 		}
 		else
 		{
-			GetMasterSubmix()->ResumeRecordingOutput();
+			FMixerSubmixWeakPtr MasterSubmix = GetMasterSubmix();
+			FMixerSubmixPtr MasterSubmixPtr = MasterSubmix.Pin();
+			check(MasterSubmixPtr.IsValid());
+
+			MasterSubmixPtr->ResumeRecordingOutput();
+		}
+	}
+
+	void FMixerDevice::StartEnvelopeFollowing(USoundSubmix* InSubmix)
+	{
+		// if we can find the submix here, record that submix. Otherwise, just record the master submix.
+		Audio::FMixerSubmixPtr* FoundSubmix = Submixes.Find(InSubmix);
+		if (FoundSubmix)
+		{
+			(*FoundSubmix)->StartEnvelopeFollowing(InSubmix->EnvelopeFollowerAttackTime, InSubmix->EnvelopeFollowerReleaseTime);
+		}
+		else
+		{
+			FMixerSubmixWeakPtr MasterSubmix = GetMasterSubmix();
+			FMixerSubmixPtr MasterSubmixPtr = MasterSubmix.Pin();
+			check(MasterSubmixPtr.IsValid());
+
+			MasterSubmixPtr->StartEnvelopeFollowing(InSubmix->EnvelopeFollowerAttackTime, InSubmix->EnvelopeFollowerReleaseTime);
+		}
+
+		EnvelopeFollowingSubmixes.AddUnique(InSubmix);
+	}
+
+	void FMixerDevice::StopEnvelopeFollowing(USoundSubmix* InSubmix)
+	{
+		// if we can find the submix here, record that submix. Otherwise, just record the master submix.
+		Audio::FMixerSubmixPtr* FoundSubmix = Submixes.Find(InSubmix);
+		if (FoundSubmix)
+		{
+			(*FoundSubmix)->StopEnvelopeFollowing();
+		}
+		else
+		{
+			FMixerSubmixWeakPtr MasterSubmix = GetMasterSubmix();
+			FMixerSubmixPtr MasterSubmixPtr = MasterSubmix.Pin();
+			check(MasterSubmixPtr.IsValid());
+
+			MasterSubmixPtr->StopEnvelopeFollowing();
+		}
+
+		EnvelopeFollowingSubmixes.RemoveSingleSwap(InSubmix);
+	}
+
+	void FMixerDevice::AddEnvelopeFollowerDelegate(USoundSubmix* InSubmix, const FOnSubmixEnvelopeBP& OnSubmixEnvelopeBP)
+	{
+		// if we can find the submix here, record that submix. Otherwise, just record the master submix.
+		Audio::FMixerSubmixPtr* FoundSubmix = Submixes.Find(InSubmix);
+		if (FoundSubmix)
+		{
+			(*FoundSubmix)->AddEnvelopeFollowerDelegate(OnSubmixEnvelopeBP);
+		}
+		else
+		{
+			FMixerSubmixWeakPtr MasterSubmix = GetMasterSubmix();
+			FMixerSubmixPtr MasterSubmixPtr = MasterSubmix.Pin();
+			check(MasterSubmixPtr.IsValid());
+
+			MasterSubmixPtr->AddEnvelopeFollowerDelegate(OnSubmixEnvelopeBP);
 		}
 	}
 
+
 	void FMixerDevice::RegisterSubmixBufferListener(ISubmixBufferListener* InSubmixBufferListener, USoundSubmix* InSubmix)
 	{
 		Audio::FMixerSubmixPtr* FoundSubmix = Submixes.Find(InSubmix);
@@ -1190,7 +1348,11 @@ namespace Audio
 		}
 		else
 		{
-			return GetMasterSubmix()->RegisterBufferListener(InSubmixBufferListener);
+			FMixerSubmixWeakPtr MasterSubmix = GetMasterSubmix();
+			FMixerSubmixPtr MasterSubmixPtr = MasterSubmix.Pin();
+			check(MasterSubmixPtr.IsValid());
+
+			return MasterSubmixPtr->RegisterBufferListener(InSubmixBufferListener);
 		}
 	}
 
@@ -1203,7 +1365,11 @@ namespace Audio
 		}
 		else
 		{
-			return GetMasterSubmix()->UnregisterBufferListener(InSubmixBufferListener);
+			FMixerSubmixWeakPtr MasterSubmix = GetMasterSubmix();
+			FMixerSubmixPtr MasterSubmixPtr = MasterSubmix.Pin();
+			check(MasterSubmixPtr.IsValid());
+
+			return MasterSubmixPtr->UnregisterBufferListener(InSubmixBufferListener);
 		}
 	}
 
diff --git a/Engine/Source/Runtime/AudioMixer/Private/AudioMixerEffectsManager.cpp b/Engine/Source/Runtime/AudioMixer/Private/AudioMixerEffectsManager.cpp
--- a/Engine/Source/Runtime/AudioMixer/Private/AudioMixerEffectsManager.cpp
+++ b/Engine/Source/Runtime/AudioMixer/Private/AudioMixerEffectsManager.cpp
@@ -105,14 +105,15 @@ namespace Audio
 	{
 		FMixerDevice* MixerDevice = (FMixerDevice*)AudioDevice;
 
-		FMixerSubmixPtr MasterReverbSubmix = MixerDevice->GetMasterReverbSubmix();
+		FMixerSubmixWeakPtr MasterReverbSubmix = MixerDevice->GetMasterReverbSubmix();
+		FMixerSubmixPtr MasterReverbSubmixPtr = MasterReverbSubmix.Pin();
 		
-		if (MasterReverbSubmix.IsValid())
+		if (MasterReverbSubmixPtr.IsValid())
 		{
-			FSoundEffectSubmixPtr SoundEffectSubmix = MasterReverbSubmix->GetSubmixEffect(0);
-			if (SoundEffectSubmix.IsValid())
+			FSoundEffectSubmix* SoundEffectSubmix = MasterReverbSubmixPtr->GetSubmixEffect(0);
+			if (SoundEffectSubmix)
 			{
-				FSubmixEffectReverb* SoundEffectReverb = static_cast<FSubmixEffectReverb*>(SoundEffectSubmix.Get());
+				FSubmixEffectReverb* SoundEffectReverb = static_cast<FSubmixEffectReverb*>(SoundEffectSubmix);
 				SoundEffectReverb->SetEffectParameters(ReverbEffectParameters);
 				PrintReverbSettings(ReverbEffectParameters);
 			}
@@ -123,19 +124,19 @@ namespace Audio
 	{
 		FMixerDevice* MixerDevice = (FMixerDevice*)AudioDevice;
 
-		FMixerSubmixPtr MasterEQSubmix = MixerDevice->GetMasterEQSubmix();
+		FMixerSubmixWeakPtr MasterEQSubmix = MixerDevice->GetMasterEQSubmix();
+		FMixerSubmixPtr MasterEQSubmixPtr = MasterEQSubmix.Pin();
 
-		if (MasterEQSubmix.IsValid())
+		if (MasterEQSubmixPtr.IsValid())
 		{
-			FSoundEffectSubmixPtr SoundEffectSubmix = MasterEQSubmix->GetSubmixEffect(0);
-			if (SoundEffectSubmix.IsValid())
+			FSoundEffectSubmix* SoundEffectSubmix = MasterEQSubmixPtr->GetSubmixEffect(0);
+			if (SoundEffectSubmix)
 			{
-				FSubmixEffectSubmixEQ* SoundEffectEQ = static_cast<FSubmixEffectSubmixEQ*>(SoundEffectSubmix.Get());
+				FSubmixEffectSubmixEQ* SoundEffectEQ = static_cast<FSubmixEffectSubmixEQ*>(SoundEffectSubmix);
 				SoundEffectEQ->SetEffectParameters(InEQEffectParameters);
 				PrintEQSettings(InEQEffectParameters);
 			}
 		}
-
 	}
 
 	void FAudioMixerEffectsManager::SetRadioEffectParameters(const FAudioRadioEffect& ReverbEffectParameters)
diff --git a/Engine/Source/Runtime/AudioMixer/Private/AudioMixerSource.cpp b/Engine/Source/Runtime/AudioMixer/Private/AudioMixerSource.cpp
--- a/Engine/Source/Runtime/AudioMixer/Private/AudioMixerSource.cpp
+++ b/Engine/Source/Runtime/AudioMixer/Private/AudioMixerSource.cpp
@@ -241,9 +241,13 @@ namespace Audio
 			// Loop through all submix sends to figure out what speaker maps this source is using
 			for (FMixerSourceSubmixSend& Send : InitParams.SubmixSends)
 			{
-				ESubmixChannelFormat SubmixChannelType = Send.Submix->GetSubmixChannels();
-				ChannelMaps[(int32)SubmixChannelType].bUsed = true;
-				ChannelMaps[(int32)SubmixChannelType].ChannelMap.Reset();
+				FMixerSubmixPtr SubmixPtr = Send.Submix.Pin();
+				if (SubmixPtr.IsValid())
+				{
+					ESubmixChannelFormat SubmixChannelType = SubmixPtr->GetSubmixChannels();
+					ChannelMaps[(int32)SubmixChannelType].bUsed = true;
+					ChannelMaps[(int32)SubmixChannelType].ChannelMap.Reset();
+				}
 			}
 
 			// Check to see if this sound has been flagged to be in debug mode
@@ -600,6 +604,11 @@ namespace Audio
 
 	float FMixerSource::GetPlaybackPercent() const
 	{
+		if (InitializationState != EMixerSourceInitializationState::Initialized)
+		{
+			return 0.f;
+		}
+
 		if (MixerSourceVoice && NumTotalFrames > 0)
 		{
 			int64 NumFrames = MixerSourceVoice->GetNumFramesPlayed();
@@ -801,7 +810,7 @@ namespace Audio
 		{
 			if (SendInfo.SoundSubmix != nullptr)
 			{
-				FMixerSubmixPtr SubmixInstance = MixerDevice->GetSubmixInstance(SendInfo.SoundSubmix);
+				FMixerSubmixWeakPtr SubmixInstance = MixerDevice->GetSubmixInstance(SendInfo.SoundSubmix);
 				MixerSourceVoice->SetSubmixSendInfo(SubmixInstance, SendInfo.SendLevel);
 
 				// Make sure we flag that we're using this submix sends since these can be dynamically added from BP
