Copying //UE4/Dev-Audio to //UE4/Dev-Main (Source: //UE4/Dev-Audio @ 4332921)

diff --git a/Engine/Source/Runtime/AudioMixer/Private/AudioMixerSourceManager.cpp b/Engine/Source/Runtime/AudioMixer/Private/AudioMixerSourceManager.cpp
--- a/Engine/Source/Runtime/AudioMixer/Private/AudioMixerSourceManager.cpp
+++ b/Engine/Source/Runtime/AudioMixer/Private/AudioMixerSourceManager.cpp
@@ -220,6 +220,7 @@ namespace Audio
 
 			SourceInfo.SourceBuffer.Reset(NumOutputFrames * 8);
 			SourceInfo.PreDistanceAttenuationBuffer.Reset(NumOutputFrames * 8);
+			SourceInfo.SourceEffectScratchBuffer.Reset(NumOutputFrames * 8);
 			SourceInfo.AudioPluginOutputData.AudioBuffer.Reset(NumOutputFrames * 2);
 		}
 
@@ -349,7 +350,11 @@ namespace Audio
 		// Remove the mixer source from its submix sends
 		for (FMixerSourceSubmixSend& SubmixSendItem : SourceInfo.SubmixSends)
 		{
-			SubmixSendItem.Submix->RemoveSourceVoice(MixerSources[SourceId]);
+			FMixerSubmixPtr SubmixPtr = SubmixSendItem.Submix.Pin();
+			if (SubmixPtr.IsValid())
+			{
+				SubmixPtr->RemoveSourceVoice(MixerSources[SourceId]);
+			}
 		}
 		SourceInfo.SubmixSends.Reset();
 
@@ -401,6 +406,7 @@ namespace Audio
 		SourceInfo.CurrentAudioChunkNumFrames = 0;
 		SourceInfo.SourceBuffer.Reset();
 		SourceInfo.PreDistanceAttenuationBuffer.Reset();
+		SourceInfo.SourceEffectScratchBuffer.Reset();
 		SourceInfo.AudioPluginOutputData.AudioBuffer.Reset();
 		SourceInfo.CurrentFrameValues.Reset();
 		SourceInfo.NextFrameValues.Reset();
@@ -611,16 +617,6 @@ namespace Audio
 				InitData.NumSourceChannels = InitParams.NumInputChannels;
 				InitData.AudioClock = MixerDevice->GetAudioTime();
 
-				if (InitParams.NumInputFrames != INDEX_NONE)
-				{
-					InitData.SourceDuration = (float)InitParams.NumInputFrames / MixerDevice->SampleRate;
-				}
-				else
-				{
-					// Procedural sound waves have no known duration
-					InitData.SourceDuration = (float)INDEX_NONE;
-				}
-
 				SourceInfo.SourceEffectChainId = InitParams.SourceEffectChainId;
 				BuildSourceEffectChain(SourceId, InitData, InitParams.SourceEffectChain);
 
@@ -707,14 +703,19 @@ namespace Audio
 			for (int32 i = 0; i < InitParams.SubmixSends.Num(); ++i)
 			{
 				const FMixerSourceSubmixSend& MixerSubmixSend = InitParams.SubmixSends[i];
-				SourceInfo.SubmixSends.Add(MixerSubmixSend);
-				MixerSubmixSend.Submix->AddOrSetSourceVoice(InitParams.SourceVoice, MixerSubmixSend.SendLevel);
 
-				// Prepare output buffers and speaker map entries for every submix channel type
-				const ESubmixChannelFormat SubmixChannelType = MixerSubmixSend.Submix->GetSubmixChannels();
+				FMixerSubmixPtr SubmixPtr = MixerSubmixSend.Submix.Pin();
+				if (SubmixPtr.IsValid())
+				{
+					SourceInfo.SubmixSends.Add(MixerSubmixSend);
+					SubmixPtr->AddOrSetSourceVoice(InitParams.SourceVoice, MixerSubmixSend.SendLevel);
 
-				// Flag that we're going to be using this channel info entry
-				SourceInfo.SubmixChannelInfo[(int32)SubmixChannelType].bUsed = true;
+					// Prepare output buffers and speaker map entries for every submix channel type
+					const ESubmixChannelFormat SubmixChannelType = SubmixPtr->GetSubmixChannels();
+
+					// Flag that we're going to be using this channel info entry
+					SourceInfo.SubmixChannelInfo[(int32)SubmixChannelType].bUsed = true;
+				}
 			}
 
 #if AUDIO_MIXER_ENABLE_DEBUG_MODE
@@ -1035,27 +1036,35 @@ namespace Audio
 		{
 			FSourceInfo& SourceInfo = SourceInfos[SourceId];
 
-			bool bIsNew = true;
-			for (FMixerSourceSubmixSend& SubmixSend : SourceInfo.SubmixSends)
+			FMixerSubmixPtr InSubmixPtr = InSubmixSend.Submix.Pin();
+			if (InSubmixPtr.IsValid())
 			{
-				if (SubmixSend.Submix->GetId() == InSubmixSend.Submix->GetId())
+				bool bIsNew = true;
+				for (FMixerSourceSubmixSend& SubmixSend : SourceInfo.SubmixSends)
 				{
-					SubmixSend.SendLevel = InSubmixSend.SendLevel;
-					bIsNew = false;
-					break;
+					FMixerSubmixPtr SubmixPtr = SubmixSend.Submix.Pin();
+					if (SubmixPtr.IsValid())
+					{
+						if (SubmixPtr->GetId() == InSubmixPtr->GetId())
+						{
+							SubmixSend.SendLevel = InSubmixSend.SendLevel;
+							bIsNew = false;
+							break;
+						}
+					}
 				}
-			}
 
-			if (bIsNew)
-			{
-				SourceInfo.SubmixSends.Add(InSubmixSend);
+				if (bIsNew)
+				{
+					SourceInfo.SubmixSends.Add(InSubmixSend);
 
-				// Flag that we're now using this submix channel info
-				ESubmixChannelFormat ChannelType = InSubmixSend.Submix->GetSubmixChannels();
-				SourceInfo.SubmixChannelInfo[(int32)ChannelType].bUsed = true;
-			}
+					// Flag that we're now using this submix channel info
+					ESubmixChannelFormat ChannelType = InSubmixPtr->GetSubmixChannels();
+					SourceInfo.SubmixChannelInfo[(int32)ChannelType].bUsed = true;
+				}
 
-			InSubmixSend.Submix->AddOrSetSourceVoice(MixerSources[SourceId], InSubmixSend.SendLevel);
+				InSubmixPtr->AddOrSetSourceVoice(MixerSources[SourceId], InSubmixSend.SendLevel);
+			}
 		});
 	}
 
@@ -1234,6 +1243,9 @@ namespace Audio
 			SourceInfo.PreDistanceAttenuationBuffer.Reset();
 			SourceInfo.PreDistanceAttenuationBuffer.AddZeroed(NumSamples);
 
+			SourceInfo.SourceEffectScratchBuffer.Reset();
+			SourceInfo.SourceEffectScratchBuffer.AddZeroed(NumSamples);
+
 			SourceInfo.SourceBuffer.Reset();
 			SourceInfo.SourceBuffer.AddZeroed(NumSamples);
 
@@ -1542,7 +1554,9 @@ namespace Audio
 			if (SourceInfo.BusSends[(int32)EBusSendType::PreEffect].Num() > 0)
 			{
 				SourceInfo.PreEffectBuffer.Reset();
-				SourceInfo.PreEffectBuffer.Append(SourceInfo.PreDistanceAttenuationBuffer);
+				SourceInfo.PreEffectBuffer.Reserve(SourceInfo.PreDistanceAttenuationBuffer.Num());
+
+				FMemory::Memcpy(SourceInfo.PreEffectBuffer.GetData(), SourceInfo.PreDistanceAttenuationBuffer.GetData(), sizeof(float)*SourceInfo.PreDistanceAttenuationBuffer.Num());
 			}
 
 			float* PreDistanceAttenBufferPtr = SourceInfo.PreDistanceAttenuationBuffer.GetData();
@@ -1583,46 +1597,32 @@ namespace Audio
 			// Now process the effect chain if it exists
 			if (!DisableSourceEffectsCvar && SourceInfo.SourceEffects.Num() > 0)
 			{
+				// Prepare this source's effect chain input data
 				SourceInfo.SourceEffectInputData.CurrentVolume = SourceInfo.VolumeSourceDestination;
+				SourceInfo.SourceEffectInputData.CurrentPitch = SourceInfo.PitchSourceParam.GetValue();
+				SourceInfo.SourceEffectInputData.AudioClock = MixerDevice->GetAudioClock();
+				if (SourceInfo.NumInputFrames > 0)
+				{
+					SourceInfo.SourceEffectInputData.CurrentPlayFraction = (float)SourceInfo.NumFramesPlayed / SourceInfo.NumInputFrames;
+				}
+				SourceInfo.SourceEffectInputData.SpatParams = SourceInfo.SpatParams;
 
-				SourceInfo.SourceEffectOutputData.AudioFrame.Reset();
-				SourceInfo.SourceEffectOutputData.AudioFrame.AddZeroed(SourceInfo.NumInputChannels);
-
-				SourceInfo.SourceEffectInputData.AudioFrame.Reset();
-				SourceInfo.SourceEffectInputData.AudioFrame.AddZeroed(SourceInfo.NumInputChannels);
+				// Get a ptr to pre-distance attenuation buffer ptr
+				float* OutputSourceEffectBufferPtr = SourceInfo.SourceEffectScratchBuffer.GetData();
 
-				float* SourceEffectInputDataFramePtr = SourceInfo.SourceEffectInputData.AudioFrame.GetData();
-				float* SourceEffectOutputDataFramePtr = SourceInfo.SourceEffectOutputData.AudioFrame.GetData();
+				SourceInfo.SourceEffectInputData.InputSourceEffectBufferPtr = SourceInfo.PreDistanceAttenuationBuffer.GetData();
+				SourceInfo.SourceEffectInputData.NumSamples = NumSamples;
 
-				// Process the effect chain for this buffer per frame
-				for (int32 Sample = 0; Sample < NumSamples; Sample += SourceInfo.NumInputChannels)
+				// Loop through the effect chain passing in buffers
+				for (FSoundEffectSource* SoundEffectSource : SourceInfo.SourceEffects)
 				{
-					// Get the buffer input sample
-					for (int32 Chan = 0; Chan < SourceInfo.NumInputChannels; ++Chan)
+					if (SoundEffectSource->IsActive())
 					{
-						SourceEffectInputDataFramePtr[Chan] = PreDistanceAttenBufferPtr[Sample + Chan];
-					}
-
-					for (int32 SourceEffectIndex = 0; SourceEffectIndex < SourceInfo.SourceEffects.Num(); ++SourceEffectIndex)
-					{
-						if (SourceInfo.SourceEffects[SourceEffectIndex]->IsActive())
-						{
-							SourceInfo.SourceEffects[SourceEffectIndex]->Update();
-
-							SourceInfo.SourceEffects[SourceEffectIndex]->ProcessAudio(SourceInfo.SourceEffectInputData, SourceInfo.SourceEffectOutputData);
-
-							// Copy the output of the effect into the input so the next effect will get the outputs audio
-							for (int32 Chan = 0; Chan < SourceInfo.NumInputChannels; ++Chan)
-							{
-								SourceEffectInputDataFramePtr[Chan] = SourceEffectOutputDataFramePtr[Chan];
-							}
-						}
-					}
+						SoundEffectSource->Update();
+						SoundEffectSource->ProcessAudio(SourceInfo.SourceEffectInputData, OutputSourceEffectBufferPtr);
 
-					// Copy audio frame back to the buffer
-					for (int32 Chan = 0; Chan < SourceInfo.NumInputChannels; ++Chan)
-					{
-						PreDistanceAttenBufferPtr[Sample + Chan] = SourceEffectInputDataFramePtr[Chan];
+						// Copy output to input
+						FMemory::Memcpy(SourceInfo.SourceEffectInputData.InputSourceEffectBufferPtr, OutputSourceEffectBufferPtr, sizeof(float)*NumSamples);
 					}
 				}
 			}
@@ -1695,7 +1695,7 @@ namespace Audio
 							SourceInfo.bIsBypassingHPF = false;
 						}
 
-						SourceInfo.HighPassFilter.ProcessAudio(&PostDistanceAttenBufferPtr[SampleIndex], &PostDistanceAttenBufferPtr[SampleIndex]);
+						SourceInfo.HighPassFilter.ProcessAudioFrame(&PostDistanceAttenBufferPtr[SampleIndex], &PostDistanceAttenBufferPtr[SampleIndex]);
 					}
 					else
 					{
@@ -1997,16 +1997,6 @@ namespace Audio
 					{
 						InitData.NumSourceChannels = SourceInfo.NumInputChannels;
 
-						if (SourceInfo.NumInputFrames != INDEX_NONE)
-						{
-							InitData.SourceDuration = (float)SourceInfo.NumInputFrames / InitData.SampleRate;
-						}
-						else
-						{
-							// Procedural sound waves have no known duration
-							InitData.SourceDuration = (float)INDEX_NONE;
-						}
-
 						// First reset the source effect chain
 						ResetSourceEffectChain(SourceId);
 
diff --git a/Engine/Source/Runtime/AudioMixer/Private/AudioMixerSourceManager.h b/Engine/Source/Runtime/AudioMixer/Private/AudioMixerSourceManager.h
--- a/Engine/Source/Runtime/AudioMixer/Private/AudioMixerSourceManager.h
+++ b/Engine/Source/Runtime/AudioMixer/Private/AudioMixerSourceManager.h
@@ -49,6 +49,7 @@ namespace Audio
 	};
 
 	typedef TSharedPtr<FMixerSubmix, ESPMode::ThreadSafe> FMixerSubmixPtr;
+	typedef TWeakPtr<FMixerSubmix, ESPMode::ThreadSafe> FMixerSubmixWeakPtr;
 
 	class ISourceListener
 	{
@@ -70,7 +71,7 @@ namespace Audio
 	struct FMixerSourceSubmixSend
 	{
 		// The submix ptr
-		FMixerSubmixPtr Submix;
+		FMixerSubmixWeakPtr Submix;
 
 		// The amount of audio that is to be mixed into this submix
 		float SendLevel;
@@ -107,13 +108,13 @@ namespace Audio
 		UReverbPluginSourceSettingsBase* ReverbPluginSettings;
 		FName AudioComponentUserID;
 		uint64 AudioComponentID;
-		bool bPlayEffectChainTails;
-		bool bUseHRTFSpatialization;
-		bool bIsDebugMode;
-		bool bOutputToBusOnly;
-		bool bIsVorbis;
-		bool bIsAmbisonics;
-		bool bIsSeeking;
+		uint8 bPlayEffectChainTails : 1;
+		uint8 bUseHRTFSpatialization : 1;
+		uint8 bIsDebugMode : 1;
+		uint8 bOutputToBusOnly : 1;
+		uint8 bIsVorbis : 1;
+		uint8 bIsAmbisonics : 1;
+		uint8 bIsSeeking : 1;
 
 		FMixerSourceVoiceInitParams()
 			: MixerSourceBuffer(nullptr)
@@ -376,6 +377,7 @@ namespace Audio
 			Audio::AlignedFloatBuffer SourceBuffer;
 			Audio::AlignedFloatBuffer PreEffectBuffer;
 			Audio::AlignedFloatBuffer PreDistanceAttenuationBuffer;
+			Audio::AlignedFloatBuffer SourceEffectScratchBuffer;
 
 			TArray<float> CurrentFrameValues;
 			TArray<float> NextFrameValues;
@@ -421,7 +423,6 @@ namespace Audio
 			TArray<USoundEffectSourcePreset*> SourceEffectPresets;
 			bool bEffectTailsDone;
 			FSoundEffectSourceInputData SourceEffectInputData;
-			FSoundEffectSourceOutputData SourceEffectOutputData;
 
 			FAudioPluginSourceOutputData AudioPluginOutputData;
 
diff --git a/Engine/Source/Runtime/AudioMixer/Private/AudioMixerSourceVoice.cpp b/Engine/Source/Runtime/AudioMixer/Private/AudioMixerSourceVoice.cpp
--- a/Engine/Source/Runtime/AudioMixer/Private/AudioMixerSourceVoice.cpp
+++ b/Engine/Source/Runtime/AudioMixer/Private/AudioMixerSourceVoice.cpp
@@ -64,7 +64,11 @@ namespace Audio
 
 			for (int32 i = 0; i < InitParams.SubmixSends.Num(); ++i)
 			{
-				SubmixSends.Add(InitParams.SubmixSends[i].Submix->GetId(), InitParams.SubmixSends[i]);
+				FMixerSubmixPtr SubmixPtr = InitParams.SubmixSends[i].Submix.Pin();
+				if (SubmixPtr.IsValid())
+				{
+					SubmixSends.Add(SubmixPtr->GetId(), InitParams.SubmixSends[i]);
+				}
 			}
 
 			bStopFadedOut = false;
@@ -242,26 +246,31 @@ namespace Audio
 		return SourceManager->MixOutputBuffers(SourceId, InSubmixChannelType, SendLevel, OutWetBuffer);
 	}
 
-	void FMixerSourceVoice::SetSubmixSendInfo(FMixerSubmixPtr Submix, const float SendLevel)
+	void FMixerSourceVoice::SetSubmixSendInfo(FMixerSubmixWeakPtr Submix, const float SendLevel)
 	{
 		AUDIO_MIXER_CHECK_GAME_THREAD(MixerDevice);
 
 		if (!bOutputToBusOnly)
 		{
-			FMixerSourceSubmixSend* SubmixSend = SubmixSends.Find(Submix->GetId());
-			if (!SubmixSend)
+			FMixerSubmixPtr SubmixPtr = Submix.Pin();
+			if (SubmixPtr.IsValid())
 			{
-				FMixerSourceSubmixSend NewSubmixSend;
-				NewSubmixSend.Submix = Submix;
-				NewSubmixSend.SendLevel = SendLevel;
-				NewSubmixSend.bIsMainSend = false;
-				SubmixSends.Add(Submix->GetId(), NewSubmixSend);
-				SourceManager->SetSubmixSendInfo(SourceId, NewSubmixSend);
-			}
-			else if (!FMath::IsNearlyEqual(SubmixSend->SendLevel, SendLevel))
-			{
-				SubmixSend->SendLevel = SendLevel;
-				SourceManager->SetSubmixSendInfo(SourceId, *SubmixSend);
+				FMixerSourceSubmixSend* SubmixSend = SubmixSends.Find(SubmixPtr->GetId());
+
+				if (!SubmixSend)
+				{
+					FMixerSourceSubmixSend NewSubmixSend;
+					NewSubmixSend.Submix = Submix;
+					NewSubmixSend.SendLevel = SendLevel;
+					NewSubmixSend.bIsMainSend = false;
+					SubmixSends.Add(SubmixPtr->GetId(), NewSubmixSend);
+					SourceManager->SetSubmixSendInfo(SourceId, NewSubmixSend);
+				}
+				else if (!FMath::IsNearlyEqual(SubmixSend->SendLevel, SendLevel))
+				{
+					SubmixSend->SendLevel = SendLevel;
+					SourceManager->SetSubmixSendInfo(SourceId, *SubmixSend);
+				}
 			}
 		}
 	}
diff --git a/Engine/Source/Runtime/AudioMixer/Private/AudioMixerSourceVoice.h b/Engine/Source/Runtime/AudioMixer/Private/AudioMixerSourceVoice.h
--- a/Engine/Source/Runtime/AudioMixer/Private/AudioMixerSourceVoice.h
+++ b/Engine/Source/Runtime/AudioMixer/Private/AudioMixerSourceVoice.h
@@ -90,7 +90,7 @@ namespace Audio
 		void MixOutputBuffers(const ESubmixChannelFormat InSubmixChannelType, const float SendLevel, AlignedFloatBuffer& OutWetBuffer) const;
 
 		// Sets the submix send levels
-		void SetSubmixSendInfo(FMixerSubmixPtr Submix, const float SendLevel);
+		void SetSubmixSendInfo(FMixerSubmixWeakPtr Submix, const float SendLevel);
 
 		// Called when the source is a bus and needs to mix other sources together to generate output
 		void OnMixBus(FMixerSourceVoiceBuffer* OutMixerSourceBuffer);
diff --git a/Engine/Source/Runtime/AudioMixer/Private/AudioMixerSubmix.cpp b/Engine/Source/Runtime/AudioMixer/Private/AudioMixerSubmix.cpp
--- a/Engine/Source/Runtime/AudioMixer/Private/AudioMixerSubmix.cpp
+++ b/Engine/Source/Runtime/AudioMixer/Private/AudioMixerSubmix.cpp
@@ -21,6 +21,7 @@ namespace Audio
 		, NumSamples(0)
 		, SubmixAmbisonicsEncoderID(INDEX_NONE)
 		, SubmixAmbisonicsDecoderID(INDEX_NONE)
+		, EnvelopeNumChannels(0)
 		, bIsRecording(false)
 		, OwningSubmixObject(nullptr)
 	{
@@ -77,7 +78,7 @@ namespace Audio
 
 					FSubmixEffectInfo EffectInfo;
 					EffectInfo.PresetId = EffectPreset->GetUniqueID();
-					EffectInfo.EffectInstance = MakeShareable(SubmixEffect);
+					EffectInfo.EffectInstance = SubmixEffect;
 
 					// Add the effect to this submix's chain
 					EffectSubmixChain.Add(EffectInfo);
@@ -124,13 +125,13 @@ namespace Audio
 		}
 	}
 
-	void FMixerSubmix::SetParentSubmix(TSharedPtr<FMixerSubmix, ESPMode::ThreadSafe> Submix)
+	void FMixerSubmix::SetParentSubmix(TWeakPtr<FMixerSubmix, ESPMode::ThreadSafe> SubmixWeakPtr)
 	{
-		SubmixCommand([this, Submix]()
+		SubmixCommand([this, SubmixWeakPtr]()
 		{
 			AUDIO_MIXER_CHECK_AUDIO_PLAT_THREAD(MixerDevice);
 
-			ParentSubmix = Submix;
+			ParentSubmix = SubmixWeakPtr;
 
 			if (ChannelFormat == ESubmixChannelFormat::Ambisonics && AmbisonicsMixer.IsValid())
 			{
@@ -139,23 +140,27 @@ namespace Audio
 		});
 	}
 
-	void FMixerSubmix::AddChildSubmix(TSharedPtr<FMixerSubmix, ESPMode::ThreadSafe> Submix)
+	void FMixerSubmix::AddChildSubmix(TWeakPtr<FMixerSubmix, ESPMode::ThreadSafe> SubmixWeakPtr)
 	{
-		SubmixCommand([this, Submix]()
+		SubmixCommand([this, SubmixWeakPtr]()
 		{
 			AUDIO_MIXER_CHECK_AUDIO_PLAT_THREAD(MixerDevice);
 
 			FChildSubmixInfo NewChildSubmixInfo;
-			NewChildSubmixInfo.SubmixPtr = Submix;
+			NewChildSubmixInfo.SubmixPtr = SubmixWeakPtr;
 
 			//TODO: switch this conditionally when we are able to route submixes to ambisonics submix.
 			NewChildSubmixInfo.bNeedsAmbisonicsEncoding = false;
 
-			ChildSubmixes.Add(Submix->GetId(), NewChildSubmixInfo);
-			
-			if (ChannelFormat == ESubmixChannelFormat::Ambisonics && AmbisonicsMixer.IsValid())
+			TSharedPtr<Audio::FMixerSubmix, ESPMode::ThreadSafe> SubmixSharedPtr = SubmixWeakPtr.Pin();
+			if (SubmixSharedPtr.IsValid())
 			{
-				UpdateAmbisonicsEncoderForChildren();
+				ChildSubmixes.Add(SubmixSharedPtr->GetId(), NewChildSubmixInfo);
+
+				if (ChannelFormat == ESubmixChannelFormat::Ambisonics && AmbisonicsMixer.IsValid())
+				{
+					UpdateAmbisonicsEncoderForChildren();
+				}
 			}
 		});
 	}
@@ -165,7 +170,7 @@ namespace Audio
 		return ChannelFormat;
 	}
 
-	TSharedPtr<FMixerSubmix, ESPMode::ThreadSafe> FMixerSubmix::GetParentSubmix()
+	TWeakPtr<FMixerSubmix, ESPMode::ThreadSafe> FMixerSubmix::GetParentSubmix()
 	{
 		return ParentSubmix;
 	}
@@ -211,18 +216,28 @@ namespace Audio
 		AUDIO_MIXER_CHECK_AUDIO_PLAT_THREAD(MixerDevice);
 		
 		// If the source has a corresponding ambisonics encoder, close it out.
-		uint32 SourceEncoderID = MixerSourceVoices[InSourceVoice].AmbisonicsEncoderId;
+		uint32 SourceEncoderID = INDEX_NONE;
+		const FSubmixVoiceData* MixerSourceVoiceData = MixerSourceVoices.Find(InSourceVoice);
+		if (MixerSourceVoiceData)
+		{
+			SourceEncoderID = MixerSourceVoiceData->AmbisonicsEncoderId;
+		}
+		
 		if (SourceEncoderID != INDEX_NONE)
 		{
 			check(AmbisonicsMixer.IsValid());
 			AmbisonicsMixer->OnCloseEncodingStream(SourceEncoderID);
 		}
 
-		int32 NumRemoved = MixerSourceVoices.Remove(InSourceVoice);
-		AUDIO_MIXER_CHECK(NumRemoved == 1);
+		// If we did find a valid corresponding FSubmixVoiceData, remove it from the map.
+		if (MixerSourceVoiceData)
+		{
+			int32 NumRemoved = MixerSourceVoices.Remove(InSourceVoice);
+			AUDIO_MIXER_CHECK(NumRemoved == 1);
+		}
 	}
 
-	void FMixerSubmix::AddSoundEffectSubmix(uint32 SubmixPresetId, FSoundEffectSubmixPtr InSoundEffectSubmix)
+	void FMixerSubmix::AddSoundEffectSubmix(uint32 SubmixPresetId, FSoundEffectSubmix* InSoundEffectSubmix)
 	{
 		AUDIO_MIXER_CHECK_AUDIO_PLAT_THREAD(MixerDevice);
 
@@ -253,8 +268,12 @@ namespace Audio
 			// If the ID's match, delete and remove the effect instance but don't modify the effect submix chain array itself
 			if (EffectSubmixChain[i].PresetId == SubmixPresetId)
 			{
-				// Release the reference to the effect instance
-				EffectSubmixChain[i].EffectInstance = nullptr;
+				// Delete the reference to the effect instance
+				if (EffectSubmixChain[i].EffectInstance)
+				{
+					delete EffectSubmixChain[i].EffectInstance;
+					EffectSubmixChain[i].EffectInstance = nullptr;
+				}
 				EffectSubmixChain[i].PresetId = INDEX_NONE;
 				return;
 			}
@@ -264,6 +283,17 @@ namespace Audio
 
 	void FMixerSubmix::ClearSoundEffectSubmixes()
 	{
+		for (FSubmixEffectInfo& Info : EffectSubmixChain)
+		{
+			if (Info.EffectInstance)
+			{
+				Info.EffectInstance->UnregisterWithPreset();
+
+				delete Info.EffectInstance;
+				Info.EffectInstance = nullptr;
+			}
+		}
+
 		EffectSubmixChain.Reset();
 	}
 
@@ -391,10 +421,12 @@ namespace Audio
 		{
 			FChildSubmixInfo& ChildSubmix = Iter.Value;
 
+			TSharedPtr<Audio::FMixerSubmix, ESPMode::ThreadSafe> SubmixPtr = ChildSubmix.SubmixPtr.Pin();
+
 			//Check to see if this child is an ambisonics submix.
-			if (ChildSubmix.SubmixPtr.IsValid() && ChildSubmix.SubmixPtr->GetSubmixChannels() == ESubmixChannelFormat::Ambisonics)
+			if (SubmixPtr.IsValid() && SubmixPtr->GetSubmixChannels() == ESubmixChannelFormat::Ambisonics)
 			{
-				UAmbisonicsSubmixSettingsBase* ChildAmbisonicsSettings = ChildSubmix.SubmixPtr->AmbisonicsSettings;
+				UAmbisonicsSubmixSettingsBase* ChildAmbisonicsSettings = SubmixPtr->AmbisonicsSettings;
 
 				//Check if this child submix needs to be reencoded.
 				if (!ChildAmbisonicsSettings || AmbisonicsMixer->ShouldReencodeBetween(ChildAmbisonicsSettings, AmbisonicsSettings))
@@ -426,9 +458,13 @@ namespace Audio
 	{
 		UAmbisonicsSubmixSettingsBase* ParentAmbisonicsSettings = nullptr;
 
-		if (ParentSubmix.IsValid() && ParentSubmix->GetSubmixChannels() == ESubmixChannelFormat::Ambisonics)
+		TSharedPtr<FMixerSubmix, ESPMode::ThreadSafe> ParentSubmixSharedPtr = ParentSubmix.Pin();
+		if (ParentSubmixSharedPtr.IsValid())
 		{
-			ParentAmbisonicsSettings = ParentSubmix->AmbisonicsSettings;
+			if (ParentSubmixSharedPtr->GetSubmixChannels() == ESubmixChannelFormat::Ambisonics)
+			{
+				ParentAmbisonicsSettings = ParentSubmixSharedPtr->AmbisonicsSettings;
+			}
 		}
 
 		// If we need to reencode between here and the parent submix, set up the submix decoder.
@@ -445,9 +481,10 @@ namespace Audio
 	void FMixerSubmix::SetUpAmbisonicsPositionalData()
 	{
 		// If there is a parent and we are not passing it this submix's ambisonics audio, retrieve that submix's channel format.
-		if (ParentSubmix.IsValid())
+		TSharedPtr<Audio::FMixerSubmix, ESPMode::ThreadSafe> ParentSubmixSharedPtr = ParentSubmix.Pin();
+		if (ParentSubmixSharedPtr.IsValid())
 		{
-			const ESubmixChannelFormat ParentSubmixFormat = ParentSubmix->GetSubmixChannels();
+			const ESubmixChannelFormat ParentSubmixFormat = ParentSubmixSharedPtr->GetSubmixChannels();
 
 			const int32 NumParentChannels = MixerDevice->GetNumChannelsForSubmixFormat(ParentSubmixFormat);
 			CachedPositionalData.OutputNumChannels = NumParentChannels;
@@ -523,27 +560,27 @@ namespace Audio
 			SCOPE_CYCLE_COUNTER(STAT_AudioMixerSubmixChildren);
 
 			// First loop this submix's child submixes mixing in their output into this submix's dry/wet buffers.
-			for (auto ChildSubmixEntry : ChildSubmixes)
+			for (auto& ChildSubmixEntry : ChildSubmixes)
 			{
-				TSharedPtr<Audio::FMixerSubmix, ESPMode::ThreadSafe> ChildSubmix = ChildSubmixEntry.Value.SubmixPtr;
-
-				AUDIO_MIXER_CHECK(ChildSubmix.IsValid());
-				
-				ScratchBuffer.Reset(NumSamples);
-				ScratchBuffer.AddZeroed(NumSamples);
+				TSharedPtr<Audio::FMixerSubmix, ESPMode::ThreadSafe> ChildSubmix = ChildSubmixEntry.Value.SubmixPtr.Pin();
+				if (ChildSubmix.IsValid())
+				{
+					ScratchBuffer.Reset(NumSamples);
+					ScratchBuffer.AddZeroed(NumSamples);
 
-				ChildSubmix->ProcessAudio(ChannelFormat, ScratchBuffer);
+					ChildSubmix->ProcessAudio(ChannelFormat, ScratchBuffer);
 
-				float* ScratchBufferPtr = ScratchBuffer.GetData();
+					float* ScratchBufferPtr = ScratchBuffer.GetData();
 
-				if (ChildSubmixEntry.Value.bNeedsAmbisonicsEncoding)
-				{
-					// Encode into ambisonics. TODO: Implement.
-					EncodeAndMixInChildSubmix(ChildSubmixEntry.Value);
-				}
-				else
-				{
-					Audio::MixInBufferFast(ScratchBufferPtr, BufferPtr, NumSamples);
+					if (ChildSubmixEntry.Value.bNeedsAmbisonicsEncoding)
+					{
+						// Encode into ambisonics. TODO: Implement.
+						EncodeAndMixInChildSubmix(ChildSubmixEntry.Value);
+					}
+					else
+					{
+						Audio::MixInBufferFast(ScratchBufferPtr, BufferPtr, NumSamples);
+					}
 				}
 			}
 		}
@@ -583,7 +620,7 @@ namespace Audio
 
 			for (FSubmixEffectInfo& SubmixEffectInfo : EffectSubmixChain)
 			{
-				FSoundEffectSubmixPtr SubmixEffect = SubmixEffectInfo.EffectInstance;
+				FSoundEffectSubmix* SubmixEffect = SubmixEffectInfo.EffectInstance;
 
 				// Reset the output scratch buffer
 				ScratchBuffer.Reset(NumSamples);
@@ -635,6 +672,39 @@ namespace Audio
 			FMemory::Memcpy((void*)OutAudioBuffer.GetData(), (void*)InputBuffer.GetData(), sizeof(float)*NumSamples);
 		}
 
+		// Perform any envelope following if we're told to do so
+		if (bIsEnvelopeFollowing)
+		{
+			const int32 OutBufferSamples = OutAudioBuffer.Num();
+			const float* OutAudioBufferPtr = OutAudioBuffer.GetData();
+
+			float TempEnvelopeValues[AUDIO_MIXER_MAX_OUTPUT_CHANNELS];
+			FMemory::Memset(TempEnvelopeValues, sizeof(float)*AUDIO_MIXER_MAX_OUTPUT_CHANNELS);
+
+			// Perform envelope following per channel
+			for (int32 ChannelIndex = 0; ChannelIndex < NumChannels; ++ChannelIndex)
+			{
+				// Get the envelope follower for the channel
+				FEnvelopeFollower& EnvFollower = EnvelopeFollowers[ChannelIndex];
+
+				// Track the last sample
+				for (int32 SampleIndex = ChannelIndex; SampleIndex < OutBufferSamples; SampleIndex += NumChannels)
+				{
+					const float SampleValue = OutAudioBufferPtr[SampleIndex];
+					EnvFollower.ProcessAudio(SampleValue);
+				}
+
+				// Store the last value
+				TempEnvelopeValues[ChannelIndex] = EnvFollower.GetCurrentValue();
+			}
+
+			FScopeLock EnvelopeScopeLock(&EnvelopeCriticalSection);
+
+			EnvelopeNumChannels = NumChannels;
+			FMemory::Memcpy(EnvelopeValues, TempEnvelopeValues, sizeof(float)*AUDIO_MIXER_MAX_OUTPUT_CHANNELS);
+		}
+
+
 		// Now loop through any buffer listeners and feed the listeners the result of this audio callback
 		{
 			double AudioClock = MixerDevice->GetAudioTime();
@@ -642,6 +712,7 @@ namespace Audio
 			FScopeLock Lock(&BufferListenerCriticalSection);
 			for (ISubmixBufferListener* BufferListener : BufferListeners)
 			{
+				check(BufferListener);
 				BufferListener->OnNewSubmixBuffer(OwningSubmixObject, OutAudioBuffer.GetData(), OutAudioBuffer.Num(), NumChannels, SampleRate, AudioClock);
 			}
 		}
@@ -662,7 +733,7 @@ namespace Audio
 		return EffectSubmixChain.Num();
 	}
 
-	FSoundEffectSubmixPtr FMixerSubmix::GetSubmixEffect(const int32 InIndex)
+	FSoundEffectSubmix* FMixerSubmix::GetSubmixEffect(const int32 InIndex)
 	{
 		if (InIndex < EffectSubmixChain.Num())
 		{
@@ -735,13 +806,64 @@ namespace Audio
 	void FMixerSubmix::RegisterBufferListener(ISubmixBufferListener* BufferListener)
 	{
 		FScopeLock Lock(&BufferListenerCriticalSection);
+		check(BufferListener);
 		BufferListeners.AddUnique(BufferListener);
 	}
 
 	void FMixerSubmix::UnregisterBufferListener(ISubmixBufferListener* BufferListener)
 	{
 		FScopeLock Lock(&BufferListenerCriticalSection);
+		check(BufferListener);
 		BufferListeners.Remove(BufferListener);
 	}
+	void FMixerSubmix::StartEnvelopeFollowing(int32 AttackTime, int32 ReleaseTime)
+	{
+		if (!bIsEnvelopeFollowing)
+		{
+			// Zero out any previous envelope values which may have been in the array before starting up
+			for (int32 ChannelIndex = 0; ChannelIndex < AUDIO_MIXER_MAX_OUTPUT_CHANNELS; ++ChannelIndex)
+			{
+				EnvelopeValues[ChannelIndex] = 0.0f;
+				EnvelopeFollowers[ChannelIndex].Init(GetSampleRate(), AttackTime, ReleaseTime);
+			}
+
+			bIsEnvelopeFollowing = true;
+		}
+	}
+
+	void FMixerSubmix::StopEnvelopeFollowing()
+	{
+		bIsEnvelopeFollowing = false;
+	}
+
+	void FMixerSubmix::AddEnvelopeFollowerDelegate(const FOnSubmixEnvelopeBP& OnSubmixEnvelopeBP)
+	{
+		OnSubmixEnvelope.AddUnique(OnSubmixEnvelopeBP);
+	}
+
+	void FMixerSubmix::BroadcastEnvelope()
+	{
+
+		if (bIsEnvelopeFollowing)
+		{
+			// Get the envelope data
+			TArray<float> EnvelopeData;
+
+			{
+				// Make the copy of the envelope values using a critical section
+				FScopeLock EnvelopeScopeLock(&EnvelopeCriticalSection);
+
+				EnvelopeData.AddUninitialized(EnvelopeNumChannels);
+				FMemory::Memcpy(EnvelopeData.GetData(), EnvelopeValues, sizeof(float)*EnvelopeNumChannels);
+			}
+
+			// Broadcast to any bound delegates
+			if (OnSubmixEnvelope.IsBound())
+			{
+				OnSubmixEnvelope.Broadcast(EnvelopeData);
+			}
+
+		}
+	}
 
 }
\ No newline at end of file
diff --git a/Engine/Source/Runtime/AudioMixer/Private/DSP/AllPassFilter.cpp b/Engine/Source/Runtime/AudioMixer/Private/DSP/AllPassFilter.cpp
--- a/Engine/Source/Runtime/AudioMixer/Private/DSP/AllPassFilter.cpp
+++ b/Engine/Source/Runtime/AudioMixer/Private/DSP/AllPassFilter.cpp
@@ -13,7 +13,7 @@ namespace Audio
 	{
 	}
 
-	void FDelayAPF::ProcessAudio(const float* InputSample, float* OutputSample)
+	float FDelayAPF::ProcessAudioSample(const float InputSample)
 	{
 		// Read the delay line to get w(n-D);
 		const float WnD = this->Read();
@@ -21,20 +21,19 @@ namespace Audio
 		// For the APF if the delay is 0.0 we just need to pass input -> output
 		if (ReadIndex == WriteIndex)
 		{
-			this->WriteDelayAndInc(*InputSample);
-			*OutputSample = *InputSample;
-			return;
+			this->WriteDelayAndInc(InputSample);
+			return InputSample;
 		}
 
 		// Form w(n) = x(n) + gw(n-D)
-		const float Wn = *InputSample + G*WnD;
+		const float Wn = InputSample + G*WnD;
 
 		// form y(n) = -gw(n) + w(n-D)
 		float Yn = -G*Wn + WnD;
 
 		Yn = UnderflowClamp(Yn);
 		this->WriteDelayAndInc(Wn);
-		*OutputSample = Yn;
+		return Yn;
 	}
 
 }
diff --git a/Engine/Source/Runtime/AudioMixer/Private/DSP/BitCrusher.cpp b/Engine/Source/Runtime/AudioMixer/Private/DSP/BitCrusher.cpp
--- a/Engine/Source/Runtime/AudioMixer/Private/DSP/BitCrusher.cpp
+++ b/Engine/Source/Runtime/AudioMixer/Private/DSP/BitCrusher.cpp
@@ -11,20 +11,21 @@ namespace Audio
 		, BitDelta(0.0f)
 		, Phase(1.0f)
 		, PhaseDelta(1.0f)
-		, LastOutputLeft(0.0f)
-		, LastOutputRight(0.0f)
 	{
 		BitDelta = 1.0f / FMath::Pow(2.0f, BitDepth);
+		LastOutput[0] = 0.0f;
+		LastOutput[1] = 0.0f;
 	}
 
 	FBitCrusher::~FBitCrusher()
 	{
 	}
 
-	void FBitCrusher::Init(const float InSampleRate)
+	void FBitCrusher::Init(const float InSampleRate, const int32 InNumChannels)
 	{
 		SampleRate = InSampleRate;
 		Phase = 1.0f;
+		NumChannels = InNumChannels;
 	}
 
 	void FBitCrusher::SetSampleRateCrush(const float InFrequency)
@@ -38,40 +39,30 @@ namespace Audio
 		BitDelta = 1.0f / FMath::Pow(2.0f, BitDepth);
 	}
 
-	void FBitCrusher::ProcessAudio(const float InLeftSample, float& OutLeftSample)
+	void FBitCrusher::ProcessAudioFrame(const float* InFrame, float* OutFrame)
 	{
 		Phase += PhaseDelta;
-
-		// Only output audio at the sample rate set, otherwise sample and hold the last output value
-		// This filters out every N samples, etc.
 		if (Phase >= 1.0f)
 		{
 			Phase -= 1.0f;
 
-			// This quantizes the input audio to the set bit depth
-			LastOutputLeft = BitDelta * (float)FMath::FloorToInt(InLeftSample / BitDelta + 0.5f);
+			for (int32 ChannelIndex = 0; ChannelIndex < NumChannels; ++ChannelIndex)
+			{
+				LastOutput[ChannelIndex] = BitDelta * (float)FMath::FloorToInt(InFrame[ChannelIndex] / BitDelta + 0.5f);
+			}
 		}
 
-		OutLeftSample = LastOutputLeft;
+		for (int32 ChannelIndex = 0; ChannelIndex < NumChannels; ++ChannelIndex)
+		{
+			OutFrame[ChannelIndex] = LastOutput[ChannelIndex];
+		}
 	}
 
-	void FBitCrusher::ProcessAudio(const float InLeftSample, const float InRightSample, float& OutLeftSample, float& OutRightSample)
+	void FBitCrusher::ProcessAudio(const float* InBuffer, const int32 InNumSamples, float* OutBuffer)
 	{
-		Phase += PhaseDelta;
-
-		// Only output audio at the sample rate set, otherwise sample and hold the last output value
-		// This filters out every N samples, etc.
-		if (Phase >= 1.0f)
+		for (int32 SampleIndex = 0; SampleIndex < InNumSamples; SampleIndex += NumChannels)
 		{
-			Phase -= 1.0f;
-
-			// This quantizes the input audio to the set bit depth
-			LastOutputLeft = BitDelta * (float)FMath::FloorToInt(InLeftSample / BitDelta + 0.5f);
-			LastOutputRight = BitDelta * (float)FMath::FloorToInt(InRightSample / BitDelta + 0.5f);
+			ProcessAudioFrame(&InBuffer[SampleIndex], &OutBuffer[SampleIndex]);
 		}
-
-		OutLeftSample = LastOutputLeft;
-		OutRightSample = LastOutputRight;
 	}
-
 }
diff --git a/Engine/Source/Runtime/AudioMixer/Private/DSP/Chorus.cpp b/Engine/Source/Runtime/AudioMixer/Private/DSP/Chorus.cpp
--- a/Engine/Source/Runtime/AudioMixer/Private/DSP/Chorus.cpp
+++ b/Engine/Source/Runtime/AudioMixer/Private/DSP/Chorus.cpp
@@ -11,6 +11,8 @@ namespace Audio
 		, Spread(0.0f)
 		, MaxFrequencySpread(10.0f)
 		, WetLevel(0.5f)
+		, DryLevel(0.5f)
+		, NumChannels(0)
 	{
 		FMemory::Memzero(Feedback, sizeof(float)*EChorusDelays::NumDelayTypes);
 	}
@@ -19,8 +21,10 @@ namespace Audio
 	{
 	}
 
-	void FChorus::Init(const float InSampleRate, const float InBufferLengthSec, const int32 InControlSamplePeriod)
+	void FChorus::Init(const float InSampleRate, const int32 InNumChannels, const float InBufferLengthSec, const int32 InControlSamplePeriod)
 	{
+		NumChannels = InNumChannels;
+
 		for (int32 i = 0; i < EChorusDelays::NumDelayTypes; ++i)
 		{
 			Delays[i].Init(InSampleRate, InBufferLengthSec);
@@ -55,6 +59,11 @@ namespace Audio
 		WetLevel = InWetLevel;
 	}
 
+	void FChorus::SetDryLevel(const float InDryLevel)
+	{
+		DryLevel = InDryLevel;
+	}
+
 	void FChorus::SetSpread(const float InSpread)
 	{
 		Spread = FMath::Clamp(InSpread, 0.0f, 1.0f);
@@ -66,7 +75,7 @@ namespace Audio
 		LFOs[EChorusDelays::Right].Update();
 	}
 
-	void FChorus::ProcessAudio(const float InLeft, const float InRight, float& OutLeft, float& OutRight)
+	void FChorus::ProcessAudioFrame(const float* InFrame, float* OutFrame)
 	{
 		float LFONormalPhaseOut = 0.0f;
 		float LFOQuadPhaseOut = 0.0f;
@@ -102,9 +111,18 @@ namespace Audio
 
 		float DelayInputs[EChorusDelays::NumDelayTypes];
 
-		DelayInputs[EChorusDelays::Left] = InLeft;
-		DelayInputs[EChorusDelays::Center] = 0.5f * InLeft + 0.5f * InRight;
-		DelayInputs[EChorusDelays::Right] = InRight;
+		if (NumChannels == 2)
+		{
+			DelayInputs[EChorusDelays::Left] = InFrame[0];
+			DelayInputs[EChorusDelays::Center] = 0.5f * InFrame[0] + 0.5f * InFrame[1];
+			DelayInputs[EChorusDelays::Right] = InFrame[1];
+		}
+		else
+		{
+			DelayInputs[EChorusDelays::Left] = InFrame[0];
+			DelayInputs[EChorusDelays::Center] = InFrame[0];
+			DelayInputs[EChorusDelays::Right] = InFrame[0];
+		}
 
 		float DelayOutputs[EChorusDelays::NumDelayTypes];
 
@@ -115,9 +133,25 @@ namespace Audio
 			Delays[i].WriteDelayAndInc(DelayInputs[i] + DelayOutputs[i] * Feedback[i]);
 		}
 
-		const float DryLevel = (1.0f - WetLevel);
+		if (NumChannels == 2)
+		{
+			OutFrame[0] = InFrame[0] * DryLevel + WetLevel * (DelayOutputs[EChorusDelays::Left] + 0.5f * DelayOutputs[EChorusDelays::Center]);
+			OutFrame[1] = InFrame[1] * DryLevel + WetLevel * (DelayOutputs[EChorusDelays::Right] + 0.5f * DelayOutputs[EChorusDelays::Center]);
+		}
+		else
+		{
+			OutFrame[0] = InFrame[0] * DryLevel + WetLevel * (DelayOutputs[EChorusDelays::Left] + 0.5f * DelayOutputs[EChorusDelays::Center]);
+			OutFrame[0] += InFrame[0] * DryLevel + WetLevel * (DelayOutputs[EChorusDelays::Right] + 0.5f * DelayOutputs[EChorusDelays::Center]);
+			OutFrame[0] *= 0.5f;
+		}
+	}
 
-		OutLeft = InLeft * DryLevel + WetLevel * (DelayOutputs[EChorusDelays::Left] + 0.5f * DelayOutputs[EChorusDelays::Center]);
-		OutRight = InRight * DryLevel + WetLevel * (DelayOutputs[EChorusDelays::Right] + 0.5f * DelayOutputs[EChorusDelays::Center]);
+	void FChorus::ProcessAudio(const float* InBuffer, const int32 InNumSamples, float* OutBuffer)
+	{
+		for (int32 SampleIndex = 0; SampleIndex < InNumSamples; SampleIndex += NumChannels)
+		{
+			ProcessAudioFrame(&InBuffer[SampleIndex], &OutBuffer[SampleIndex]);
+		}
 	}
+
 }
diff --git a/Engine/Source/Runtime/AudioMixer/Private/DSP/Delay.cpp b/Engine/Source/Runtime/AudioMixer/Private/DSP/Delay.cpp
--- a/Engine/Source/Runtime/AudioMixer/Private/DSP/Delay.cpp
+++ b/Engine/Source/Runtime/AudioMixer/Private/DSP/Delay.cpp
@@ -158,17 +158,7 @@ namespace Audio
 		}
 	}
 
-	void FDelay::ProcessAudio(const float* InAudio, float* OutAudio)
-	{
-		Update();
-
-		const float Xn = *InAudio;
-		const float Yn = DelayInSamples == 0 ? Xn : Read();
-		WriteDelayAndInc(Xn);
-		*OutAudio = OutputAttenuation * Yn;
-	}
-
-	float FDelay::ProcessAudio(const float InAudio)
+	float FDelay::ProcessAudioSample(const float InAudio)
 	{
 		Update();
 
diff --git a/Engine/Source/Runtime/AudioMixer/Private/DSP/DelayStereo.cpp b/Engine/Source/Runtime/AudioMixer/Private/DSP/DelayStereo.cpp
--- a/Engine/Source/Runtime/AudioMixer/Private/DSP/DelayStereo.cpp
+++ b/Engine/Source/Runtime/AudioMixer/Private/DSP/DelayStereo.cpp
@@ -10,6 +10,7 @@ namespace Audio
 		, Feedback(0.0f)
 		, DelayRatio(0.0f)
 		, WetLevel(0.0f)
+		, NumChannels(0)
 		, bIsInit(true)
 	{
 	}
@@ -45,11 +46,16 @@ namespace Audio
 		WetLevel = FMath::Clamp(InWetLevel, 0.0f, 1.0f);;
 	}
 
-	void FDelayStereo::Init(const float InSampleRate, const float InDelayLengthSec)
+	void FDelayStereo::Init(const float InSampleRate, int32 InNumChannels, const float InDelayLengthSec)
 	{
+		NumChannels = InNumChannels;
+
 		// Init the delay lines
-		LeftDelay.Init(InSampleRate, 2.0f * InDelayLengthSec);
-		RightDelay.Init(InSampleRate, 2.0f * InDelayLengthSec);
+		for (int32 Channel = 0; Channel < InNumChannels; ++Channel)
+		{
+			int32 Index = Delays.Add(FDelay());
+			Delays[Index].Init(InSampleRate, 2.0f * InDelayLengthSec);
+		}
 
 		Reset();
 	}
@@ -57,50 +63,119 @@ namespace Audio
 	void FDelayStereo::Reset()
 	{
 		bIsInit = true;
-		LeftDelay.Reset();
-		RightDelay.Reset();
+		for (FDelay& Delay : Delays)
+		{
+			Delay.Reset();
+		}
 	}
 
 	void FDelayStereo::UpdateDelays()
 	{
 		// As delay ratio goes to zero, the delay times are the same
-		LeftDelay.SetEasedDelayMsec(DelayTimeMsec * (1.0f + DelayRatio), bIsInit);
-		RightDelay.SetEasedDelayMsec(DelayTimeMsec * (1.0f - DelayRatio), bIsInit);
+		for (FDelay& Delay : Delays)
+		{
+			Delay.SetEasedDelayMsec(DelayTimeMsec * (1.0f + DelayRatio), bIsInit);
+		}
 	}
 
-	void FDelayStereo::ProcessAudio(const float InLeftSample, const float InRightSample, float& OutLeftSample, float& OutRightSample)
+	void FDelayStereo::ProcessAudioFrame(const float* InFrame, float* OutFrame)
 	{
 		bIsInit = false;
 
-		float LeftDelayOut = LeftDelay.Read();
-		float RightDelayOut = RightDelay.Read();
-
-		float LeftDelayIn = 0.0f;
-		float RightDelayIn = 0.0f;
-
-		if (DelayMode == EStereoDelayMode::Normal)
+		// 1-channel audio just does a simple delay
+		if (NumChannels == 1)
 		{
-			LeftDelayIn = InLeftSample + LeftDelayOut * Feedback;
-			RightDelayIn = InRightSample + RightDelayOut * Feedback;
+			FDelay& Delay = Delays[0];
+			float DelayOut = Delay.Read();
+			float DelayIn = InFrame[0] + DelayOut * Feedback;
+			DelayOut = Delay.ProcessAudioSample(DelayIn);
+			OutFrame[0] = InFrame[0] + WetLevel * DelayOut;
 		}
-		else if (DelayMode == EStereoDelayMode::Cross)
+		else
 		{
-			LeftDelayIn = InRightSample + LeftDelayOut * Feedback;
-			RightDelayIn = InLeftSample + RightDelayOut * Feedback;
-		}
-		else if (DelayMode == EStereoDelayMode::PingPong)
-		{
-			LeftDelayIn = InRightSample + RightDelayOut * Feedback;
-			RightDelayIn = InLeftSample + LeftDelayOut * Feedback;
+			float LeftDelayOut = Delays[0].Read();
+			float RightDelayOut = Delays[1].Read();
+
+			float LeftDelayIn = 0.0f;
+			float RightDelayIn = 0.0f;
+
+			if (DelayMode == EStereoDelayMode::Normal)
+			{
+				LeftDelayIn = InFrame[0] + LeftDelayOut * Feedback;
+				RightDelayIn = InFrame[1] + RightDelayOut * Feedback;
+			}
+			else if (DelayMode == EStereoDelayMode::Cross)
+			{
+				LeftDelayIn = InFrame[1] + LeftDelayOut * Feedback;
+				RightDelayIn = InFrame[0] + RightDelayOut * Feedback;
+			}
+			else if (DelayMode == EStereoDelayMode::PingPong)
+			{
+				LeftDelayIn = InFrame[1] + RightDelayOut * Feedback;
+				RightDelayIn = InFrame[0] + LeftDelayOut * Feedback;
+			}
+
+			float WetLeftOut = 0.0f;
+			float WetRightOut = 0.0f;
+			WetLeftOut = Delays[0].ProcessAudioSample(LeftDelayIn);
+			WetRightOut = Delays[1].ProcessAudioSample(RightDelayIn);
+
+			OutFrame[0] = InFrame[0] + WetLevel * WetLeftOut;
+			OutFrame[1] = InFrame[1] + WetLevel * WetRightOut;
 		}
+	}
 
-		float WetLeftOut = 0.0f;
-		float WetRightOut = 0.0f;
-		LeftDelay.ProcessAudio(&LeftDelayIn, &WetLeftOut);
-		RightDelay.ProcessAudio(&RightDelayIn, &WetRightOut);
+	void FDelayStereo::ProcessAudio(const float* InBuffer, const int32 InNumSamples, float* OutBuffer)
+	{
+		bIsInit = false;
 
-		OutLeftSample = InLeftSample + WetLevel * WetLeftOut;
-		OutRightSample = InRightSample + WetLevel * WetRightOut;
+		if (NumChannels == 1)
+		{
+			FDelay& Delay = Delays[0];
+
+			for (int32 SampleIndex = 0; SampleIndex < InNumSamples; ++SampleIndex)
+			{
+				float DelayOut = Delay.Read();
+				float DelayIn = InBuffer[SampleIndex] + DelayOut * Feedback;
+				DelayOut = Delay.ProcessAudioSample(DelayIn);
+				OutBuffer[SampleIndex] = InBuffer[SampleIndex] + WetLevel * DelayOut;
+			}
+		}
+		else
+		{
+			for (int32 SampleIndex = 0; SampleIndex < InNumSamples; SampleIndex += NumChannels)
+			{
+				float LeftDelayOut = Delays[0].Read();
+				float RightDelayOut = Delays[1].Read();
+
+				float LeftDelayIn = 0.0f;
+				float RightDelayIn = 0.0f;
+
+				if (DelayMode == EStereoDelayMode::Normal)
+				{
+					LeftDelayIn = InBuffer[SampleIndex] + LeftDelayOut * Feedback;
+					RightDelayIn = InBuffer[SampleIndex + 1] + RightDelayOut * Feedback;
+				}
+				else if (DelayMode == EStereoDelayMode::Cross)
+				{
+					LeftDelayIn = InBuffer[SampleIndex + 1] + LeftDelayOut * Feedback;
+					RightDelayIn = InBuffer[SampleIndex] + RightDelayOut * Feedback;
+				}
+				else if (DelayMode == EStereoDelayMode::PingPong)
+				{
+					LeftDelayIn = InBuffer[SampleIndex + 1] + RightDelayOut * Feedback;
+					RightDelayIn = InBuffer[SampleIndex] + LeftDelayOut * Feedback;
+				}
+
+				float WetLeftOut = 0.0f;
+				float WetRightOut = 0.0f;
+				WetLeftOut = Delays[0].ProcessAudioSample(LeftDelayIn);
+				WetRightOut = Delays[1].ProcessAudioSample(RightDelayIn);
+
+				OutBuffer[SampleIndex] = InBuffer[SampleIndex] + WetLevel * WetLeftOut;
+				OutBuffer[SampleIndex + 1] = InBuffer[SampleIndex + 1] + WetLevel * WetRightOut;
+			}
+		}
 	}
 
 
diff --git a/Engine/Source/Runtime/AudioMixer/Private/DSP/DynamicsProcessor.cpp b/Engine/Source/Runtime/AudioMixer/Private/DSP/DynamicsProcessor.cpp
--- a/Engine/Source/Runtime/AudioMixer/Private/DSP/DynamicsProcessor.cpp
+++ b/Engine/Source/Runtime/AudioMixer/Private/DSP/DynamicsProcessor.cpp
@@ -14,6 +14,7 @@ namespace Audio
 		, HalfKneeBandwidthDb(5.0f)
 		, InputGain(1.0f)
 		, OutputGain(1.0f)
+		, NumChannels(0)
 		, bIsChannelLinked(true)
 		, bIsAnalogMode(true)
 	{
@@ -25,9 +26,10 @@ namespace Audio
 	{
 	}
 
-	void FDynamicsProcessor::Init(const float SampleRate, const int32 NumChannels)
+	void FDynamicsProcessor::Init(const float SampleRate, const int32 InNumChannels)
 	{
-		for (int32 Channel = 0; Channel < NumChannels; ++Channel)
+		NumChannels = InNumChannels;
+		for (int32 Channel = 0; Channel < InNumChannels; ++Channel)
 		{
 			int32 Index = LookaheadDelay.Add(FDelay());
 			LookaheadDelay[Index].Init(SampleRate, 0.1f);
@@ -118,16 +120,14 @@ namespace Audio
 		ProcessingMode = InProcessingMode;
 	}
 
-	void FDynamicsProcessor::ProcessAudio(const float* InputFrame, int32 NumChannels, float* OutputFrame)
+	void FDynamicsProcessor::ProcessAudioFrame(const float* InFrame, float* OutFrame)
 	{
-		check(NumChannels <= EnvFollower.Num());
-
 		// Get detector outputs
 		DetectorOuts.Reset();
 		 
 		for (int32 Channel = 0; Channel < NumChannels; ++Channel)
 		{
-			DetectorOuts.Add(EnvFollower[Channel].ProcessAudio(InputGain * InputFrame[Channel]));
+			DetectorOuts.Add(EnvFollower[Channel].ProcessAudio(InputGain * InFrame[Channel]));
 		}
 
 		Gain.Reset();
@@ -171,11 +171,18 @@ namespace Audio
 			// Write and read into the look ahead delay line.
 			// We apply the compression output of the direct input to the output of this delay line
 			// This way sharp transients can be "caught" with the gain.
-			float LookaheadOutput = 0.0f;
-			LookaheadDelay[Channel].ProcessAudio(&InputFrame[Channel], &LookaheadOutput);
+			float LookaheadOutput = LookaheadDelay[Channel].ProcessAudioSample(InFrame[Channel]);
 
 			// Write into the output with the computed gain value
-			OutputFrame[Channel] = Gain[Channel] * LookaheadOutput * OutputGain;
+			OutFrame[Channel] = Gain[Channel] * LookaheadOutput * OutputGain;
+		}
+	}
+
+	void FDynamicsProcessor::ProcessAudio(const float* InBuffer, const int32 InNumSamples, float* OutBuffer)
+	{
+		for (int32 SampleIndex = 0; SampleIndex < InNumSamples; SampleIndex += NumChannels)
+		{
+			ProcessAudioFrame(&InBuffer[SampleIndex], &OutBuffer[SampleIndex]);
 		}
 	}
 
diff --git a/Engine/Source/Runtime/AudioMixer/Private/DSP/Filter.cpp b/Engine/Source/Runtime/AudioMixer/Private/DSP/Filter.cpp
--- a/Engine/Source/Runtime/AudioMixer/Private/DSP/Filter.cpp
+++ b/Engine/Source/Runtime/AudioMixer/Private/DSP/Filter.cpp
@@ -49,21 +49,21 @@ namespace Audio
 		CalculateBiquadCoefficients();
 	}
 
-	void FBiquadFilter::ProcessAudioFrame(const float* InAudio, float* OutAudio)
+	void FBiquadFilter::ProcessAudioFrame(const float* InFrame, float* OutFrame)
 	{
-		if (bEnabled)
+		for (int32 Channel = 0; Channel < NumChannels; ++Channel)
 		{
-			for (int32 Channel = 0; Channel < NumChannels; ++Channel)
-			{
-				OutAudio[Channel] = Biquad[Channel].ProcessAudio(InAudio[Channel]);
-			}
+			OutFrame[Channel] = Biquad[Channel].ProcessAudio(InFrame[Channel]);
 		}
-		else
+	}
+
+	void FBiquadFilter::ProcessAudio(const float* InBuffer, const int32 InNumSamples, float* OutBuffer)
+	{
+		if (bEnabled)
 		{
-			// Pass through if disabled
-			for (int32 Channel = 0; Channel < NumChannels; ++Channel)
+			for (int32 SampleIndex = 0; SampleIndex < InNumSamples; SampleIndex += NumChannels)
 			{
-				OutAudio[Channel] = InAudio[Channel];
+				ProcessAudioFrame(&InBuffer[SampleIndex], &OutBuffer[SampleIndex]);
 			}
 		}
 	}
@@ -389,15 +389,62 @@ namespace Audio
 		A0 = G / (1.0f + G);
 	}
 
-	void FOnePoleFilter::ProcessAudio(const float* InSamples, float* OutSamples)
+	void FOnePoleFilter::ProcessAudioFrame(const float* InFrame, float* OutFrame)
 	{
-		for (int32 Channel = 0; Channel < NumChannels; ++Channel)
+		if (FilterType == EFilter::HighPass)
 		{
-			const float Vn = (InSamples[Channel] - Z1[Channel]) * A0;
-			const float LPF = Vn + Z1[Channel];
-			Z1[Channel] = Vn + LPF;
+			for (int32 Channel = 0; Channel < NumChannels; ++Channel)
+			{
+				const float Vn = (InFrame[Channel] - Z1[Channel]) * A0;
+				const float LPF = Vn + Z1[Channel];
+				Z1[Channel] = Vn + LPF;
 
-			OutSamples[Channel] = (FilterType == EFilter::HighPass) ? InSamples[Channel] - LPF : LPF;
+				OutFrame[Channel] = InFrame[Channel] - LPF;
+			}
+		}
+		else
+		{
+			for (int32 Channel = 0; Channel < NumChannels; ++Channel)
+			{
+				const float Vn = (InFrame[Channel] - Z1[Channel]) * A0;
+				const float LPF = Vn + Z1[Channel];
+				Z1[Channel] = Vn + LPF;
+
+				OutFrame[Channel] = LPF;
+			}
+		}
+	}
+
+	void FOnePoleFilter::ProcessAudio(const float* InSamples, const int32 InNumSamples, float* OutSamples)
+	{
+		if (FilterType == EFilter::HighPass)
+		{
+			for (int32 SampleIndex = 0; SampleIndex < InNumSamples; SampleIndex += NumChannels)
+			{
+				for (int32 Channel = 0; Channel < NumChannels; ++Channel)
+				{
+					const float InputSample = InSamples[SampleIndex + Channel];
+					const float Vn = (InputSample - Z1[Channel]) * A0;
+					const float LPF = Vn + Z1[Channel];
+					Z1[Channel] = Vn + LPF;
+
+					OutSamples[SampleIndex + Channel] = InputSample - LPF;
+				}
+			}
+		}
+		else
+		{
+			for (int32 SampleIndex = 0; SampleIndex < InNumSamples; SampleIndex += NumChannels)
+			{
+				for (int32 Channel = 0; Channel < NumChannels; ++Channel)
+				{
+					const float Vn = (InSamples[SampleIndex + Channel] - Z1[Channel]) * A0;
+					const float LPF = Vn + Z1[Channel];
+					Z1[Channel] = Vn + LPF;
+
+					OutSamples[SampleIndex + Channel] = LPF;
+				}
+			}
 		}
 	}
 
@@ -447,38 +494,41 @@ namespace Audio
 		Feedback = 2.0f * Dampening + G;
 	}
 
-	void FStateVariableFilter::ProcessAudio(const float* InSamples, float* OutSamples)
+	void FStateVariableFilter::ProcessAudio(const float* InSamples, const int32 InNumSamples, float* OutSamples)
 	{
-		for (int32 Channel = 0; Channel < NumChannels; ++Channel)
+		for (int32 SampleIndex = 0; SampleIndex < InNumSamples; SampleIndex += NumChannels)
 		{
-			const float HPF = InputScale * (InSamples[Channel] - Feedback * FilterState[Channel].Z1_1 - FilterState[Channel].Z1_2);
-			float BPF = Audio::FastTanh(A0 * HPF + FilterState[Channel].Z1_1);
+			for (int32 Channel = 0; Channel < NumChannels; ++Channel)
+			{
+				const float HPF = InputScale * (InSamples[SampleIndex + Channel] - Feedback * FilterState[Channel].Z1_1 - FilterState[Channel].Z1_2);
+				float BPF = Audio::FastTanh(A0 * HPF + FilterState[Channel].Z1_1);
 
-			const float LPF = A0 * BPF + FilterState[Channel].Z1_2;
-			const float Dampening = 0.5f / Q;
-			const float BSF = BandStopParam * HPF + (1.0f - BandStopParam) * LPF;
+				const float LPF = A0 * BPF + FilterState[Channel].Z1_2;
+				const float Dampening = 0.5f / Q;
+				const float BSF = BandStopParam * HPF + (1.0f - BandStopParam) * LPF;
 
-			FilterState[Channel].Z1_1 = A0 * HPF + BPF;
-			FilterState[Channel].Z1_2 = A0 * BPF + LPF;
+				FilterState[Channel].Z1_1 = A0 * HPF + BPF;
+				FilterState[Channel].Z1_2 = A0 * BPF + LPF;
 
-			switch (FilterType)
-			{
+				switch (FilterType)
+				{
 				default:
 				case EFilter::LowPass:
-					OutSamples[Channel] = LPF;
+					OutSamples[SampleIndex + Channel] = LPF;
 					break;
 
 				case EFilter::HighPass:
-					OutSamples[Channel] = HPF;
+					OutSamples[SampleIndex + Channel] = HPF;
 					break;
 
 				case EFilter::BandPass:
-					OutSamples[Channel] = BPF;
+					OutSamples[SampleIndex + Channel] = BPF;
 					break;
 
 				case EFilter::BandStop:
-					OutSamples[Channel] = BSF;
+					OutSamples[SampleIndex + Channel] = BSF;
 					break;
+				}
 			}
 		}
 	}
@@ -582,42 +632,48 @@ namespace Audio
 		PassBandGainCompensation = InPassBandGainCompensation;
 	}
 
-	void FLadderFilter::ProcessAudio(const float* InSamples, float* OutSamples)
+	void FLadderFilter::ProcessAudio(const float* InSamples, const int32 InNumSamples, float* OutSamples)
 	{
-		// Compute input into first LPF
 		float U[MaxFilterChannels];
 
-		for (int32 Channel = 0; Channel < NumChannels; ++Channel)
-		{
-			float Sigma = 0.0f;
-			for (int32 i = 0; i < 4; ++i)
-			{
-				Sigma += OnePoleFilters[i].GetState(Channel) * Beta[i];
-			}
-
-			float InSample = InSamples[Channel];
-			InSample *= 1.0f + PassBandGainCompensation * K;
-			U[Channel] = FMath::Min(Audio::FastTanh((InSample - K * Sigma)* Alpha), 1.0f);
-		}
-
-		// Feed U into first filter, then cascade down
 		float OutputFilter0[MaxFilterChannels];
 		float OutputFilter1[MaxFilterChannels];
 		float OutputFilter2[MaxFilterChannels];
 		float OutputFilter3[MaxFilterChannels];
 
-		OnePoleFilters[0].ProcessAudio(U, OutputFilter0);
-		OnePoleFilters[1].ProcessAudio(OutputFilter0, OutputFilter1);
-		OnePoleFilters[2].ProcessAudio(OutputFilter1, OutputFilter2);
-		OnePoleFilters[3].ProcessAudio(OutputFilter2, OutputFilter3);
-
-		for (int32 Channel = 0; Channel < NumChannels; ++Channel)
+		for (int32 SampleIndex = 0; SampleIndex < InNumSamples; SampleIndex += NumChannels)
 		{
-			OutSamples[Channel] =  Factors[0] * U[Channel];
-			OutSamples[Channel] += Factors[1] * OutputFilter0[Channel];
-			OutSamples[Channel] += Factors[2] * OutputFilter1[Channel];
-			OutSamples[Channel] += Factors[3] * OutputFilter2[Channel];
-			OutSamples[Channel] += Factors[4] * OutputFilter3[Channel];
+			for (int32 Channel = 0; Channel < NumChannels; ++Channel)
+			{
+				float Sigma = 0.0f;
+				for (int32 i = 0; i < 4; ++i)
+				{
+					Sigma += OnePoleFilters[i].GetState(Channel) * Beta[i];
+				}
+
+				float InSample = InSamples[SampleIndex + Channel];
+				InSample *= 1.0f + PassBandGainCompensation * K;
+
+				// Compute input into first LPF
+				U[Channel] = FMath::Min(Audio::FastTanh((InSample - K * Sigma)* Alpha), 1.0f);
+			}
+
+			OnePoleFilters[0].ProcessAudioFrame(U, OutputFilter0);
+			OnePoleFilters[1].ProcessAudioFrame(OutputFilter0, OutputFilter1);
+			OnePoleFilters[2].ProcessAudioFrame(OutputFilter1, OutputFilter2);
+			OnePoleFilters[3].ProcessAudioFrame(OutputFilter2, OutputFilter3);
+
+			for (int32 Channel = 0; Channel < NumChannels; ++Channel)
+			{
+				const int32 OutputSampleIndex = SampleIndex + Channel;
+
+				// Feed U into first filter, then cascade down
+				OutSamples[OutputSampleIndex] = Factors[0] * U[Channel];
+				OutSamples[OutputSampleIndex] += Factors[1] * OutputFilter0[Channel];
+				OutSamples[OutputSampleIndex] += Factors[2] * OutputFilter1[Channel];
+				OutSamples[OutputSampleIndex] += Factors[3] * OutputFilter2[Channel];
+				OutSamples[OutputSampleIndex] += Factors[4] * OutputFilter3[Channel];
+			}
 		}
 	}
 
diff --git a/Engine/Source/Runtime/AudioMixer/Private/DSP/FoldbackDistortion.cpp b/Engine/Source/Runtime/AudioMixer/Private/DSP/FoldbackDistortion.cpp
--- a/Engine/Source/Runtime/AudioMixer/Private/DSP/FoldbackDistortion.cpp
+++ b/Engine/Source/Runtime/AudioMixer/Private/DSP/FoldbackDistortion.cpp
@@ -10,6 +10,7 @@ namespace Audio
 		, Threshold2(2.0f * Threshold)
 		, Threshold4(4.0f * Threshold)
 		, OutputGain(1.0f)
+		, NumChannels(0)
 	{
 	}
 
@@ -17,9 +18,9 @@ namespace Audio
 	{
 	}
 
-	void FFoldbackDistortion::Init(const float InSampleRate)
+	void FFoldbackDistortion::Init(const float InSampleRate, const int32 InNumChannels)
 	{
-		// nothing to do...
+		NumChannels = InNumChannels;
 	}
 
 	void FFoldbackDistortion::SetThresholdDb(const float InThresholdDb)
@@ -39,9 +40,10 @@ namespace Audio
 		OutputGain = Audio::ConvertToLinear(InOutputGainDb);
 	}
 
-	void FFoldbackDistortion::ProcessAudio(const float InSample, float& OutSample)
+	float FFoldbackDistortion::ProcessAudioSample(const float InSample)
 	{
-		float Sample = InputGain * InSample;
+		const float Sample = InputGain * InSample;
+		float OutSample = 0.0f;
 
 		if (Sample > Threshold || Sample < -Threshold)
 		{
@@ -52,13 +54,23 @@ namespace Audio
 			OutSample = Sample;
 		}
 
-		OutSample *= OutputGain;
+		return OutSample * OutputGain;
 	}
 
-	void FFoldbackDistortion::ProcessAudio(const float InLeftSample, const float InRightSample, float& OutLeftSample, float& OutRightSample)
+	void FFoldbackDistortion::ProcessAudioFrame(const float* InFrame, float* OutFrame)
 	{
-		ProcessAudio(InLeftSample, OutLeftSample);
-		ProcessAudio(InRightSample, OutRightSample);
+		for (int32 Channel = 0; Channel < NumChannels; ++Channel)
+		{
+			OutFrame[Channel] = ProcessAudioSample(InFrame[Channel]);
+		}
+	}
+
+	void FFoldbackDistortion::ProcessAudio(const float* InBuffer, const int32 InNumSamples, float* OutBuffer)
+	{
+		for (int32 SampleIndex = 0; SampleIndex < InNumSamples; SampleIndex += NumChannels)
+		{
+			ProcessAudioFrame(&InBuffer[SampleIndex], &OutBuffer[SampleIndex]);
+		}
 	}
 
 }
diff --git a/Engine/Source/Runtime/AudioMixer/Private/DSP/Phaser.cpp b/Engine/Source/Runtime/AudioMixer/Private/DSP/Phaser.cpp
--- a/Engine/Source/Runtime/AudioMixer/Private/DSP/Phaser.cpp
+++ b/Engine/Source/Runtime/AudioMixer/Private/DSP/Phaser.cpp
@@ -11,6 +11,7 @@ namespace Audio
 		, WetLevel(0.4f)
 		, Feedback(0.2f)
 		, LFOType(ELFO::Sine)
+		, NumChannels(0)
 		, bIsBiquadPhase(true)
 	{
 	}
@@ -20,8 +21,10 @@ namespace Audio
 
 	}
 
-	void FPhaser::Init(const float SampleRate)
+	void FPhaser::Init(const float SampleRate, const int32 InNumChannels)
 	{
+		check(NumChannels <= MaxNumChannels);
+		NumChannels = InNumChannels;
 
 		// Initialize all the APFs
 		for (int32 Channel = 0; Channel < NumChannels; ++Channel)
@@ -119,9 +122,8 @@ namespace Audio
 		}
 	}
 
-	void FPhaser::ProcessAudio(const float* InFrame, float* OutFrame)
+	void FPhaser::ProcessAudioFrame(const float* InFrame, float* OutFrame)
 	{
-
 		ControlSampleCount = ControlSampleCount & (ControlRate - 1);
 		if (ControlSampleCount == 0)
 		{
@@ -163,6 +165,14 @@ namespace Audio
 		}
 	}
 
+	void FPhaser::ProcessAudio(const float* InBuffer, const int32 InNumSamples, float* OutBuffer)
+	{
+		for (int32 SampleIndex = 0; SampleIndex < InNumSamples; SampleIndex += NumChannels)
+		{
+			ProcessAudioFrame(&InBuffer[SampleIndex], &OutBuffer[SampleIndex]);
+		}
+	}
+
 
 
 }
diff --git a/Engine/Source/Runtime/AudioMixer/Private/DSP/Reverb.cpp b/Engine/Source/Runtime/AudioMixer/Private/DSP/Reverb.cpp
--- a/Engine/Source/Runtime/AudioMixer/Private/DSP/Reverb.cpp
+++ b/Engine/Source/Runtime/AudioMixer/Private/DSP/Reverb.cpp
@@ -108,10 +108,8 @@ namespace Audio
 
 	float FEarlyReflections::ProcessDelayLine(const float InSample, FDelayAPF& InAPF, FOnePoleLPF& InLPF)
 	{
-		float APFOut = 0.0f;
-		InAPF.ProcessAudio(&InSample, &APFOut);
-		float LPFOut = 0.0f;
-		InLPF.ProcessAudio(&APFOut, &LPFOut);
+		float APFOut = InAPF.ProcessAudioSample(InSample);
+		float LPFOut = InLPF.ProcessAudioSample(APFOut);
 		return LPFOut;
 	}
 
@@ -129,14 +127,12 @@ namespace Audio
 			float InputSample = InBuffer[Channel];
 
 			// Input -> PreDelay
-			float PreDelayOut = 0.0f;
-			Data[Channel].PreDelay.ProcessAudio(&InputSample, &PreDelayOut);
+			float PreDelayOut = Data[Channel].PreDelay.ProcessAudioSample(InputSample);
 
 			// PreDelay -> InputLPF
 			PreDelayOut *= Settings.Bandwidth;
 
-			float InputLPFOut = 0.0f;
-			Data[Channel].InputLPF.ProcessAudio(&PreDelayOut, &InputLPFOut);
+			float InputLPFOut = Data[Channel].InputLPF.ProcessAudioSample(PreDelayOut);
 
 			// Process each delay line input
 
@@ -330,31 +326,25 @@ namespace Audio
 			// INPUT DIFFUSION
 
 			// InputSample -> PreDelay
-			float PreDelayOut = 0.0f;
-			PreDelay.ProcessAudio(&InputSample, &PreDelayOut);
+			float PreDelayOut = PreDelay.ProcessAudioSample(InputSample);
 
 			// Scale by bandwidth
 			PreDelayOut *= Settings.Bandwidth;
 
 			// PreDelay -> InputLPFOut
-			float InputLPFOut = 0.0f;
-			InputLPF.ProcessAudio(&PreDelayOut, &InputLPFOut);
+			float InputLPFOut = InputLPF.ProcessAudioSample(PreDelayOut);
 
 			// InputLPFOut -> APF1
-			float APF1Out = 0.0f;
-			APF1.ProcessAudio(&InputLPFOut, &APF1Out);
+			float APF1Out = APF1.ProcessAudioSample(InputLPFOut);
 
 			// APF1 -> APF2
-			float APF2Out = 0.0f;
-			APF2.ProcessAudio(&APF1Out, &APF2Out);
+			float APF2Out = APF2.ProcessAudioSample(APF1Out);;
 
 			// APF2 -> APF3
-			float APF3Out = 0.0f;
-			APF3.ProcessAudio(&APF2Out, &APF3Out);
+			float APF3Out = APF3.ProcessAudioSample(APF2Out);
 
 			// APF3 -> APF4
-			float APF4Out = 0.0f;
-			APF4.ProcessAudio(&APF3Out, &APF4Out);
+			float APF4Out = APF4.ProcessAudioSample(APF3Out);
 
 			// -------------------
 			// Modulation
@@ -378,30 +368,25 @@ namespace Audio
 			float RightPlateInput = APF4Out + UnderflowClamp(LeftPlate.PreviousSample);
 
 			// Input -> ModulatedAPF
-			float RightPlateModulatedAPFOut = 0.0f;
-			RightPlate.ModulatedAPF.ProcessAudio(&RightPlateInput, &RightPlateModulatedAPFOut);
+			float RightPlateModulatedAPFOut = RightPlate.ModulatedAPF.ProcessAudioSample(RightPlateInput);
 
 			// ModulatedAPF -> Delay1
-			float RightPlateDelay1Out = 0.0f;
-			RightPlate.Delay1.ProcessAudio(&RightPlateModulatedAPFOut, &RightPlateDelay1Out);
+			float RightPlateDelay1Out = RightPlate.Delay1.ProcessAudioSample(RightPlateModulatedAPFOut);
 
 			// Apply dampening
 			RightPlateDelay1Out *= (1.0f - Settings.Dampening);
 
 			// Delay1 -> LPF
-			float RightPlateLPFOut = 0.0f;
-			RightPlate.LPF.ProcessAudio(&RightPlateDelay1Out, &RightPlateLPFOut);
+			float RightPlateLPFOut = RightPlate.LPF.ProcessAudioSample(RightPlateDelay1Out);
 
 			// Apply decay
 			RightPlateLPFOut *= (1.0f - Settings.Decay);
 
 			// LPF -> APF
-			float RightPlateAPFOut = 0.0f;
-			RightPlate.APF.ProcessAudio(&RightPlateLPFOut, &RightPlateAPFOut);
+			float RightPlateAPFOut = RightPlate.APF.ProcessAudioSample(RightPlateLPFOut);
 
 			// APF-> Delay2
-			float RightPlateDelay2Out = 0.0f;
-			RightPlate.Delay2.ProcessAudio(&RightPlateAPFOut, &RightPlateDelay2Out);
+			float RightPlateDelay2Out = RightPlate.Delay2.ProcessAudioSample(RightPlateAPFOut);
 
 			// Write to PreviousSample for feedback
 			RightPlate.PreviousSample = (1.0f - Settings.Decay) * RightPlateDelay2Out;
@@ -413,30 +398,25 @@ namespace Audio
 			float LeftPlateInput = APF4Out + UnderflowClamp(RightPlate.PreviousSample);
 
 			// Input -> ModulatedAPF
-			float LeftPlateModulatedAPFOut = 0.0f;
-			LeftPlate.ModulatedAPF.ProcessAudio(&LeftPlateInput, &LeftPlateModulatedAPFOut);
+			float LeftPlateModulatedAPFOut = LeftPlate.ModulatedAPF.ProcessAudioSample(LeftPlateInput);
 
 			// ModulatedAPF -> Delay1
-			float LeftPlateDelay1Out = 0.0f;
-			LeftPlate.Delay1.ProcessAudio(&LeftPlateModulatedAPFOut, &LeftPlateDelay1Out);
+			float LeftPlateDelay1Out = LeftPlate.Delay1.ProcessAudioSample(LeftPlateModulatedAPFOut);
 
 			// Apply dampening
 			LeftPlateDelay1Out *= (1.0f - Settings.Dampening);
 
 			// Delay1 -> LPF
-			float LeftPlateLPFOut = 0.0f;
-			LeftPlate.LPF.ProcessAudio(&LeftPlateDelay1Out, &LeftPlateLPFOut);
+			float LeftPlateLPFOut = LeftPlate.LPF.ProcessAudioSample(LeftPlateDelay1Out);
 
 			// Apply decay
 			LeftPlateLPFOut *= (1.0f - Settings.Decay);
 
 			// LPF -> APF
-			float LeftPlateAPFOut = 0.0f;
-			LeftPlate.APF.ProcessAudio(&LeftPlateLPFOut, &LeftPlateAPFOut);
-
+			float LeftPlateAPFOut = LeftPlate.APF.ProcessAudioSample(LeftPlateLPFOut);
+			
 			// APF-> Delay2
-			float LeftPlateDelay2Out = 0.0f;
-			LeftPlate.Delay2.ProcessAudio(&LeftPlateAPFOut, &LeftPlateDelay2Out);
+			float LeftPlateDelay2Out = LeftPlate.Delay2.ProcessAudioSample(LeftPlateAPFOut);
 
 			// Write to PreviousSample for feedback
 			LeftPlate.PreviousSample = (1.0f - Settings.Decay) * LeftPlateDelay2Out;
diff --git a/Engine/Source/Runtime/AudioMixer/Private/DSP/RingModulation.cpp b/Engine/Source/Runtime/AudioMixer/Private/DSP/RingModulation.cpp
--- a/Engine/Source/Runtime/AudioMixer/Private/DSP/RingModulation.cpp
+++ b/Engine/Source/Runtime/AudioMixer/Private/DSP/RingModulation.cpp
@@ -8,6 +8,9 @@ namespace Audio
 	FRingModulation::FRingModulation()
 		: ModulationFrequency(800.0f)
 		, ModulationDepth(0.5f)
+		, DryLevel(0.0f)
+		, WetLevel(1.0f)
+		, NumChannels(0)
 	{
 
 	}
@@ -17,12 +20,14 @@ namespace Audio
 
 	}
 
-	void FRingModulation::Init(const float InSampleRate)
+	void FRingModulation::Init(const float InSampleRate, const int32 InNumChannels)
 	{
 		Osc.Init(InSampleRate);
 		Osc.SetFrequency(ModulationFrequency);
 		Osc.Update();
 		Osc.Start();
+
+		NumChannels = InNumChannels;
 	}
 
 	void FRingModulation::SetModulatorWaveType(const EOsc::Type InType)
@@ -41,11 +46,22 @@ namespace Audio
 		ModulationDepth = FMath::Clamp(InModulationDepth, -1.0f, 1.0f);
 	}
 
-	void FRingModulation::ProcessAudio(const float InLeftSample, const float InRightSample, float& OutLeftSample, float& OutRightSample)
+	void FRingModulation::ProcessAudioFrame(const float* InFrame, float* OutFrame)
 	{
 		float OscOut = Osc.Generate();
-		OutLeftSample = InLeftSample * OscOut * ModulationDepth;
-		OutRightSample = InRightSample * OscOut * ModulationDepth;
+		for (int32 Channel = 0; Channel < NumChannels; ++Channel)
+		{
+			OutFrame[Channel] = DryLevel * InFrame[Channel] + WetLevel * InFrame[Channel] * OscOut * ModulationDepth;
+		}
 	}
 
+	void FRingModulation::ProcessAudio(const float* InBuffer, const int32 InNumSamples, float* OutBuffer)
+	{
+		for (int32 SampleIndex = 0; SampleIndex < InNumSamples; SampleIndex += NumChannels)
+		{
+			ProcessAudioFrame(&InBuffer[SampleIndex], &OutBuffer[SampleIndex]);
+		}
+	}
+
+
 }
diff --git a/Engine/Source/Runtime/AudioMixer/Private/Effects/AudioMixerSubmixEffectReverb.cpp b/Engine/Source/Runtime/AudioMixer/Private/Effects/AudioMixerSubmixEffectReverb.cpp
--- a/Engine/Source/Runtime/AudioMixer/Private/Effects/AudioMixerSubmixEffectReverb.cpp
+++ b/Engine/Source/Runtime/AudioMixer/Private/Effects/AudioMixerSubmixEffectReverb.cpp
@@ -9,7 +9,8 @@
 class UReverbEffect;
 
 FSubmixEffectReverb::FSubmixEffectReverb()
-	: bIsEnabled(false)
+	: DryLevel(0.0f)
+	, bIsEnabled(false)
 {
 }
 
@@ -26,6 +27,8 @@ void FSubmixEffectReverb::Init(const FSoundEffectSubmixInitData& InitData)
 	NewSettings.Density = 0.8f;
 	NewSettings.Wetness = 1.0f;
 
+	DryLevel = 0.0f;
+
 	Params.SetParams(NewSettings);
 
 	PlateReverb.Init(InitData.SampleRate);
@@ -59,14 +62,16 @@ void FSubmixEffectReverb::OnPresetChanged()
 	ReverbEffect.AirAbsorptionGainHF = Settings.AirAbsorptionGainHF;
 	ReverbEffect.RoomRolloffFactor = 0.0f; // not used
 	ReverbEffect.Volume = Settings.WetLevel;
-		
+
+	DryLevel = Settings.DryLevel;
+
 	SetEffectParameters(ReverbEffect);
 }
 
 void FSubmixEffectReverb::OnProcessAudio(const FSoundEffectSubmixInputData& InData, FSoundEffectSubmixOutputData& OutData)
 {
 	check(InData.NumChannels == 2);
-	if (OutData.NumChannels < 2 || !bIsEnabled) 
+ 	if (OutData.NumChannels < 2 || !bIsEnabled) 
 	{
 		// Not supported
 		return;
@@ -85,6 +90,8 @@ void FSubmixEffectReverb::OnProcessAudio(const FSoundEffectSubmixInputData& InDa
 		for (int32 SampleIndex = 0; SampleIndex < InData.AudioBuffer->Num(); SampleIndex += OutData.NumChannels)
 		{
 			PlateReverb.ProcessAudioFrame(&AudioData[SampleIndex], InData.NumChannels, &OutAudioData[SampleIndex], OutData.NumChannels);
+
+			OutAudioData[SampleIndex] += DryLevel * AudioData[SampleIndex];
 		}
 	}
 	// 5.1 or higher surround sound. Map stereo output to quad output
@@ -95,6 +102,8 @@ void FSubmixEffectReverb::OnProcessAudio(const FSoundEffectSubmixInputData& InDa
 			// Processed downmixed audio frame
 			PlateReverb.ProcessAudioFrame(&AudioData[InSampleIndex], InData.NumChannels, &OutAudioData[OutSampleIndex], InData.NumChannels);
 
+			OutAudioData[OutSampleIndex] += DryLevel * AudioData[InSampleIndex];
+
 			// Now do a cross-over to the back-left/back-right speakers from the front-left and front-right
 			
 			// Using standard speaker map order map the right output to the BackLeft channel
@@ -142,7 +151,7 @@ void FSubmixEffectReverb::UpdateParameters()
 	}
 }
 
-void USubmixEffectReverbPreset::SetSettingsWithReverbEffect(const UReverbEffect* InReverbEffect, const float WetLevel)
+void USubmixEffectReverbPreset::SetSettingsWithReverbEffect(const UReverbEffect* InReverbEffect, const float WetLevel, const float DryLevel)
 {
 	if (InReverbEffect)
 	{
@@ -158,6 +167,7 @@ void USubmixEffectReverbPreset::SetSettingsWithReverbEffect(const UReverbEffect*
 		Settings.LateDelay = InReverbEffect->LateDelay;
 		Settings.AirAbsorptionGainHF = InReverbEffect->AirAbsorptionGainHF;
 		Settings.WetLevel = WetLevel;
+		Settings.DryLevel = DryLevel;
 
 		Update();
 	}
diff --git a/Engine/Source/Runtime/AudioMixer/Public/AudioMixerDevice.h b/Engine/Source/Runtime/AudioMixer/Public/AudioMixerDevice.h
--- a/Engine/Source/Runtime/AudioMixer/Public/AudioMixerDevice.h
+++ b/Engine/Source/Runtime/AudioMixer/Public/AudioMixerDevice.h
@@ -79,6 +79,7 @@ namespace Audio
 		virtual void FadeOut() override;
 		virtual void TeardownHardware() override;
 		virtual void UpdateHardwareTiming() override;
+		virtual void UpdateGameThread() override;
 		virtual void UpdateHardware() override;
 		virtual double GetAudioTime() const override;
 		virtual FAudioEffectsManager* CreateEffectsManager() override;
@@ -114,6 +115,11 @@ namespace Audio
 		virtual void PauseRecording(USoundSubmix* InSubmix);
 		virtual void ResumeRecording(USoundSubmix* InSubmix);
 
+		/** Submix envelope following */
+		virtual void StartEnvelopeFollowing(USoundSubmix* InSubmix) override;
+		virtual void StopEnvelopeFollowing(USoundSubmix* InSubmix) override;
+		virtual void AddEnvelopeFollowerDelegate(USoundSubmix* InSubmix, const FOnSubmixEnvelopeBP& OnSubmixEnvelopeBP) override;
+
 		// Submix buffer listener callbacks
 		virtual void RegisterSubmixBufferListener(ISubmixBufferListener* InSubmixBufferListener, USoundSubmix* InSubmix = nullptr) override;
 		virtual void UnregisterSubmixBufferListener(ISubmixBufferListener* InSubmixBufferListener, USoundSubmix* InSubmix = nullptr) override;
@@ -124,7 +130,7 @@ namespace Audio
 		virtual void OnAudioStreamShutdown() override;
 		//~ End IAudioMixer
 
-		FMixerSubmixPtr GetSubmixInstance(USoundSubmix* SoundSubmix);
+		FMixerSubmixWeakPtr GetSubmixInstance(USoundSubmix* SoundSubmix);
 
 		// Functions which check the thread it's called on and helps make sure functions are called from correct threads
 		void CheckAudioThread();
@@ -150,11 +156,11 @@ namespace Audio
 
 		FMixerSourceManager* GetSourceManager();
 
-		FMixerSubmixPtr GetMasterSubmix(); 
-		FMixerSubmixPtr GetMasterReverbSubmix();
-		FMixerSubmixPtr GetMasterReverbPluginSubmix();
-		FMixerSubmixPtr GetMasterEQSubmix();
-		FMixerSubmixPtr GetMasterAmbisonicsSubmix();
+		FMixerSubmixWeakPtr GetMasterSubmix(); 
+		FMixerSubmixWeakPtr GetMasterReverbSubmix();
+		FMixerSubmixWeakPtr GetMasterReverbPluginSubmix();
+		FMixerSubmixWeakPtr GetMasterEQSubmix();
+		FMixerSubmixWeakPtr GetMasterAmbisonicsSubmix();
 
 		// Add submix effect to master submix
 		void AddMasterSubmixEffect(uint32 SubmixEffectId, FSoundEffectSubmix* SoundEffect);
@@ -248,6 +254,9 @@ namespace Audio
 		/** Map of USoundSubmix static data objects to the dynamic audio mixer submix. */
 		TMap<USoundSubmix*, FMixerSubmixPtr> Submixes;
 
+		/** Which submixes have been told to envelope follow with this audio device. */
+		TArray<USoundSubmix*> EnvelopeFollowingSubmixes;
+
 		/** Queue of mixer source voices. */
 		TQueue<FMixerSourceVoice*> SourceVoices;
 
diff --git a/Engine/Source/Runtime/AudioMixer/Public/AudioMixerSubmix.h b/Engine/Source/Runtime/AudioMixer/Public/AudioMixerSubmix.h
--- a/Engine/Source/Runtime/AudioMixer/Public/AudioMixerSubmix.h
+++ b/Engine/Source/Runtime/AudioMixer/Public/AudioMixerSubmix.h
@@ -5,6 +5,7 @@
 #include "AudioMixer.h"
 #include "Sound/SoundSubmix.h"
 #include "Sound/SampleBuffer.h"
+#include "DSP/EnvelopeFollower.h"
 
 class USoundEffectSubmix;
 
@@ -14,8 +15,6 @@ namespace Audio
 	class FMixerSourceVoice;
 	class FMixerDevice;
 
-	typedef TSharedPtr<FSoundEffectSubmix, ESPMode::ThreadSafe> FSoundEffectSubmixPtr;
-
 	struct FSubmixVoiceData
 	{
 		float SendLevel;
@@ -33,7 +32,7 @@ namespace Audio
 
 	struct FChildSubmixInfo
 	{
-		TSharedPtr<FMixerSubmix, ESPMode::ThreadSafe> SubmixPtr;
+		TWeakPtr<FMixerSubmix, ESPMode::ThreadSafe> SubmixPtr;
 		bool bNeedsAmbisonicsEncoding;
 
 		FChildSubmixInfo()
@@ -55,16 +54,16 @@ namespace Audio
 		uint32 GetId() const { return Id; }
 
 		// Sets the parent submix to the given submix
-		void SetParentSubmix(TSharedPtr<FMixerSubmix, ESPMode::ThreadSafe> Submix);
+		void SetParentSubmix(TWeakPtr<FMixerSubmix, ESPMode::ThreadSafe> Submix);
 
 		// Adds the given submix to this submix's children
-		void AddChildSubmix(TSharedPtr<FMixerSubmix, ESPMode::ThreadSafe> Submix);
+		void AddChildSubmix(TWeakPtr<FMixerSubmix, ESPMode::ThreadSafe> Submix);
 
 		// Gets the submix channels channels
 		ESubmixChannelFormat GetSubmixChannels() const;
 
 		// Gets this submix's parent submix
-		TSharedPtr<FMixerSubmix, ESPMode::ThreadSafe> GetParentSubmix();
+		TWeakPtr<FMixerSubmix, ESPMode::ThreadSafe> GetParentSubmix();
 
 		// Returns the number of source voices currently a part of this submix.
 		int32 GetNumSourceVoices() const;
@@ -79,7 +78,7 @@ namespace Audio
 		void RemoveSourceVoice(FMixerSourceVoice* InSourceVoice);
 
 		/** Appends the effect submix to the effect submix chain. */
-		void AddSoundEffectSubmix(uint32 SubmixPresetId, FSoundEffectSubmixPtr InSoundEffectSubmix);
+		void AddSoundEffectSubmix(uint32 SubmixPresetId, FSoundEffectSubmix* InSoundEffectSubmix);
 
 		/** Removes the submix effect from the effect submix chain. */
 		void RemoveSoundEffectSubmix(uint32 SubmixPresetId);
@@ -103,7 +102,7 @@ namespace Audio
 		int32 GetNumChainEffects() const;
 
 		// Returns the submix effect at the given effect chain index
-		FSoundEffectSubmixPtr GetSubmixEffect(const int32 InIndex);
+		FSoundEffectSubmix* GetSubmixEffect(const int32 InIndex);
 
 		// updates settings, potentially creating or removing ambisonics streams based on
 		void OnAmbisonicsSettingsChanged(UAmbisonicsSubmixSettingsBase* AmbisonicsSettings);
@@ -126,6 +125,18 @@ namespace Audio
 		// Unregister buffer listener with this submix
 		void UnregisterBufferListener(ISubmixBufferListener* BufferListener);
 
+		// Starts envelope following with the given attack time and release time
+		void StartEnvelopeFollowing(int32 AttackTime, int32 ReleaseTime);
+
+		// Stops envelope following the submix
+		void StopEnvelopeFollowing();
+
+		// Adds an envelope follower delegate
+		void AddEnvelopeFollowerDelegate(const FOnSubmixEnvelopeBP& OnSubmixEnvelopeBP);
+
+		// Broadcast the envelope value on the game thread
+		void BroadcastEnvelope();
+
 	protected:
 		// Down mix the given buffer to the desired down mix channel count
 		void FormatChangeBuffer(const ESubmixChannelFormat NewChannelType, AlignedFloatBuffer& InBuffer, AlignedFloatBuffer& OutNewBuffer);
@@ -173,7 +184,7 @@ namespace Audio
 		uint32 Id;
 
 		// Parent submix. 
-		TSharedPtr<FMixerSubmix, ESPMode::ThreadSafe> ParentSubmix;
+		TWeakPtr<FMixerSubmix, ESPMode::ThreadSafe> ParentSubmix;
 
 		// Child submixes
 		TMap<uint32, FChildSubmixInfo> ChildSubmixes;
@@ -185,10 +196,11 @@ namespace Audio
 			uint32 PresetId;
 
 			// The effect instance ptr
-			FSoundEffectSubmixPtr EffectInstance;
+			FSoundEffectSubmix* EffectInstance;
 
 			FSubmixEffectInfo()
 				: PresetId(INDEX_NONE)
+				, EffectInstance(nullptr)
 			{
 			}
 		};
@@ -220,6 +232,12 @@ namespace Audio
 		// Decoder ID set up with Ambisonics Mixer. Set to INDEX_NONE if there is no decoder stream open.
 		uint32 SubmixAmbisonicsDecoderID;
 
+		// Envelope following data
+		float EnvelopeValues[AUDIO_MIXER_MAX_OUTPUT_CHANNELS];
+		Audio::FEnvelopeFollower EnvelopeFollowers[AUDIO_MIXER_MAX_OUTPUT_CHANNELS];
+		int32 EnvelopeNumChannels;
+		FCriticalSection EnvelopeCriticalSection;
+
 		// This buffer is encoded into for each source, then summed into the ambisonics buffer.
 		AlignedFloatBuffer InputAmbisonicsBuffer;
 
@@ -241,6 +259,12 @@ namespace Audio
 		// Bool set to true when this submix is recording data.
 		uint8 bIsRecording : 1;
 
+		// Bool set to true when envelope following is enabled
+		FThreadSafeBool bIsEnvelopeFollowing;
+
+		// Multi-cast delegate to broadcast envelope data from this submix instance
+		FOnSubmixEnvelope OnSubmixEnvelope;
+
 		// Critical section used for when we are appending recorded data.
 		FCriticalSection RecordingCriticalSection;
 
diff --git a/Engine/Source/Runtime/AudioMixer/Public/DSP/AllPassFilter.h b/Engine/Source/Runtime/AudioMixer/Public/DSP/AllPassFilter.h
--- a/Engine/Source/Runtime/AudioMixer/Public/DSP/AllPassFilter.h
+++ b/Engine/Source/Runtime/AudioMixer/Public/DSP/AllPassFilter.h
@@ -22,7 +22,7 @@ namespace Audio
 		void SetG(float InG) { G = InG; }
 
 		// overrides
-		virtual void ProcessAudio(const float* pInput, float* pOutput) override;
+		virtual float ProcessAudioSample(const float pInput) override;
 
 	protected:
 		// Feedback/Feedforward gain coefficient
diff --git a/Engine/Source/Runtime/AudioMixer/Public/DSP/BitCrusher.h b/Engine/Source/Runtime/AudioMixer/Public/DSP/BitCrusher.h
--- a/Engine/Source/Runtime/AudioMixer/Public/DSP/BitCrusher.h
+++ b/Engine/Source/Runtime/AudioMixer/Public/DSP/BitCrusher.h
@@ -19,7 +19,7 @@ namespace Audio
 		~FBitCrusher();
 
 		// Initialize the equalizer
-		void Init(const float InSampleRate);
+		void Init(const float InSampleRate, const int32 InNumChannels);
 
 		// The amount to reduce the sample rate of the audio stream.
 		void SetSampleRateCrush(const float InFrequency);
@@ -27,11 +27,9 @@ namespace Audio
 		// The amount to reduce the bit depth of the audio stream.
 		void SetBitDepthCrush(const float InBitDepth);
 
-		// Processes a mono stream
-		void ProcessAudio(const float InLeftSample, float& OutLeftSample);
-
-		// Processes a stereo stream
-		void ProcessAudio(const float InLeftSample, const float InRightSample, float& OutLeftSample, float& OutRightSample);
+		// Processes audio
+		void ProcessAudioFrame(const float* InFrame, float* OutFrame);
+		void ProcessAudio(const float* InBuffer, const int32 InNumSamples, float* OutBuffer);
 
 	private:
 		// i.e. 8 bit, etc. But can be float!
@@ -46,9 +44,9 @@ namespace Audio
 		float PhaseDelta;
 
 		// Used to sample+hold the last output
-		float LastOutputLeft;
-		float LastOutputRight;
+		float LastOutput[2];
 
+		int32 NumChannels;
 	};
 
 }
diff --git a/Engine/Source/Runtime/AudioMixer/Public/DSP/Chorus.h b/Engine/Source/Runtime/AudioMixer/Public/DSP/Chorus.h
--- a/Engine/Source/Runtime/AudioMixer/Public/DSP/Chorus.h
+++ b/Engine/Source/Runtime/AudioMixer/Public/DSP/Chorus.h
@@ -25,14 +25,16 @@ namespace Audio
 		FChorus();
 		~FChorus();
 
-		void Init(const float InSampleRate, const float InDelayLengthSec = 2.0f, const int32 InControlSamplePeriod = 256);
+		void Init(const float InSampleRate, const int32 InNumChannels, const float InDelayLengthSec = 2.0f, const int32 InControlSamplePeriod = 256);
 
 		void SetDepth(const EChorusDelays::Type InType, const float InDepth);
 		void SetFrequency(const EChorusDelays::Type InType, const float InFrequency);
 		void SetFeedback(const EChorusDelays::Type InType, const float InFeedback);
 		void SetWetLevel(const float InWetLevel);
+		void SetDryLevel(const float InDryLevel);
 		void SetSpread(const float InSpread);
-		void ProcessAudio(const float InLeft, const float InRight, float& OutLeft, float& OutRight);
+		void ProcessAudioFrame(const float* InFrame, float* OutFrame);
+		void ProcessAudio(const float* InBuffer, const int32 InNumSamples, float* OutBuffer);
 
 	protected:
 		FDelay Delays[EChorusDelays::NumDelayTypes];
@@ -46,6 +48,8 @@ namespace Audio
 		float Spread;
 		float MaxFrequencySpread;
 		float WetLevel;
+		float DryLevel;
+		int32 NumChannels;
 	};
 
 
diff --git a/Engine/Source/Runtime/AudioMixer/Public/DSP/Delay.h b/Engine/Source/Runtime/AudioMixer/Public/DSP/Delay.h
--- a/Engine/Source/Runtime/AudioMixer/Public/DSP/Delay.h
+++ b/Engine/Source/Runtime/AudioMixer/Public/DSP/Delay.h
@@ -49,11 +49,8 @@ namespace Audio
 		// Write the input and increment read/write pointers
 		void WriteDelayAndInc(const float InDelayInput);
 
-		// Process audio input and output buffer
-		virtual void ProcessAudio(const float* InAudio, float* OutAudio);
-
 		// Process audio in the delay line, return the delayed value
-		float ProcessAudio(const float InAudio);
+		virtual float ProcessAudioSample(const float InAudio);
 
 	protected:
 
diff --git a/Engine/Source/Runtime/AudioMixer/Public/DSP/DelayStereo.h b/Engine/Source/Runtime/AudioMixer/Public/DSP/DelayStereo.h
--- a/Engine/Source/Runtime/AudioMixer/Public/DSP/DelayStereo.h
+++ b/Engine/Source/Runtime/AudioMixer/Public/DSP/DelayStereo.h
@@ -27,13 +27,16 @@ namespace Audio
 		~FDelayStereo();
 
 		// Initializes the stereo delay with given sample rate and default max delay length
-		void Init(const float InSampleRate, const float InDelayLengthSec = 2.0f);
+		void Init(const float InSampleRate, const int32 InNumChannels, const float InDelayLengthSec = 2.0f);
 
 		// Resets the stereo delay state
 		void Reset();
 
-		// Process the stereo delay to generate new stereo output
-		void ProcessAudio(const float InLeftSample, const float InRightSample, float& OutLeftSample, float& OutRightSample);
+		// Process a single frame of audio
+		void ProcessAudioFrame(const float* InFrame, float* OutFrame);
+
+		// Process a buffer of audio
+		void ProcessAudio(const float* InBuffer, const int32 InNumSamples, float* OutBuffer);
 
 		// Sets which delay stereo mode to use
 		void SetMode(const EStereoDelayMode::Type InMode);
@@ -58,10 +61,7 @@ namespace Audio
 		void UpdateDelays();
 
 		// Left channel delay line
-		FDelay LeftDelay;
-
-		// Right channel delay line
-		FDelay RightDelay;
+		TArray<FDelay> Delays;
 
 		// What mode the stereo delay is in
 		EStereoDelayMode::Type DelayMode;
@@ -78,6 +78,9 @@ namespace Audio
 		// The amount of wet level on the output
 		float WetLevel;
 
+		// The number of channels to use (will sum mono)
+		int32 NumChannels;
+
 		// If the delay has started processing yet
 		bool bIsInit;
 	};
diff --git a/Engine/Source/Runtime/AudioMixer/Public/DSP/DynamicsProcesser.h b/Engine/Source/Runtime/AudioMixer/Public/DSP/DynamicsProcesser.h
--- a/Engine/Source/Runtime/AudioMixer/Public/DSP/DynamicsProcesser.h
+++ b/Engine/Source/Runtime/AudioMixer/Public/DSP/DynamicsProcesser.h
@@ -44,8 +44,8 @@ namespace Audio
 		void SetPeakMode(const EPeakMode::Type InEnvelopeFollowerModeType);
 		void SetProcessingMode(const EDynamicsProcessingMode::Type ProcessingMode);
 
-		void ProcessAudio(const float* InputFrame, int32 NumChannels, float* OutputFrame);
-
+		void ProcessAudioFrame(const float* InFrame, float* OutFrame);
+		void ProcessAudio(const float* InBuffer, const int32 InNumSamples, float* OutBuffer);
 
 
 	protected:
@@ -90,6 +90,9 @@ namespace Audio
 		// Amount of output gain
 		float OutputGain;
 
+		// Number of channels to use for the dynamics processor
+		int32 NumChannels;
+
 		// Whether or not the stereo channels are linked
 		bool bIsChannelLinked;
 
diff --git a/Engine/Source/Runtime/AudioMixer/Public/DSP/Filter.h b/Engine/Source/Runtime/AudioMixer/Public/DSP/Filter.h
--- a/Engine/Source/Runtime/AudioMixer/Public/DSP/Filter.h
+++ b/Engine/Source/Runtime/AudioMixer/Public/DSP/Filter.h
@@ -39,7 +39,10 @@ namespace Audio
 		void Reset();
 
 		// Processes a single frame of audio
-		void ProcessAudioFrame(const float* InAudio, float* OutAudio);
+		void ProcessAudioFrame(const float* InFrame, float* OutFrame);
+
+		// Process a buffer of audio
+		void ProcessAudio(const float* InBuffer, const int32 InNumSamples, float* OutBuffer);
 
 		// Sets all filter parameters with one function
 		void SetParams(const EBiquadFilter::Type InFilterType, const float InCutoffFrequency, const float InBandwidth, const float InGainDB);
@@ -148,7 +151,10 @@ namespace Audio
 		virtual void Update();
 
 		// Processes a single frame of audio. Number of channels MUST be what was used during filter initialization.
-		virtual void ProcessAudio(const float* InSamples, float* OutSamples) = 0;
+		virtual void ProcessAudioFrame(const float* InFrame, float* OutFrame) {}
+
+		// Process an audio buffer.
+		virtual void ProcessAudio(const float* InBuffer, const int32 InNumSamples, float* OutBuffer) = 0;
 
 		// Filter patch destinations
 		FPatchDestination GetModDestCutoffFrequency() const { return ModCutoffFrequencyDest; }
@@ -220,7 +226,8 @@ namespace Audio
 		virtual void Init(const float InSampleRate, const int32 InNumChannels, const int32 InVoiceId = 0, FModulationMatrix* InModMatrix = nullptr) override;
 		virtual void Reset() override;
 		virtual void Update() override;
-		virtual void ProcessAudio(const float* InSamples, float* OutSamples) override;
+		virtual void ProcessAudioFrame(const float* InFrame, float* OutFrame) override;
+		virtual void ProcessAudio(const float* InSamples, const int32 InNumSamples, float* OutSamples) override;
 
 		void SetCoefficient(const float InCoefficient) { A0 = InCoefficient; }
 		float GetCoefficient() const { return A0; }
@@ -242,7 +249,7 @@ namespace Audio
 		virtual void SetBandStopControl(const float InBandStop) override;
 		virtual void Reset() override;
 		virtual void Update() override;
-		virtual void ProcessAudio(const float* InSamples, float* OutSamples) override;
+		virtual void ProcessAudio(const float* InSamples, const int32 InNumSamples, float* OutSamples) override;
 
 	protected:
 		float InputScale;
@@ -275,7 +282,7 @@ namespace Audio
 		virtual void Update() override;
 		virtual void SetQ(const float InQ) override;
 		virtual void SetPassBandGainCompensation(const float InPassBandGainCompensation) override;
-		virtual void ProcessAudio(const float* InSamples, float* OutSamples) override;
+		virtual void ProcessAudio(const float* InSamples, const int32 InNumSamples, float* OutSamples) override;
 
 	protected:
 
diff --git a/Engine/Source/Runtime/AudioMixer/Public/DSP/FoldbackDistortion.h b/Engine/Source/Runtime/AudioMixer/Public/DSP/FoldbackDistortion.h
--- a/Engine/Source/Runtime/AudioMixer/Public/DSP/FoldbackDistortion.h
+++ b/Engine/Source/Runtime/AudioMixer/Public/DSP/FoldbackDistortion.h
@@ -19,7 +19,7 @@ namespace Audio
 		~FFoldbackDistortion();
 
 		// Initialize the equalizer
-		void Init(const float InSampleRate);
+		void Init(const float InSampleRate, const int32 InNumChannels);
 
 		// Sets the foldback distortion threshold
 		void SetThresholdDb(const float InThresholdDb);
@@ -30,11 +30,14 @@ namespace Audio
 		// Sets the output gain
 		void SetOutputGainDb(const float InOutputGainDb);
 
+		// Processes a single audio sample
+		float ProcessAudioSample(const float InSample);
+
 		// Processes a mono stream
-		void ProcessAudio(const float InLeftSample, float& OutLeftSample);
+		void ProcessAudioFrame(const float* InFrame, float* OutFrame);
 
 		// Processes a stereo stream
-		void ProcessAudio(const float InLeftSample, const float InRightSample, float& OutLeftSample, float& OutRightSample);
+		void ProcessAudio(const float* InBuffer, const int32 InNumSamples, float* OutBuffer);
 
 	private:
 		// Threshold to check before folding audio back on itself
@@ -51,6 +54,9 @@ namespace Audio
 
 		// A final gain scaler to apply to the output
 		float OutputGain;
+
+		// How many channels we expect the audio intput to be.
+		int32 NumChannels;
 	};
 
 }
diff --git a/Engine/Source/Runtime/AudioMixer/Public/DSP/OnePole.h b/Engine/Source/Runtime/AudioMixer/Public/DSP/OnePole.h
--- a/Engine/Source/Runtime/AudioMixer/Public/DSP/OnePole.h
+++ b/Engine/Source/Runtime/AudioMixer/Public/DSP/OnePole.h
@@ -46,7 +46,7 @@ namespace Audio
 			}
 		}
 
-		float ProcessAudio(const float InputSample)
+		FORCEINLINE float ProcessAudioSample(const float InputSample)
 		{
 			float Yn = InputSample*A0 + B1*Z1;
 			Yn = UnderflowClamp(Yn);
@@ -54,12 +54,6 @@ namespace Audio
 			return Yn;
 		}
 
-		// Process audio
-		void ProcessAudio(const float* InputSample, float* OutputSample)
-		{
-			*OutputSample = ProcessAudio(*InputSample);
-		}
-
 	protected:
 		float CutoffFrequency;
 
diff --git a/Engine/Source/Runtime/AudioMixer/Public/DSP/Phaser.h b/Engine/Source/Runtime/AudioMixer/Public/DSP/Phaser.h
--- a/Engine/Source/Runtime/AudioMixer/Public/DSP/Phaser.h
+++ b/Engine/Source/Runtime/AudioMixer/Public/DSP/Phaser.h
@@ -16,7 +16,7 @@ namespace Audio
 		FPhaser();
 		~FPhaser();
 
-		void Init(const float SampleRate);
+		void Init(const float SampleRate, const int32 InNumChannels);
 
 		// Sets the phaser LFO rate
 		void SetFrequency(const float InFreqHz);
@@ -33,15 +33,18 @@ namespace Audio
 		// Sets whether or not to put the phaser in quadrature mode
 		void SetQuadPhase(const bool bQuadPhase);
 
-		// Generates the next stereo frame of audio
-		void ProcessAudio(const float* InFrame, float* OutFrame);
+		// Process an audio frame
+		void ProcessAudioFrame(const float* InFrame, float* OutFrame);
+
+		// Process an audio buffer.
+		void ProcessAudio(const float* InBuffer, const int32 InNumSamples, float* OutBuffer);
 
 	protected:
 		void ComputeNewCoefficients(const int32 ChannelIndex, const float LFOValue);
 
 		// First-order all-pass filters in series
 		static const int32 NumApfs = 6;
-		static const int32 NumChannels = 2;
+		static const int32 MaxNumChannels = 2;
 
 		int32 ControlSampleCount;
 		int32 ControlRate;
@@ -49,16 +52,20 @@ namespace Audio
 		float WetLevel;
 		float Feedback;
 		ELFO::Type LFOType;
-		bool bIsBiquadPhase;
 
 		// 6 APFs per channel
-		FBiquadFilter APFs[NumChannels][NumApfs];
+		FBiquadFilter APFs[MaxNumChannels][NumApfs];
 		FVector2D APFFrequencyRanges[NumApfs];
 
 		// Feedback samples
-		float FeedbackFrame[NumChannels];
+		float FeedbackFrame[MaxNumChannels];
 
 		FLFO LFO;
+
+		// Number of channels actually used
+		int32 NumChannels;
+
+		bool bIsBiquadPhase;
 	};
 
 }
diff --git a/Engine/Source/Runtime/AudioMixer/Public/DSP/RingModulation.h b/Engine/Source/Runtime/AudioMixer/Public/DSP/RingModulation.h
--- a/Engine/Source/Runtime/AudioMixer/Public/DSP/RingModulation.h
+++ b/Engine/Source/Runtime/AudioMixer/Public/DSP/RingModulation.h
@@ -19,7 +19,7 @@ namespace Audio
 		~FRingModulation();
 
 		// Initialize the equalizer
-		void Init(const float InSampleRate);
+		void Init(const float InSampleRate, const int32 InNumChannels);
 
 		// The type of modulation
 		void SetModulatorWaveType(const EOsc::Type InType);
@@ -30,13 +30,25 @@ namespace Audio
 		// Set the ring modulation depth
 		void SetModulationDepth(const float InModulationDepth);
 
+		// Set the dry level of the ring modulation
+		void SetDryLevel(const float InDryLevel) { DryLevel = InDryLevel; }
+
+		// Set the wet level of the ring modulation
+		void SetWetLevel(const float InWetLevel) { WetLevel = InWetLevel; }
+
 		// Processes the audio frame (audio frame must have channels equal to that used during initialization)
-		void ProcessAudio(const float InLeftSample, const float InRightSample, float& OutLeftSample, float& OutRightSample);
+		void ProcessAudioFrame(const float* InFrame, float* OutFrame);
+
+		// Process audio buffer
+		void ProcessAudio(const float* InBuffer, const int32 InNumSamples, float* OutBuffer);
 
 	private:
 		Audio::FOsc Osc;
 		float ModulationFrequency;
 		float ModulationDepth;
+		float DryLevel;
+		float WetLevel;
+		int32 NumChannels;
 	};
 
 }
diff --git a/Engine/Source/Runtime/Engine/Classes/Sound/SoundEffectSource.h b/Engine/Source/Runtime/Engine/Classes/Sound/SoundEffectSource.h
--- a/Engine/Source/Runtime/Engine/Classes/Sound/SoundEffectSource.h
+++ b/Engine/Source/Runtime/Engine/Classes/Sound/SoundEffectSource.h
@@ -5,6 +5,7 @@
 #include "UObject/ObjectMacros.h"
 #include "Sound/SoundEffectPreset.h"
 #include "Sound/SoundEffectBase.h"
+#include "IAudioExtensionPlugin.h"
 
 #include "SoundEffectSource.generated.h"
 
@@ -59,7 +60,6 @@ struct FSoundEffectSourceInitData
 {
 	float SampleRate;
 	int32 NumSourceChannels;
-	float SourceDuration;
 	double AudioClock;
 
 	// The object id of the parent preset
@@ -68,7 +68,6 @@ struct FSoundEffectSourceInitData
 	FSoundEffectSourceInitData()
 	: SampleRate(0.0f)
 	, NumSourceChannels(0)
-	, SourceDuration(0.0f)
 	, AudioClock(0.0)
 	, ParentPresetUniqueId(INDEX_NONE)
 	{}
@@ -77,33 +76,25 @@ struct FSoundEffectSourceInitData
 /** Struct which has data to initialize the source effect. */
 struct FSoundEffectSourceInputData
 {
-	TArray<float> AudioFrame;
-	FVector SourcePosition;
-	FVector LeftChannelPosition;
-	FVector RightChannelPosition;
-	FVector ListenerPosition;
 	float CurrentVolume;
 	float CurrentPitch;
-	float CurrentPlayTime;
-	float Duration;
-	float Distance;
 	double AudioClock;
-	int32 CurrentLoopCount;
-	int32 MaxLoopCount;
-	uint8 bLooping : 1;
-	uint8 bIsSpatialized : 1;
+	float CurrentPlayFraction;
+	FSpatializationParams SpatParams;
+	float* InputSourceEffectBufferPtr;
+	int32 NumSamples;
 
 	FSoundEffectSourceInputData()
+		: CurrentVolume(0.0f)
+		, CurrentPitch(0.0f)
+		, AudioClock(0.0)
+		, SpatParams(FSpatializationParams())
+		, InputSourceEffectBufferPtr(nullptr)
+		, NumSamples(0)
 	{
-		FMemory::Memzero(this, sizeof(FSoundEffectSourceInputData));
 	}
 };
 
-struct FSoundEffectSourceOutputData
-{
-	TArray<float> AudioFrame;
-};
-
 class ENGINE_API FSoundEffectSource : public FSoundEffectBase
 {
 public:
@@ -113,6 +104,6 @@ public:
 	virtual void Init(const FSoundEffectSourceInitData& InSampleRate) = 0;
 
 	/** Process the input block of audio. Called on audio thread. */
-	virtual void ProcessAudio(const FSoundEffectSourceInputData& InData, FSoundEffectSourceOutputData& OutData) = 0;
+	virtual void ProcessAudio(const FSoundEffectSourceInputData& InData, float* OutAudioBufferData) = 0;
 };
 
diff --git a/Engine/Source/Runtime/Engine/Classes/Sound/SoundEffectSubmix.h b/Engine/Source/Runtime/Engine/Classes/Sound/SoundEffectSubmix.h
--- a/Engine/Source/Runtime/Engine/Classes/Sound/SoundEffectSubmix.h
+++ b/Engine/Source/Runtime/Engine/Classes/Sound/SoundEffectSubmix.h
@@ -15,6 +15,9 @@ UCLASS(config = Engine, hidecategories = Object, abstract, editinlinenew, Bluepr
 class ENGINE_API USoundEffectSubmixPreset : public USoundEffectPreset
 {
 	GENERATED_UCLASS_BODY()
+
+	virtual FColor GetPresetColor() const override { return FColor(162.0f, 84.0f, 101.0f); }
+
 };
 
 /** Struct which has data needed to initialize the submix effect. */
diff --git a/Engine/Source/Runtime/Engine/Classes/Sound/SoundSourceBus.h b/Engine/Source/Runtime/Engine/Classes/Sound/SoundSourceBus.h
--- a/Engine/Source/Runtime/Engine/Classes/Sound/SoundSourceBus.h
+++ b/Engine/Source/Runtime/Engine/Classes/Sound/SoundSourceBus.h
@@ -24,7 +24,7 @@ enum class ESourceBusChannels : uint8
 UCLASS(hidecategories= (Compression, SoundWave, Streaming, Subtitles, Sound, Info, ImportSettings), ClassGroup = Sound, meta = (BlueprintSpawnableComponent))
 class ENGINE_API USoundSourceBus : public USoundWave
 {
-	GENERATED_BODY()
+	GENERATED_UCLASS_BODY()
 
 public:
 
@@ -42,7 +42,6 @@ public:
 
 	//~ Begin UObject Interface. 
 	virtual void PostLoad() override;
-	virtual void BeginDestroy() override;
 #if WITH_EDITOR
 	virtual void PostEditChangeProperty(FPropertyChangedEvent& PropertyChangedEvent) override;
 #endif // WITH_EDITOR
@@ -54,6 +53,7 @@ public:
 	//~ End USoundBase Interface.
 
 protected:
+	void Init();
 
 	uint32 DurationSamples;
 	bool bInitialized;
diff --git a/Engine/Source/Runtime/Engine/Classes/Sound/SoundSubmix.h b/Engine/Source/Runtime/Engine/Classes/Sound/SoundSubmix.h
--- a/Engine/Source/Runtime/Engine/Classes/Sound/SoundSubmix.h
+++ b/Engine/Source/Runtime/Engine/Classes/Sound/SoundSubmix.h
@@ -75,6 +75,14 @@ struct ENGINE_API FSoundSubmixSendInfo
 */
 DECLARE_DYNAMIC_MULTICAST_DELEGATE_OneParam(FOnSubmixRecordedFileDone,const USoundWave*, ResultingSoundWave);
 
+/** 
+* Called when a new submix envelope value is generated on the given audio device id (different for multiple PIE). Array is an envelope value for each channel.
+*/
+DECLARE_DYNAMIC_MULTICAST_DELEGATE_OneParam(FOnSubmixEnvelope, const TArray<float>&, Envelope);
+
+DECLARE_DYNAMIC_DELEGATE_OneParam(FOnSubmixEnvelopeBP, const TArray<float>&, Envelope);
+
+
 #if WITH_EDITOR
 
 /** Interface for sound submix graph interaction with the AudioEditor module. */
@@ -116,6 +124,14 @@ class ENGINE_API USoundSubmix : public UObject
 	UPROPERTY(EditAnywhere, BlueprintReadWrite, Category = SoundSubmix)
 	UAmbisonicsSubmixSettingsBase* AmbisonicsPluginSettings;
 
+	/** The attack time in milliseconds for the envelope follower. Delegate callbacks can be registered to get the envelope value of sounds played with this submix. */
+	UPROPERTY(EditAnywhere, BlueprintReadWrite, Category = EnvelopeFollower, meta = (ClampMin = "0", UIMin = "0"))
+	int32 EnvelopeFollowerAttackTime;
+
+	/** The release time in milliseconds for the envelope follower. Delegate callbacks can be registered to get the envelope value of sounds played with this submix. */
+	UPROPERTY(EditAnywhere, BlueprintReadWrite, Category = EnvelopeFollower, meta = (ClampMin = "0", UIMin = "0"))
+	int32 EnvelopeFollowerReleaseTime;
+
 	// Blueprint delegate for when a recorded file is finished exporting.
 	UPROPERTY(BlueprintAssignable)
 	FOnSubmixRecordedFileDone OnSubmixRecordedFileDone;
@@ -132,6 +148,21 @@ class ENGINE_API USoundSubmix : public UObject
 
 	void StopRecordingOutput(FAudioDevice* InDevice, EAudioRecordingExportType ExportType, const FString& Name, FString Path, USoundWave* ExistingSoundWaveToOverwrite = nullptr);
 
+	// Start envelope following the submix output. Register with OnSubmixEnvelope to receive envelope follower data in BP.
+	UFUNCTION(BlueprintCallable, Category = "Audio|EnvelopeFollowing", meta = (WorldContext = "WorldContextObject"))
+	void StartEnvelopeFollowing(const UObject* WorldContextObject);
+
+	void StartEnvelopeFollowing(FAudioDevice* InDevice);
+
+	// Start envelope following the submix output. Register with OnSubmixEnvelope to receive envelope follower data in BP.
+	UFUNCTION(BlueprintCallable, Category = "Audio|EnvelopeFollowing", meta = (WorldContext = "WorldContextObject"))
+	void StopEnvelopeFollowing(const UObject* WorldContextObject);
+
+	void StopEnvelopeFollowing(FAudioDevice* InDevice);
+
+	UFUNCTION(BlueprintCallable, Category = "Audio|EnvelopeFollowing", meta = (WorldContext = "WorldContextObject"))
+	void AddEnvelopeFollowerDelegate(const UObject* WorldContextObject, const FOnSubmixEnvelopeBP& OnSubmixEnvelopeBP);
+
 	// Registers and unregisters buffer listeners with the submix
 	void RegisterSubmixBufferListener(ISubmixBufferListener* InBufferListener);
 	void UnregisterSubmixBufferListener(ISubmixBufferListener* InBufferListener);
@@ -147,7 +178,7 @@ protected:
 	virtual void PostEditChangeProperty(struct FPropertyChangedEvent& PropertyChangedEvent) override;
 #endif
 	//~ End UObject Interface.
-
+	
 	// State handling for bouncing output.
 	TUniquePtr<Audio::FAudioRecordingData> RecordingData;
 
@@ -196,5 +227,7 @@ private:
 
 #endif
 
+
+
 };
 
diff --git a/Engine/Source/Runtime/Engine/Classes/Sound/SoundWave.h b/Engine/Source/Runtime/Engine/Classes/Sound/SoundWave.h
--- a/Engine/Source/Runtime/Engine/Classes/Sound/SoundWave.h
+++ b/Engine/Source/Runtime/Engine/Classes/Sound/SoundWave.h
@@ -187,8 +187,8 @@ class ENGINE_API USoundWave : public USoundBase
 	UPROPERTY(EditAnywhere, Category=Subtitles )
 	uint8 bSingleLine:1;
 
-	/** Allows sound to play at 0 volume, otherwise will stop the sound when the sound is silent. */
-	UPROPERTY(EditAnywhere, Category=Sound)
+	/** Allows sound to continue playing when silent. This prevents issues with sounds restarting when coming back in range, etc. */
+	UPROPERTY(EditAnywhere, Category=Sound, meta = (DisplayName = "Play When Silent"))
 	uint8 bVirtualizeWhenSilent:1;
 
 	/** Whether or not this source is ambisonics file format. */
diff --git a/Engine/Source/Runtime/Engine/Private/AudioDerivedData.cpp b/Engine/Source/Runtime/Engine/Private/AudioDerivedData.cpp
--- a/Engine/Source/Runtime/Engine/Private/AudioDerivedData.cpp
+++ b/Engine/Source/Runtime/Engine/Private/AudioDerivedData.cpp
@@ -16,6 +16,7 @@
 #include "ProfilingDebugging/CookStats.h"
 #include "AudioResampler.h"
 #include "AudioCompressionSettingsUtils.h"
+#include "Sound/SoundSourceBus.h"
 
 DEFINE_LOG_CATEGORY_STATIC(LogAudioDerivedData, Log, All);
 
@@ -841,6 +842,12 @@ static void CookSimpleWave(USoundWave* SoundWave, FName FormatName, const IAudio
 
 	int32 WaveSampleRate = 0;
 
+	// Don't need to cook source buses
+	if (USoundSourceBus* SourceBus = Cast<USoundSourceBus>(SoundWave))
+	{
+		return;
+	}
+
 	// check if there is any raw sound data
 	if( SoundWave->RawData.GetBulkDataSize() > 0 )
 	{
