Remove hacks to workaround bad codegen on VS2017 v15.8.

diff --git a/Engine/Source/Runtime/AudioMixer/Private/AudioMixerSourceVoice.cpp b/Engine/Source/Runtime/AudioMixer/Private/AudioMixerSourceVoice.cpp
--- a/Engine/Source/Runtime/AudioMixer/Private/AudioMixerSourceVoice.cpp
+++ b/Engine/Source/Runtime/AudioMixer/Private/AudioMixerSourceVoice.cpp
@@ -1,4 +1,4 @@
-﻿// Copyright 1998-2018 Epic Games, Inc. All Rights Reserved.
+// Copyright 1998-2018 Epic Games, Inc. All Rights Reserved.
 
 #include "AudioMixerSourceVoice.h"
 #include "AudioMixerSource.h"
@@ -246,8 +246,6 @@ namespace Audio
 		return SourceManager->MixOutputBuffers(SourceId, InSubmixChannelType, SendLevel, OutWetBuffer);
 	}
 
-// MSVC 2017 15.8.4 is generating bad code where setting submix will crash. See https://developercommunity.visualstudio.com/content/problem/345511/bad-code-generation-in-vs-2017-v1585.html.
-MSVC_PRAGMA(optimize("", off))
 	void FMixerSourceVoice::SetSubmixSendInfo(FMixerSubmixWeakPtr Submix, const float SendLevel)
 	{
 		AUDIO_MIXER_CHECK_GAME_THREAD(MixerDevice);
@@ -276,7 +274,6 @@ MSVC_PRAGMA(optimize("", off))
 			}
 		}
 	}
-MSVC_PRAGMA(optimize("", on))
 
 	void FMixerSourceVoice::OnMixBus(FMixerSourceVoiceBuffer* OutMixerSourceBuffer)
 	{
diff --git a/Engine/Source/Runtime/AudioMixer/Private/AudioMixerSubmix.cpp b/Engine/Source/Runtime/AudioMixer/Private/AudioMixerSubmix.cpp
--- a/Engine/Source/Runtime/AudioMixer/Private/AudioMixerSubmix.cpp
+++ b/Engine/Source/Runtime/AudioMixer/Private/AudioMixerSubmix.cpp
@@ -137,8 +137,6 @@ namespace Audio
 		});
 	}
 
-// MSVC 2017 15.8.4 is generating bad code where setting submix will crash. See https://developercommunity.visualstudio.com/content/problem/345511/bad-code-generation-in-vs-2017-v1585.html.
-MSVC_PRAGMA(optimize("", off))
 	void FMixerSubmix::AddChildSubmix(TWeakPtr<FMixerSubmix, ESPMode::ThreadSafe> SubmixWeakPtr)
 	{
 		SubmixCommand([this, SubmixWeakPtr]()
@@ -163,7 +161,6 @@ MSVC_PRAGMA(optimize("", off))
 			}
 		});
 	}
-MSVC_PRAGMA(optimize("", on))
 
 	ESubmixChannelFormat FMixerSubmix::GetSubmixChannels() const
 	{
