Fix misnamed file #rb ethan.geller

diff --git a/Engine/Plugins/Runtime/Synthesis/Source/Synthesis/Classes/SourceEffects/SourceEffectDynamicsProcessor.h b/Engine/Plugins/Runtime/Synthesis/Source/Synthesis/Classes/SourceEffects/SourceEffectDynamicsProcessor.h
--- a/Engine/Plugins/Runtime/Synthesis/Source/Synthesis/Classes/SourceEffects/SourceEffectDynamicsProcessor.h
+++ b/Engine/Plugins/Runtime/Synthesis/Source/Synthesis/Classes/SourceEffects/SourceEffectDynamicsProcessor.h
@@ -3,7 +3,7 @@
 #pragma once
 
 #include "Sound/SoundEffectSource.h"
-#include "DSP/DynamicsProcesser.h"
+#include "DSP/DynamicsProcessor.h"
 #include "SourceEffectDynamicsProcessor.generated.h"
 
 UENUM(BlueprintType)
diff --git a/Engine/Plugins/Runtime/TimeSynth/Source/TimeSynth/Classes/TimeSynthComponent.h b/Engine/Plugins/Runtime/TimeSynth/Source/TimeSynth/Classes/TimeSynthComponent.h
--- a/Engine/Plugins/Runtime/TimeSynth/Source/TimeSynth/Classes/TimeSynthComponent.h
+++ b/Engine/Plugins/Runtime/TimeSynth/Source/TimeSynth/Classes/TimeSynthComponent.h
@@ -18,7 +18,7 @@
 #include "DSP/SpectrumAnalyzer.h"
 #include "DSP/Filter.h"
 #include "DSP/EnvelopeFollower.h"
-#include "DSP/DynamicsProcesser.h"
+#include "DSP/DynamicsProcessor.h"
 #include "TimeSynthComponent.generated.h"
 
 class USoundWave;
diff --git a/Engine/Source/Runtime/AudioMixer/Private/DSP/DynamicsProcessor.cpp b/Engine/Source/Runtime/AudioMixer/Private/DSP/DynamicsProcessor.cpp
--- a/Engine/Source/Runtime/AudioMixer/Private/DSP/DynamicsProcessor.cpp
+++ b/Engine/Source/Runtime/AudioMixer/Private/DSP/DynamicsProcessor.cpp
@@ -1,6 +1,6 @@
 // Copyright 1998-2019 Epic Games, Inc. All Rights Reserved.
 
-#include "DSP/DynamicsProcesser.h"
+#include "DSP/DynamicsProcessor.h"
 
 namespace Audio
 {
diff --git a/Engine/Source/Runtime/AudioMixer/Public/DSP/DynamicsProcesser.h b/Engine/Source/Runtime/AudioMixer/Public/DSP/DynamicsProcesser.h
deleted file mode 100644
diff --git a/Engine/Source/Runtime/AudioMixer/Public/DSP/DynamicsProcessor.h b/Engine/Source/Runtime/AudioMixer/Public/DSP/DynamicsProcessor.h
new file mode 100644
--- /dev/null
+++ b/Engine/Source/Runtime/AudioMixer/Public/DSP/DynamicsProcessor.h
@@ -0,0 +1,103 @@
+// Copyright 1998-2019 Epic Games, Inc. All Rights Reserved.
+
+#pragma once
+
+#include "CoreMinimal.h"
+#include "DSP/EnvelopeFollower.h"
+#include "DSP/Delay.h"
+
+namespace Audio
+{
+	// What mode the compressor is in
+	namespace EDynamicsProcessingMode
+	{
+		enum Type
+		{
+			Compressor,
+			Limiter,
+			Expander,
+			Gate,
+			Count
+		};
+	}
+
+	// Dynamic range compressor
+	// https://en.wikipedia.org/wiki/Dynamic_range_compression
+	class AUDIOMIXER_API FDynamicsProcessor
+	{
+	public:
+		FDynamicsProcessor();
+		~FDynamicsProcessor();
+
+		void Init(const float SampleRate, const int32 NumChannels = 2);
+
+		void SetLookaheadMsec(const float InLookAheadMsec);
+		void SetAttackTime(const float InAttackTimeMsec);
+		void SetReleaseTime(const float InReleaseTimeMsec);
+		void SetThreshold(const float InThresholdDb);
+		void SetRatio(const float InCompressionRatio);
+		void SetKneeBandwidth(const float InKneeBandwidthDb);
+		void SetInputGain(const float InInputGainDb);
+		void SetOutputGain(const float InOutputGainDb);
+		void SetChannelLinked(const bool bInIsChannelLinked);
+		void SetAnalogMode(const bool bInIsAnalogMode);
+		void SetPeakMode(const EPeakMode::Type InEnvelopeFollowerModeType);
+		void SetProcessingMode(const EDynamicsProcessingMode::Type ProcessingMode);
+
+		void ProcessAudioFrame(const float* InFrame, float* OutFrame);
+		void ProcessAudio(const float* InBuffer, const int32 InNumSamples, float* OutBuffer);
+
+
+	protected:
+
+		float ComputeGain(const float InEnvFollowerDb);
+
+		EDynamicsProcessingMode::Type ProcessingMode;
+
+		// Lookahead delay lines
+		TArray<FDelay> LookaheadDelay;
+
+		// Envelope followers
+		TArray<FEnvelopeFollower> EnvFollower;
+
+		// Points in the knee used for lagrangian interpolation
+		TArray<FVector2D> KneePoints;
+
+		TArray<float> DetectorOuts;
+		TArray<float> Gain;
+
+		// How far ahead to look in the audio
+		float LookaheedDelayMsec;
+
+		// The period of which the compressor decreases gain to the level determined by the compression ratio
+		float AttackTimeMsec;
+
+		// The period of which the compressor increases gain to 0 dB once level has fallen below the threshold
+		float ReleaseTimeMsec;
+
+		// Amplitude threshold above which gain will be reduced
+		float ThresholdDb;
+
+		// Amount of gain reduction.
+		float Ratio;
+
+		// Defines how hard or soft the gain reduction blends from no gain reduction to gain reduction (determined by the ratio)
+		float HalfKneeBandwidthDb;
+
+		// Amount of input gain
+		float InputGain;
+
+		// Amount of output gain
+		float OutputGain;
+
+		// Number of channels to use for the dynamics processor
+		int32 NumChannels;
+
+		// Whether or not the stereo channels are linked
+		bool bIsChannelLinked;
+
+		// Whether or not we're in analog mode
+		bool bIsAnalogMode;
+
+	};
+}
