Fixing pops on audio mixer. Pops due to timing problems with mixes.

diff --git a/Engine/Source/Runtime/AudioMixer/Private/AudioMixerDevice.cpp b/Engine/Source/Runtime/AudioMixer/Private/AudioMixerDevice.cpp
--- a/Engine/Source/Runtime/AudioMixer/Private/AudioMixerDevice.cpp
+++ b/Engine/Source/Runtime/AudioMixer/Private/AudioMixerDevice.cpp
@@ -99,6 +99,11 @@ namespace Audio
 		return CurrentThreadId == AudioPlatformThreadId;
 	}
 
+	void FMixerDevice::UpdateDeviceDeltaTime()
+	{
+		DeviceDeltaTime = GetGameDeltaTime();
+	}
+
 	void FMixerDevice::GetAudioDeviceList(TArray<FString>& OutAudioDeviceNames) const
 	{
 		if (AudioMixerPlatform && AudioMixerPlatform->IsInitialized())
diff --git a/Engine/Source/Runtime/AudioMixer/Public/AudioMixerDevice.h b/Engine/Source/Runtime/AudioMixer/Public/AudioMixerDevice.h
--- a/Engine/Source/Runtime/AudioMixer/Public/AudioMixerDevice.h
+++ b/Engine/Source/Runtime/AudioMixer/Public/AudioMixerDevice.h
@@ -73,6 +73,7 @@ namespace Audio
 		~FMixerDevice();
 
 		//~ Begin FAudioDevice
+		virtual void UpdateDeviceDeltaTime() override;
 		virtual void GetAudioDeviceList(TArray<FString>& OutAudioDeviceNames) const override;
 		virtual bool InitializeHardware() override;
 		virtual void FadeIn() override;
