Fixing leak in audiomixer

diff --git a/Engine/Source/Runtime/AudioMixer/Private/AudioMixerSourceBuffer.cpp b/Engine/Source/Runtime/AudioMixer/Private/AudioMixerSourceBuffer.cpp
--- a/Engine/Source/Runtime/AudioMixer/Private/AudioMixerSourceBuffer.cpp
+++ b/Engine/Source/Runtime/AudioMixer/Private/AudioMixerSourceBuffer.cpp
@@ -430,6 +430,9 @@ namespace Audio
 		if (AsyncRealtimeAudioTask) 
 		{ 
 			AsyncRealtimeAudioTask->EnsureCompletion(); 
+
+			delete AsyncRealtimeAudioTask;
+			AsyncRealtimeAudioTask = nullptr;
 		}
 	}
 
