synth component fix.
making sure USynthComponent ref is set on USoundSynth object.

diff --git a/Engine/Source/Runtime/AudioMixer/Public/Components/SynthComponent.h b/Engine/Source/Runtime/AudioMixer/Public/Components/SynthComponent.h
--- a/Engine/Source/Runtime/AudioMixer/Public/Components/SynthComponent.h
+++ b/Engine/Source/Runtime/AudioMixer/Public/Components/SynthComponent.h
@@ -46,7 +46,9 @@ class AUDIOMIXER_API USynthSound : public USoundWaveProcedural
 	/** End USoundWave */
 
 protected:
+	UPROPERTY()
 	USynthComponent* OwningSynthComponent;
+
 	TArray<float> FloatBuffer;
 	bool bAudioMixer;
 };
