Fix for hang in audio mixer during flush.

- Fix removes unnecessary bool to flag requiring a flush.
- Cleans up command queue swap.
- Switches to using a manual-reset Event to prevent hangs.

diff --git a/Engine/Source/Runtime/AudioMixer/Private/AudioMixerSourceManager.cpp b/Engine/Source/Runtime/AudioMixer/Private/AudioMixerSourceManager.cpp
--- a/Engine/Source/Runtime/AudioMixer/Private/AudioMixerSourceManager.cpp
+++ b/Engine/Source/Runtime/AudioMixer/Private/AudioMixerSourceManager.cpp
@@ -108,7 +108,9 @@ namespace Audio
 		, bUsingSpatializationPlugin(false)
 		, MaxChannelsSupportedBySpatializationPlugin(1)
 	{
-		CommandsProcessedEvent = FPlatformProcess::GetSynchEventFromPool();
+		// Get a manual resetable event
+		const bool bIsManualReset = true;
+		CommandsProcessedEvent = FPlatformProcess::GetSynchEventFromPool(bIsManualReset);
 		check(CommandsProcessedEvent != nullptr);
 	}
 
@@ -264,7 +266,6 @@ namespace Audio
 		}
 
 		bInitialized = true;
-		bPumpQueue = false;
 	}
 
 	void FMixerSourceManager::Update()
@@ -282,22 +283,21 @@ namespace Audio
 		}
 #endif
 
-		int32 CurrentRenderIndex = RenderThreadCommandBufferIndex.GetValue();
-		int32 CurrentGameIndex = AudioThreadCommandBufferIndex.GetValue();
-		check(CurrentGameIndex == 0 || CurrentGameIndex == 1);
-		check(CurrentRenderIndex == 0 || CurrentRenderIndex == 1);
-
-		// If these values are the same, that means the audio render thread has finished the last buffer queue so is ready for the next block
-		if (CurrentRenderIndex == CurrentGameIndex)
+		// If the command was triggered, then we want to do a swap of command buffers
+		if (CommandsProcessedEvent->Wait(0))
 		{
+			int32 CurrentGameIndex = AudioThreadCommandBufferIndex.GetValue();
+
 			// This flags the audio render thread to be able to pump the next batch of commands
 			// And will allow the audio thread to write to a new command slot
-			const int32 NextIndex = !CurrentGameIndex;
+			const int32 NextIndex = (CurrentGameIndex + 1) & 1;
 
 			// Make sure we've actually emptied the command queue from the render thread before writing to it
 			check(CommandBuffers[NextIndex].SourceCommandQueue.Num() == 0);
 			AudioThreadCommandBufferIndex.Set(NextIndex);
-			bPumpQueue = true;
+			RenderThreadCommandBufferIndex.Set(CurrentGameIndex);
+
+			CommandsProcessedEvent->Reset();
 		}
 	}
 
@@ -2315,11 +2315,7 @@ namespace Audio
 		SCOPE_CYCLE_COUNTER(STAT_AudioMixerSourceManagerUpdate);
 
 		// Get the this blocks commands before rendering audio
-		if (bPumpQueue)
-		{
-			bPumpQueue = false;
-			PumpCommandQueue();
-		}
+		PumpCommandQueue();
 
 		// Update pending tasks and release them if they're finished
 		UpdatePendingReleaseData();
@@ -2387,6 +2383,12 @@ namespace Audio
 
 	void FMixerSourceManager::PumpCommandQueue()
 	{
+		// If we're already triggered, we need to wait for the audio thread to reset it before pumping
+		if (CommandsProcessedEvent->Wait(0))
+		{
+			return;
+		}
+
 		int32 CurrentRenderThreadIndex = RenderThreadCommandBufferIndex.GetValue();
 
 		FCommands& Commands = CommandBuffers[CurrentRenderThreadIndex];
@@ -2400,8 +2402,6 @@ namespace Audio
 
 		Commands.SourceCommandQueue.Reset();
 
-		RenderThreadCommandBufferIndex.Set(!CurrentRenderThreadIndex);
-
 		check(CommandsProcessedEvent != nullptr);
 		CommandsProcessedEvent->Trigger();
 	}
diff --git a/Engine/Source/Runtime/AudioMixer/Private/AudioMixerSourceManager.h b/Engine/Source/Runtime/AudioMixer/Private/AudioMixerSourceManager.h
--- a/Engine/Source/Runtime/AudioMixer/Private/AudioMixerSourceManager.h
+++ b/Engine/Source/Runtime/AudioMixer/Private/AudioMixerSourceManager.h
@@ -548,7 +548,5 @@ namespace Audio
 		int32 MaxChannelsSupportedBySpatializationPlugin;
 
 		friend class FMixerSourceVoice;
-		// Set to true when the audio source manager should pump the command queue
-		FThreadSafeBool bPumpQueue;
 	};
 }
