Reverb fixes

- Alowing reverb submix to be muted while in background
- Fixing dry-levels in reverb submixes
- Allowing reverb to be disabled via cvar while running vs at app-start
- Cleaning up default manual send levels (to default to 0.0)

diff --git a/Engine/Source/Runtime/AudioMixer/Classes/SubmixEffects/AudioMixerSubmixEffectReverb.h b/Engine/Source/Runtime/AudioMixer/Classes/SubmixEffects/AudioMixerSubmixEffectReverb.h
--- a/Engine/Source/Runtime/AudioMixer/Classes/SubmixEffects/AudioMixerSubmixEffectReverb.h
+++ b/Engine/Source/Runtime/AudioMixer/Classes/SubmixEffects/AudioMixerSubmixEffectReverb.h
@@ -96,6 +96,9 @@ public:
 	// We want to receive downmixed submix audio to stereo input for the reverb effect
 	virtual uint32 GetDesiredInputChannelCountOverride() const override { return 2; }
 
+	// Returns the drylevel of the effect
+	virtual float GetDryLevel() const override { return DryLevel; }
+
 	// Process the input block of audio. Called on audio thread.
 	virtual void OnProcessAudio(const FSoundEffectSubmixInputData& InData, FSoundEffectSubmixOutputData& OutData) override;
 
diff --git a/Engine/Source/Runtime/AudioMixer/Private/AudioMixerDevice.cpp b/Engine/Source/Runtime/AudioMixer/Private/AudioMixerDevice.cpp
--- a/Engine/Source/Runtime/AudioMixer/Private/AudioMixerDevice.cpp
+++ b/Engine/Source/Runtime/AudioMixer/Private/AudioMixerDevice.cpp
@@ -21,14 +21,6 @@
 #include "AudioEditorModule.h"
 #endif
 
-static int32 DisableSubmixReverbCVar = 0;
-FAutoConsoleVariableRef CVarDisableSubmixReverb(
-	TEXT("au.DisableReverbSubmix"),
-	DisableSubmixReverbCVar,
-	TEXT("Disables the reverb submix.\n")
-	TEXT("0: Not Disabled, 1: Disabled"),
-	ECVF_Default);
-
 static int32 DisableSubmixEffectEQCvar = 0;
 FAutoConsoleVariableRef CVarDisableSubmixEQ(
 	TEXT("au.DisableSubmixEffectEQ"),
@@ -46,6 +38,7 @@ namespace Audio
 		: AudioMixerPlatform(InAudioMixerPlatform)
 		, AudioClockDelta(0.0)
 		, AudioClock(0.0)
+		, PreviousMasterVolume((float)INDEX_NONE)
 		, SourceManager(this)
 		, GameOrAudioThreadId(INDEX_NONE)
 		, AudioPlatformThreadId(INDEX_NONE)
@@ -395,6 +388,35 @@ namespace Audio
 				});
 			}
 		}
+
+		// Check if the background mute changed state and update the submixes which are enabled to do background muting
+		const float CurrentMasterVolume = GetMasterVolume();
+		if (!FMath::IsNearlyEqual(PreviousMasterVolume, CurrentMasterVolume))
+		{
+			PreviousMasterVolume = CurrentMasterVolume;
+			bool IsMuted = FMath::IsNearlyZero(CurrentMasterVolume);
+
+			for (TObjectIterator<USoundSubmix> It; It; ++It)
+			{
+				if (*It && It->bMuteWhenBackgrounded)
+				{
+					FMixerSubmix* SubmixInstance = GetMasterSubmixInstance(*It);
+					if (!SubmixInstance)
+					{
+						FMixerSubmixPtr* NonMasterSubmix = Submixes.Find(*It);
+						if (NonMasterSubmix)
+						{
+							SubmixInstance = (*NonMasterSubmix).Get();
+						}
+					}
+
+					if (SubmixInstance)
+					{
+						SubmixInstance->SetBackgroundMuted(IsMuted);
+					}
+				}
+			}
+		}
 	}
 
 	double FMixerDevice::GetAudioTime() const
@@ -567,28 +589,20 @@ namespace Audio
 			FMixerDevice::MasterSubmixes.Add(MasterSubmix);
 
 			// Master Reverb Plugin
-			if (DisableSubmixReverbCVar)
-			{
-				FMixerDevice::MasterSubmixes.Add(nullptr);
-			}
-			else
-			{
-				USoundSubmix* ReverbPluginSubmix = NewObject<USoundSubmix>(USoundSubmix::StaticClass(), TEXT("Master Reverb Plugin Submix"));
-				ReverbPluginSubmix->AddToRoot();
-				FMixerDevice::MasterSubmixes.Add(ReverbPluginSubmix);
-			}
+			USoundSubmix* ReverbPluginSubmix = NewObject<USoundSubmix>(USoundSubmix::StaticClass(), TEXT("Master Reverb Plugin Submix"));
+			// Make the master reverb mute when backgrounded
+			ReverbPluginSubmix->bMuteWhenBackgrounded = true;
+
+			ReverbPluginSubmix->AddToRoot();
+			FMixerDevice::MasterSubmixes.Add(ReverbPluginSubmix);
 
 			// Master Reverb
-			if (DisableSubmixReverbCVar)
-			{
-				FMixerDevice::MasterSubmixes.Add(nullptr);
-			}
-			else
-			{
-				USoundSubmix* ReverbSubmix = NewObject<USoundSubmix>(USoundSubmix::StaticClass(), TEXT("Master Reverb Submix"));
-				ReverbSubmix->AddToRoot();
-				FMixerDevice::MasterSubmixes.Add(ReverbSubmix);
-			}
+			USoundSubmix* ReverbSubmix = NewObject<USoundSubmix>(USoundSubmix::StaticClass(), TEXT("Master Reverb Submix"));
+			// Make the master reverb mute when backgrounded
+			ReverbSubmix->bMuteWhenBackgrounded = true;
+
+			ReverbSubmix->AddToRoot();
+			FMixerDevice::MasterSubmixes.Add(ReverbSubmix);
 
 			// Master EQ
 			if (DisableSubmixEffectEQCvar)
@@ -939,6 +953,19 @@ namespace Audio
 		return false;
 	}
 
+	FMixerSubmix* FMixerDevice::GetMasterSubmixInstance(USoundSubmix* InSubmix)
+	{
+		check(MasterSubmixes.Num() == EMasterSubmixType::Count);
+		for (int32 i = 0; i < EMasterSubmixType::Count; ++i)
+		{
+			if (InSubmix == FMixerDevice::MasterSubmixes[i])
+			{
+				return MasterSubmixInstances[i].Get();
+			}
+		}
+		return nullptr;
+	}
+
 	void FMixerDevice::RegisterSoundSubmix(USoundSubmix* InSoundSubmix, bool bInit)
 	{
 		if (InSoundSubmix)
diff --git a/Engine/Source/Runtime/AudioMixer/Private/AudioMixerSubmix.cpp b/Engine/Source/Runtime/AudioMixer/Private/AudioMixerSubmix.cpp
--- a/Engine/Source/Runtime/AudioMixer/Private/AudioMixerSubmix.cpp
+++ b/Engine/Source/Runtime/AudioMixer/Private/AudioMixerSubmix.cpp
@@ -23,6 +23,7 @@ namespace Audio
 		, SubmixAmbisonicsDecoderID(INDEX_NONE)
 		, EnvelopeNumChannels(0)
 		, bIsRecording(false)
+		, bIsBackgroundMuted(false)
 		, OwningSubmixObject(nullptr)
 	{
 	}
@@ -294,6 +295,14 @@ namespace Audio
 		EffectSubmixChain.Reset();
 	}
 
+	void FMixerSubmix::SetBackgroundMuted(bool bInMuted)
+	{
+		SubmixCommand([this, bInMuted]()
+		{
+			bIsBackgroundMuted = bInMuted;
+		});
+	}
+
 	void FMixerSubmix::FormatChangeBuffer(const ESubmixChannelFormat InNewChannelType, AlignedFloatBuffer& InBuffer, AlignedFloatBuffer& OutNewBuffer)
 	{
 		// Retrieve ptr to the cached downmix channel map from the mixer device
@@ -673,9 +682,21 @@ namespace Audio
 					SubmixEffect->ProcessAudio(InputData, OutputData);
 				}
 
+				// Mix in the dry signal directly
+				const float DryLevel = SubmixEffect->GetDryLevel();
+				if (DryLevel > 0.0f)
+				{
+					Audio::MixInBufferFast(InputBuffer, ScratchBuffer, DryLevel);
+				}
+
 				FMemory::Memcpy((void*)BufferPtr, (void*)ScratchBuffer.GetData(), sizeof(float)*NumSamples);
 			}
+		}
 
+		// If we're muted, memzero the buffer. Note we are still doing all the work to maintain buffer state between mutings.
+		if (bIsBackgroundMuted)
+		{
+			FMemory::Memzero((void*)BufferPtr, sizeof(float)*NumSamples);
 		}
 
 		// If we are recording, Add out buffer to the RecordingData buffer:
@@ -777,10 +798,6 @@ namespace Audio
 		return nullptr;
 	}
 
-	void FMixerSubmix::Update()
-	{
-	}
-
 	void FMixerSubmix::OnAmbisonicsSettingsChanged(UAmbisonicsSubmixSettingsBase* InAmbisonicsSettings)
 	{
 		check(InAmbisonicsSettings != nullptr);
diff --git a/Engine/Source/Runtime/AudioMixer/Private/Effects/AudioMixerSubmixEffectReverb.cpp b/Engine/Source/Runtime/AudioMixer/Private/Effects/AudioMixerSubmixEffectReverb.cpp
--- a/Engine/Source/Runtime/AudioMixer/Private/Effects/AudioMixerSubmixEffectReverb.cpp
+++ b/Engine/Source/Runtime/AudioMixer/Private/Effects/AudioMixerSubmixEffectReverb.cpp
@@ -6,6 +6,14 @@
 #include "Audio.h"
 #include "AudioMixer.h"
 
+static int32 DisableSubmixReverbCVar = 0;
+FAutoConsoleVariableRef CVarDisableSubmixReverb(
+	TEXT("au.DisableReverbSubmix"),
+	DisableSubmixReverbCVar,
+	TEXT("Disables the reverb submix.\n")
+	TEXT("0: Not Disabled, 1: Disabled"),
+	ECVF_Default);
+
 class UReverbEffect;
 
 FSubmixEffectReverb::FSubmixEffectReverb()
@@ -77,7 +85,7 @@ void FSubmixEffectReverb::OnProcessAudio(const FSoundEffectSubmixInputData& InDa
 	LLM_SCOPE(ELLMTag::AudioMixer);
 
 	check(InData.NumChannels == 2);
- 	if (OutData.NumChannels < 2 || !bIsEnabled) 
+ 	if (OutData.NumChannels < 2 || !bIsEnabled || DisableSubmixReverbCVar == 1) 
 	{
 		// Not supported
 		return;
@@ -96,9 +104,6 @@ void FSubmixEffectReverb::OnProcessAudio(const FSoundEffectSubmixInputData& InDa
 		for (int32 SampleIndex = 0; SampleIndex < InData.AudioBuffer->Num(); SampleIndex += OutData.NumChannels)
 		{
 			PlateReverb.ProcessAudioFrame(&AudioData[SampleIndex], InData.NumChannels, &OutAudioData[SampleIndex], OutData.NumChannels);
-
-			OutAudioData[SampleIndex] += DryLevel * AudioData[SampleIndex];
-			OutAudioData[SampleIndex + 1] += DryLevel * AudioData[SampleIndex + 1];
 		}
 	}
 	// 5.1 or higher surround sound. Map stereo output to quad output
@@ -113,10 +118,6 @@ void FSubmixEffectReverb::OnProcessAudio(const FSoundEffectSubmixInputData& InDa
 			// Using standard speaker map order map the right output to the BackLeft channel
 			OutAudioData[OutSampleIndex + EAudioMixerChannel::BackRight] = OutAudioData[OutSampleIndex + EAudioMixerChannel::FrontLeft];
 			OutAudioData[OutSampleIndex + EAudioMixerChannel::BackLeft] = OutAudioData[OutSampleIndex + EAudioMixerChannel::FrontRight];
-
-			// Copy dry output to output data to stereo fronts
-			OutAudioData[OutSampleIndex] += DryLevel * AudioData[InSampleIndex];
-			OutAudioData[OutSampleIndex + 1] += DryLevel * AudioData[InSampleIndex + 1];
 		}
 	}
 }
diff --git a/Engine/Source/Runtime/AudioMixer/Public/AudioMixerDevice.h b/Engine/Source/Runtime/AudioMixer/Public/AudioMixerDevice.h
--- a/Engine/Source/Runtime/AudioMixer/Public/AudioMixerDevice.h
+++ b/Engine/Source/Runtime/AudioMixer/Public/AudioMixerDevice.h
@@ -219,6 +219,7 @@ namespace Audio
 	private:
 
 		bool IsMasterSubmixType(USoundSubmix* InSubmix) const;
+		FMixerSubmix* GetMasterSubmixInstance(USoundSubmix* InSubmix);
 
 		// Pushes the command to a audio render thread command queue to be executed on render thread
 		void AudioRenderThreadCommand(TFunction<void()> Command);
@@ -258,6 +259,9 @@ namespace Audio
 		/** The audio clock from device initialization, updated at block rate. */
 		double AudioClock;
 
+		/** What the previous master volume was. */
+		float PreviousMasterVolume;
+
 		/** Timing data for audio thread. */
 		FAudioThreadTimingData AudioThreadTimingData;
 
diff --git a/Engine/Source/Runtime/AudioMixer/Public/AudioMixerSubmix.h b/Engine/Source/Runtime/AudioMixer/Public/AudioMixerSubmix.h
--- a/Engine/Source/Runtime/AudioMixer/Public/AudioMixerSubmix.h
+++ b/Engine/Source/Runtime/AudioMixer/Public/AudioMixerSubmix.h
@@ -87,6 +87,9 @@ namespace Audio
 		/** Clears all submix effects from the effect submix chain. */
 		void ClearSoundEffectSubmixes();
 
+		/** Whether or not this submix instance is muted. */
+		void SetBackgroundMuted(bool bInMuted);
+
 		// Function which processes audio.
 		void ProcessAudio(const ESubmixChannelFormat ParentInputChannels, AlignedFloatBuffer& OutAudio);
 
@@ -96,9 +99,6 @@ namespace Audio
 		// Returns the output channels this submix is rendering to
 		int32 GetNumOutputChannels() const;
 
-		// Updates the submix from the main thread.
-		void Update();
-
 		// Returns the number of effects in this submix's effect chain
 		int32 GetNumChainEffects() const;
 
@@ -282,6 +282,9 @@ namespace Audio
 		// Bool set to true when this submix is recording data.
 		uint8 bIsRecording : 1;
 
+		// Whether or not this submix is muted.
+		uint8 bIsBackgroundMuted : 1;
+
 		// Bool set to true when envelope following is enabled
 		FThreadSafeBool bIsEnvelopeFollowing;
 
diff --git a/Engine/Source/Runtime/Engine/Classes/Sound/SoundAttenuation.h b/Engine/Source/Runtime/Engine/Classes/Sound/SoundAttenuation.h
--- a/Engine/Source/Runtime/Engine/Classes/Sound/SoundAttenuation.h
+++ b/Engine/Source/Runtime/Engine/Classes/Sound/SoundAttenuation.h
@@ -330,7 +330,7 @@ struct ENGINE_API FSoundAttenuationSettings : public FBaseAttenuationSettings
 		, ReverbWetLevelMax(0.95f)
 		, ReverbDistanceMin(AttenuationShapeExtents.X)
 		, ReverbDistanceMax(AttenuationShapeExtents.X + FalloffDistance)
-		, ManualReverbSendLevel(0.2f)
+		, ManualReverbSendLevel(0.0f)
 	{
 	}
 
diff --git a/Engine/Source/Runtime/Engine/Classes/Sound/SoundClass.h b/Engine/Source/Runtime/Engine/Classes/Sound/SoundClass.h
--- a/Engine/Source/Runtime/Engine/Classes/Sound/SoundClass.h
+++ b/Engine/Source/Runtime/Engine/Classes/Sound/SoundClass.h
@@ -128,7 +128,7 @@ struct FSoundClassProperties
 		, bIsUISound(false)
 		, bIsMusic(false)
 		, bReverb(true)
-		, Default2DReverbSendAmount(0.5f)
+		, Default2DReverbSendAmount(0.0f)
 		, bCenterChannelOnly(false)
 		, bApplyAmbientVolumes(false)
 		, OutputTarget(EAudioOutputTarget::Speaker)
diff --git a/Engine/Source/Runtime/Engine/Classes/Sound/SoundEffectSubmix.h b/Engine/Source/Runtime/Engine/Classes/Sound/SoundEffectSubmix.h
--- a/Engine/Source/Runtime/Engine/Classes/Sound/SoundEffectSubmix.h
+++ b/Engine/Source/Runtime/Engine/Classes/Sound/SoundEffectSubmix.h
@@ -91,6 +91,9 @@ public:
 	/** Process the input block of audio. Called on audio thread. */
 	virtual void OnProcessAudio(const FSoundEffectSubmixInputData& InData, FSoundEffectSubmixOutputData& OutData) {};
 
+	/** Allow effects to supply a drylevel. */
+	virtual float GetDryLevel() const { return 0.0f; }
+
 	/** Processes audio in the submix effect. */
 	void ProcessAudio(FSoundEffectSubmixInputData& InData, FSoundEffectSubmixOutputData& OutData);
 };
diff --git a/Engine/Source/Runtime/Engine/Classes/Sound/SoundSubmix.h b/Engine/Source/Runtime/Engine/Classes/Sound/SoundSubmix.h
--- a/Engine/Source/Runtime/Engine/Classes/Sound/SoundSubmix.h
+++ b/Engine/Source/Runtime/Engine/Classes/Sound/SoundSubmix.h
@@ -117,6 +117,10 @@ class ENGINE_API USoundSubmix : public UObject
 	UPROPERTY()
 	ESubmixChannelFormat ChannelFormat;
 
+	/** Mute this submix when the application is muted or in the background. Used to prevent submix effect tails from continuing when tabbing out of application or if application is muted. */
+	UPROPERTY(EditAnywhere, BlueprintReadOnly, Category = SoundSubmix)
+	uint8 bMuteWhenBackgrounded : 1;
+
 	UPROPERTY(EditAnywhere, BlueprintReadOnly, Category = SoundSubmix)
 	TArray<USoundEffectSubmixPreset*> SubmixEffectChain;
 
diff --git a/Engine/Source/Runtime/Engine/Private/ActiveSound.cpp b/Engine/Source/Runtime/Engine/Private/ActiveSound.cpp
--- a/Engine/Source/Runtime/Engine/Private/ActiveSound.cpp
+++ b/Engine/Source/Runtime/Engine/Private/ActiveSound.cpp
@@ -436,7 +436,7 @@ void FActiveSound::UpdateWaveInstances( TArray<FWaveInstance*> &InWaveInstances,
 	// If the sound is a preview sound, then ignore the transient master volume and application volume
 	if (!bIsPreviewSound)
 	{
-		ParseParams.VolumeApp = AudioDevice->GetTransientMasterVolume() * FApp::GetVolumeMultiplier();
+		ParseParams.VolumeApp = AudioDevice->GetMasterVolume();
 	}
 
 	const float TotalConcurrencyVolumeScale = GetTotalConcurrencyVolumeScale();
diff --git a/Engine/Source/Runtime/Engine/Private/AudioDevice.cpp b/Engine/Source/Runtime/Engine/Private/AudioDevice.cpp
--- a/Engine/Source/Runtime/Engine/Private/AudioDevice.cpp
+++ b/Engine/Source/Runtime/Engine/Private/AudioDevice.cpp
@@ -216,6 +216,7 @@ FAudioDevice::FAudioDevice()
 	, TestAudioComponent(nullptr)
 	, DebugState(DEBUGSTATE_None)
 	, TransientMasterVolume(1.0f)
+	, MasterVolume(1.0f)
 	, GlobalPitchScale(1.0f)
 	, LastUpdateTime(FPlatformTime::Seconds())
 	, NextResourceID(1)
@@ -225,7 +226,7 @@ FAudioDevice::FAudioDevice()
 	, Effects(nullptr)
 	, CurrentReverbEffect(nullptr)
 	, PlatformAudioHeadroom(1.0f)
-	, DefaultReverbSendLevel(0.2f)
+	, DefaultReverbSendLevel(0.0f)
 	, bHRTFEnabledForAll_OnGameThread(false)
 	, bGameWasTicking(true)
 	, bDisableAudioCaching(false)
@@ -3725,6 +3726,9 @@ void FAudioDevice::Update(bool bGameTicking)
 
 	bIsStoppingVoicesEnabled = !DisableStoppingVoicesCvar;
 
+	// Update the master volume
+	MasterVolume = GetTransientMasterVolume() * FApp::GetVolumeMultiplier();
+
 	UpdateAudioPluginSettingsObjectCache();
 
 	{
diff --git a/Engine/Source/Runtime/Engine/Private/AudioSettings.cpp b/Engine/Source/Runtime/Engine/Private/AudioSettings.cpp
--- a/Engine/Source/Runtime/Engine/Private/AudioSettings.cpp
+++ b/Engine/Source/Runtime/Engine/Private/AudioSettings.cpp
@@ -68,7 +68,7 @@ void UAudioSettings::AddDefaultSettings()
 	GConfig->GetInt(TEXT("Audio"), TEXT("MaxChannels"), DefaultSettings.MaxChannels, GEngineIni); // for backwards compatibility
 	QualityLevels.Add(DefaultSettings);
 	bAllowVirtualizedSounds = true;
-	DefaultReverbSendLevel = 0.2f;
+	DefaultReverbSendLevel = 0.0f;
 	VoiPSampleRate = EVoiceSampleRate::Low16000Hz;
 	VoipBufferingDelay = 0.2f;
 	MaxWaveInstances = 100;
diff --git a/Engine/Source/Runtime/Engine/Private/SoundWave.cpp b/Engine/Source/Runtime/Engine/Private/SoundWave.cpp
--- a/Engine/Source/Runtime/Engine/Private/SoundWave.cpp
+++ b/Engine/Source/Runtime/Engine/Private/SoundWave.cpp
@@ -1599,8 +1599,11 @@ void USoundWave::Parse( FAudioDevice* AudioDevice, const UPTRINT NodeWaveInstanc
 		}
 
 		// Copy reverb send settings
-		WaveInstance->ReverbSendMethod = ParseParams.ReverbSendMethod;
-		WaveInstance->ManualReverbSendLevel = ParseParams.ManualReverbSendLevel;
+		if (!bUseSoundClassDefaultReverb)
+		{
+			WaveInstance->ReverbSendMethod = ParseParams.ReverbSendMethod;
+			WaveInstance->ManualReverbSendLevel = ParseParams.ManualReverbSendLevel;
+		}
 		WaveInstance->CustomRevebSendCurve = ParseParams.CustomReverbSendCurve;
 		WaveInstance->ReverbSendLevelRange = ParseParams.ReverbSendLevelRange;
 		WaveInstance->ReverbSendLevelDistanceRange = ParseParams.ReverbSendLevelDistanceRange;
diff --git a/Engine/Source/Runtime/Engine/Public/ActiveSound.h b/Engine/Source/Runtime/Engine/Public/ActiveSound.h
--- a/Engine/Source/Runtime/Engine/Public/ActiveSound.h
+++ b/Engine/Source/Runtime/Engine/Public/ActiveSound.h
@@ -172,7 +172,7 @@ struct FSoundParseParameters
 		, ReverbSendMethod(EReverbSendMethod::Linear)
 		, ReverbSendLevelRange(0.0f, 0.0f)
 		, ReverbSendLevelDistanceRange(0.0f, 0.0f)
-		, ManualReverbSendLevel(0.2f)
+		, ManualReverbSendLevel(0.0f)
 		, StereoSpread(0.0f)
 		, SpatializationMethod(ESoundSpatializationAlgorithm::SPATIALIZATION_Default)
 		, SpatializationPluginSettings(nullptr)
diff --git a/Engine/Source/Runtime/Engine/Public/AudioDevice.h b/Engine/Source/Runtime/Engine/Public/AudioDevice.h
--- a/Engine/Source/Runtime/Engine/Public/AudioDevice.h
+++ b/Engine/Source/Runtime/Engine/Public/AudioDevice.h
@@ -1526,6 +1526,9 @@ public:
 	float GetTransientMasterVolume() const { check(IsInAudioThread()); return TransientMasterVolume; }
 	void SetTransientMasterVolume(float TransientMasterVolume);
 
+	/** Returns the volume that combines transient master volume and the FApp::GetVolumeMultiplier() value */
+	float GetMasterVolume() const { return MasterVolume; }
+
 	FSoundSource* GetSoundSource(FWaveInstance* WaveInstance) const;
 
 	const FGlobalFocusSettings& GetGlobalFocusSettings() const;
@@ -1631,6 +1634,9 @@ private:
 	/** transient master volume multiplier that can be modified at runtime without affecting user settings automatically reset to 1.0 on level change */
 	float TransientMasterVolume;
 
+	/** The master volume of the game combines the FApp::GetVolumeMultipler() value and the TransientMastervolume. */
+	float MasterVolume;
+
 	/** Global dynamic pitch scale parameter */
 	FDynamicParameter GlobalPitchScale;
 
