Fixing crash involving async decoding during audio streaming.

Issue is streaming sound is being removed from the streaming manager before it's stopped playing in the audio renderer. Fix is to prevent any stream chunk or sound from getting removed from the streaming manager if there is an active decoder using any stream chunk.

diff --git a/Engine/Source/Runtime/AudioMixer/Private/AudioMixerBuffer.cpp b/Engine/Source/Runtime/AudioMixer/Private/AudioMixerBuffer.cpp
--- a/Engine/Source/Runtime/AudioMixer/Private/AudioMixerBuffer.cpp
+++ b/Engine/Source/Runtime/AudioMixer/Private/AudioMixerBuffer.cpp
@@ -3,7 +3,7 @@
 #include "AudioMixerBuffer.h"
 #include "AudioMixerDevice.h"
 #include "Interfaces/IAudioFormat.h"
-
+#include "AudioMixerSourceDecode.h"
 namespace Audio
 {
 
@@ -153,26 +153,6 @@ namespace Audio
 		return false;
 	}
 
-	bool FMixerBuffer::ReadCompressedData(uint8* Destination, int32 InNumFrames, bool bLooping)
-	{
-		if (!DecompressionState)
-		{
-			UE_LOG(LogAudioMixer, Warning, TEXT("Attempting to read compressed data without a compression state instance for resource '%s'"), *ResourceName);
-			return false;
-		}
-
-		const int32 kPCMBufferSize = NumChannels * InNumFrames * sizeof(int16);
-
-		if (BufferType == EBufferType::Streaming)
-		{
-			return DecompressionState->StreamCompressedData(Destination, bLooping, kPCMBufferSize);
-		}
-		else
-		{
-			return DecompressionState->ReadCompressedData(Destination, bLooping, kPCMBufferSize);
-		}
-	}
-
 	void FMixerBuffer::Seek(const float SeekTime)
 	{
 		if (ensure(DecompressionState))
@@ -403,6 +383,16 @@ namespace Audio
 		return BufferType == EBufferType::PCMRealTime || BufferType == EBufferType::Streaming;
 	}
 
+	ICompressedAudioInfo* FMixerBuffer::GetDecompressionState(bool bTakesOwnership)
+	{
+		ICompressedAudioInfo* Output = DecompressionState;
+		if (bTakesOwnership)
+		{
+			DecompressionState = nullptr;
+		}
+		return Output;
+	}
+
 	void FMixerBuffer::GetPCMData(uint8** OutData, uint32* OutDataSize)
 	{
 		*OutData = Data;
diff --git a/Engine/Source/Runtime/AudioMixer/Private/AudioMixerBuffer.h b/Engine/Source/Runtime/AudioMixer/Private/AudioMixerBuffer.h
--- a/Engine/Source/Runtime/AudioMixer/Private/AudioMixerBuffer.h
+++ b/Engine/Source/Runtime/AudioMixer/Private/AudioMixerBuffer.h
@@ -3,7 +3,6 @@
 #pragma once
 
 #include "CoreMinimal.h"
-#include "AudioMixerSourceDecode.h"
 #include "AudioMixer.h"
 
 namespace Audio
@@ -22,6 +21,7 @@ namespace Audio
 
 	class FMixerDevice;
 	class FMixerBuffer;
+	class IAudioTask;
 
 	class FMixerBuffer : public FSoundBuffer
 	{
@@ -35,7 +35,6 @@ namespace Audio
 		int32 GetCurrentChunkOffset() const override;
 		bool IsRealTimeSourceReady() override;
 		bool ReadCompressedInfo(USoundWave* SoundWave) override;
-		bool ReadCompressedData(uint8* Destination, int32 NumFrames, bool bLooping) override;
 		void Seek(const float SeekTime) override;
 		//~ End FSoundBuffer Interface
 
@@ -50,6 +49,9 @@ namespace Audio
 		EBufferType::Type GetType() const;
 		bool IsRealTimeBuffer() const;
 
+		/** Retrieve the compressed audio info pointer. This transfers ownership of the decompression state. */
+		ICompressedAudioInfo* GetDecompressionState(bool bTakesOwnership = false);
+
 		/** Returns the contained raw PCM data and data size */
 		void GetPCMData(uint8** OutData, uint32* OutDataSize);
 
diff --git a/Engine/Source/Runtime/AudioMixer/Private/AudioMixerSource.cpp b/Engine/Source/Runtime/AudioMixer/Private/AudioMixerSource.cpp
--- a/Engine/Source/Runtime/AudioMixer/Private/AudioMixerSource.cpp
+++ b/Engine/Source/Runtime/AudioMixer/Private/AudioMixerSource.cpp
@@ -282,6 +282,10 @@ namespace Audio
 			// Update the buffer sample rate to the wave instance sample rate in case it was serialized incorrectly
 			MixerBuffer->InitSampleRate(WaveInstance->WaveData->GetSampleRateForCurrentPlatform());
 
+			// Pass the decompression state off to the mixer source buffer
+			ICompressedAudioInfo* Decoder = MixerBuffer->GetDecompressionState(true);
+			MixerSourceBuffer->SetDecoder(Decoder);
+
 			// Hand off the mixer source buffer decoder
 			InitParams.MixerSourceBuffer = MixerSourceBuffer;
 			MixerSourceBuffer = nullptr;
@@ -427,13 +431,14 @@ namespace Audio
 						if (WaveInstance->StartTime > 0.0f || WaveInstance->WaveData->bProcedural || WaveInstance->WaveData->bIsBus || !WaveInstance->WaveData->CachedRealtimeFirstBuffer)
 						{
 							// Before reading more PCMRT data, we first need to seek the buffer
-							if (WaveInstance->StartTime > 0.0f && !WaveInstance->WaveData->bIsBus && !WaveInstance->WaveData->bProcedural)
+							if (WaveInstance->StartTime > 0.0f && !(BufferType == EBufferType::Streaming) && !WaveInstance->WaveData->bIsBus && !WaveInstance->WaveData->bProcedural)
 							{
 								MixerBuffer->Seek(WaveInstance->StartTime);
 							}
 
 							check(MixerSourceBuffer.IsValid());
-							MixerSourceBuffer->ReadMoreRealtimeData(0, EBufferReadMode::Asynchronous);
+							ICompressedAudioInfo* Decoder = MixerBuffer->GetDecompressionState(false);
+							MixerSourceBuffer->ReadMoreRealtimeData(Decoder, 0, EBufferReadMode::Asynchronous);
 
 							// not ready
 							return false;
@@ -485,6 +490,11 @@ namespace Audio
 	{
 		LLM_SCOPE(ELLMTag::AudioMixer);
 
+		if (InitializationState == EMixerSourceInitializationState::NotInitialized)
+		{
+			return;
+		}
+
 		if (!MixerSourceVoice)
 		{
 			StopNow();
@@ -687,11 +697,21 @@ namespace Audio
 		}
 
 		MixerSourceBuffer = nullptr;
-		MixerBuffer = nullptr;
 		Buffer = nullptr;
 		bLoopCallback = false;
 		NumTotalFrames = 0;
 
+		if (MixerBuffer)
+		{
+			EBufferType::Type BufferType = MixerBuffer->GetType();
+			if (BufferType == EBufferType::PCMRealTime || BufferType == EBufferType::Streaming)
+			{
+				delete MixerBuffer;
+			}
+
+			MixerBuffer = nullptr;
+		}
+
 		// Reset the source's channel maps
 		for (int32 i = 0; i < (int32)ESubmixChannelFormat::Count; ++i)
 		{
diff --git a/Engine/Source/Runtime/AudioMixer/Private/AudioMixerSourceBuffer.cpp b/Engine/Source/Runtime/AudioMixer/Private/AudioMixerSourceBuffer.cpp
--- a/Engine/Source/Runtime/AudioMixer/Private/AudioMixerSourceBuffer.cpp
+++ b/Engine/Source/Runtime/AudioMixer/Private/AudioMixerSourceBuffer.cpp
@@ -1,6 +1,8 @@
 // Copyright 1998-2019 Epic Games, Inc. All Rights Reserved.
 
 #include "AudioMixerSourceBuffer.h"
+#include "AudioMixerSourceDecode.h"
+#include "ContentStreaming.h"
 
 namespace Audio
 {
@@ -55,10 +57,12 @@ namespace Audio
 	FMixerSourceBuffer::FMixerSourceBuffer()
 		: NumBuffersQeueued(0)
 		, CurrentBuffer(0)
-		, MixerBuffer(nullptr)
 		, SoundWave(nullptr)
 		, AsyncRealtimeAudioTask(nullptr)
+		, DecompressionState(nullptr)
 		, LoopingMode(ELoopingMode::LOOP_Never)
+		, BufferType(Audio::EBufferType::Invalid)
+		, NumChannels(0)
 		, bInitialized(false)
 		, bBufferFinished(false)
 		, bPlayedCachedBuffer(false)
@@ -69,16 +73,42 @@ namespace Audio
 
 	FMixerSourceBuffer::~FMixerSourceBuffer()
 	{
+		// Make sure we have completed our async realtime task before deleting the decompression state
+		if (AsyncRealtimeAudioTask)
+		{
+			AsyncRealtimeAudioTask->CancelTask();
+			delete AsyncRealtimeAudioTask;
+			AsyncRealtimeAudioTask = nullptr;
+		}
+
 		OnEndGenerate();
+
+		// Clean up decompression state after things have been finished using it
+		if (DecompressionState)
+		{
+			IStreamingManager::Get().GetAudioStreamingManager().RemoveDecoder(DecompressionState);
+
+			delete DecompressionState;
+			DecompressionState = nullptr;
+		}
 	}
 
 	bool FMixerSourceBuffer::PreInit(FMixerBuffer* InBuffer, USoundWave* InWave, ELoopingMode InLoopingMode, bool bInIsSeeking)
 	{
 		LLM_SCOPE(ELLMTag::AudioMixer);
 
-		// Mixer source buffer now owns this buffer
-		MixerBuffer = InBuffer;
-		check(MixerBuffer);
+		NumChannels = InBuffer->NumChannels;
+		BufferType = InBuffer->GetType();
+
+		// If this is a PCM buffer type, we need to take ownership of the raw pcm data
+		if (BufferType == EBufferType::PCM || BufferType == EBufferType::PCMPreview)
+		{
+			check(RawPCMDataBuffer.Data == nullptr);
+			check(RawPCMDataBuffer.DataSize == 0);
+
+			// GetPCMData will retrieve ownership of the raw PCM data
+			InBuffer->GetPCMData(&RawPCMDataBuffer.Data, &RawPCMDataBuffer.DataSize);
+		}
 
 		// May or may not be nullptr
 		SoundWave = InWave;
@@ -89,7 +119,7 @@ namespace Audio
 
 		BufferQueue.Empty();
 
-		const uint32 TotalSamples = MONO_PCM_BUFFER_SAMPLES * MixerBuffer->NumChannels;
+		const uint32 TotalSamples = MONO_PCM_BUFFER_SAMPLES * NumChannels;
 		for (int32 BufferIndex = 0; BufferIndex < Audio::MAX_BUFFERS_QUEUED; ++BufferIndex)
 		{
 			SourceVoiceBuffers.Add(TSharedPtr<FMixerSourceVoiceBuffer>(new FMixerSourceVoiceBuffer()));
@@ -102,6 +132,15 @@ namespace Audio
 		return true;
 	}
 
+	void FMixerSourceBuffer::SetDecoder(ICompressedAudioInfo* InCompressedAudioInfo)
+	{
+		DecompressionState = InCompressedAudioInfo;
+		if (BufferType == EBufferType::Streaming)
+		{
+			IStreamingManager::Get().GetAudioStreamingManager().AddDecoder(DecompressionState);
+		}
+	}
+
 	bool FMixerSourceBuffer::Init()
 	{
 		check(SoundWave);
@@ -116,7 +155,6 @@ namespace Audio
 		// SoundWave->bIsSoundActive will prevent GC until it is released in audio render thread
 		bInitialized = true;
 
-		const EBufferType::Type BufferType = MixerBuffer->GetType();
 		switch (BufferType)
 		{
 			case EBufferType::PCM:
@@ -158,9 +196,7 @@ namespace Audio
 	{
 		CurrentBuffer = 0;
 
-		RawPCMDataBuffer.Data = nullptr;
-		RawPCMDataBuffer.DataSize = 0;
-		MixerBuffer->GetPCMData(&RawPCMDataBuffer.Data, &RawPCMDataBuffer.DataSize);
+		check(RawPCMDataBuffer.Data != nullptr);
 
 		RawPCMDataBuffer.NumSamples = RawPCMDataBuffer.DataSize / sizeof(int16);
 		RawPCMDataBuffer.CurrentSample = 0;
@@ -174,7 +210,7 @@ namespace Audio
 		RawPCMDataBuffer.LoopCount = (LoopingMode != LOOP_Never) ? Audio::LOOP_FOREVER : 0;
 
 		// Submit the first two format-converted chunks to the source voice
-		const uint32 NumSamplesPerBuffer = MONO_PCM_BUFFER_SAMPLES * MixerBuffer->NumChannels;
+		const uint32 NumSamplesPerBuffer = MONO_PCM_BUFFER_SAMPLES * NumChannels;
 		int16* RawPCMBufferDataPtr = (int16*)RawPCMDataBuffer.Data;
 
 		// Prepare the buffer for the PCM submission
@@ -200,7 +236,7 @@ namespace Audio
 			bPlayedCachedBuffer = true;
 
 			const int32 NumPrecacheFrames = SoundWave->NumPrecacheFrames;
-			const uint32 NumSamples = NumPrecacheFrames * MixerBuffer->NumChannels;
+			const uint32 NumSamples = NumPrecacheFrames * NumChannels;
 			const uint32 BufferSize = NumSamples * sizeof(int16);
 
 			// Format convert the first cached buffers
@@ -262,9 +298,9 @@ namespace Audio
 		}
 	}
 
-	bool FMixerSourceBuffer::ReadMoreRealtimeData(const int32 BufferIndex, EBufferReadMode BufferReadMode)
+	bool FMixerSourceBuffer::ReadMoreRealtimeData(ICompressedAudioInfo* InDecoder, const int32 BufferIndex, EBufferReadMode BufferReadMode)
 	{
-		const int32 MaxSamples = MONO_PCM_BUFFER_SAMPLES * MixerBuffer->NumChannels;
+		const int32 MaxSamples = MONO_PCM_BUFFER_SAMPLES * NumChannels;
 		SourceVoiceBuffers[BufferIndex]->AudioData.Reset();
 		SourceVoiceBuffers[BufferIndex]->AudioData.AddUninitialized(MaxSamples);
 
@@ -274,14 +310,14 @@ namespace Audio
 			NewTaskData.ProceduralSoundWave = SoundWave;
 			NewTaskData.AudioData = SourceVoiceBuffers[BufferIndex]->AudioData.GetData();
 			NewTaskData.NumSamples = MaxSamples;
-			NewTaskData.NumChannels = MixerBuffer->NumChannels;
+			NewTaskData.NumChannels = NumChannels;
 			check(!AsyncRealtimeAudioTask);
 			AsyncRealtimeAudioTask = CreateAudioTask(NewTaskData);
 
 			// Procedural sound waves never loop
 			return false;
 		}
-		else if (!MixerBuffer->IsRealTimeBuffer())
+		else if (BufferType != EBufferType::PCMRealTime && BufferType != EBufferType::Streaming)
 		{
 			check(RawPCMDataBuffer.Data != nullptr);
 
@@ -289,9 +325,13 @@ namespace Audio
 			return RawPCMDataBuffer.GetNextBuffer(SourceVoiceBuffers[BufferIndex].Get(), MaxSamples);
 		}
 
+		check(InDecoder != nullptr);
+
 		FDecodeAudioTaskData NewTaskData;
-		NewTaskData.MixerBuffer = MixerBuffer;
 		NewTaskData.AudioData = SourceVoiceBuffers[BufferIndex]->AudioData.GetData();
+		NewTaskData.DecompressionState = InDecoder;
+		NewTaskData.BufferType = BufferType;
+		NewTaskData.NumChannels = NumChannels;
 		NewTaskData.bLoopingMode = LoopingMode != LOOP_Never;
 		NewTaskData.bSkipFirstBuffer = (BufferReadMode == EBufferReadMode::AsynchronousSkipFirstFrame);
 		NewTaskData.NumFramesToDecode = MONO_PCM_BUFFER_SAMPLES;
@@ -395,7 +435,7 @@ namespace Audio
 				DataReadMode = EBufferReadMode::Asynchronous;
 			}
 
-			const bool bLooped = ReadMoreRealtimeData(CurrentBuffer, DataReadMode);
+			const bool bLooped = ReadMoreRealtimeData(DecompressionState, CurrentBuffer, DataReadMode);
 
 			// If this was a synchronous read, then immediately write it
 			if (AsyncRealtimeAudioTask == nullptr)
@@ -429,7 +469,7 @@ namespace Audio
 	{
 		if (AsyncRealtimeAudioTask) 
 		{ 
-			AsyncRealtimeAudioTask->EnsureCompletion(); 
+			AsyncRealtimeAudioTask->CancelTask(); 
 
 			delete AsyncRealtimeAudioTask;
 			AsyncRealtimeAudioTask = nullptr;
@@ -476,17 +516,6 @@ namespace Audio
 
 			SoundWave = nullptr;
 		}
-
-		if (MixerBuffer)
-		{
-			EBufferType::Type BufferType = MixerBuffer->GetType();
-			if (BufferType == EBufferType::PCMRealTime || BufferType == EBufferType::Streaming)
-			{
-				delete MixerBuffer;
-			}
-
-			MixerBuffer = nullptr;
-		}
 	}
 
 }
diff --git a/Engine/Source/Runtime/AudioMixer/Private/AudioMixerSourceBuffer.h b/Engine/Source/Runtime/AudioMixer/Private/AudioMixerSourceBuffer.h
--- a/Engine/Source/Runtime/AudioMixer/Private/AudioMixerSourceBuffer.h
+++ b/Engine/Source/Runtime/AudioMixer/Private/AudioMixerSourceBuffer.h
@@ -52,6 +52,9 @@ namespace Audio
 		bool PreInit(FMixerBuffer* InBuffer, USoundWave* InWave, ELoopingMode InLoopingMode, bool bInIsSeeking);
 		bool Init();
 
+		// Sets the decoder to use for realtime async decoding
+		void SetDecoder(ICompressedAudioInfo* InCompressedAudioInfo);
+
 		// Called by source manager when needing more buffers
 		void OnBufferEnd();
 
@@ -68,7 +71,7 @@ namespace Audio
 		bool DidBufferFinish() const { return bBufferFinished; }
 
 		// Called to start an async task to read more data
-		bool ReadMoreRealtimeData(const int32 BufferIndex, EBufferReadMode BufferReadMode);
+		bool ReadMoreRealtimeData(ICompressedAudioInfo* InDecoder, int32 BufferIndex, EBufferReadMode BufferReadMode);
 
 		// Returns true if async task is in progress
 		bool IsAsyncTaskInProgress() const;
@@ -104,15 +107,17 @@ namespace Audio
 		TArray<TSharedPtr<FMixerSourceVoiceBuffer>> SourceVoiceBuffers;
 		TQueue<TSharedPtr<FMixerSourceVoiceBuffer>> BufferQueue;
 		int32 CurrentBuffer;
-		FMixerBuffer* MixerBuffer;
 		USoundWave* SoundWave;
 		IAudioTask* AsyncRealtimeAudioTask;
+		ICompressedAudioInfo* DecompressionState;
 		ELoopingMode LoopingMode;
-		bool bInitialized;
-		bool bBufferFinished;
-		bool bPlayedCachedBuffer;
-		bool bIsSeeking;
-		bool bLoopCallback;
+		int32 NumChannels;
+		Audio::EBufferType::Type BufferType;
+		uint32 bInitialized : 1;
+		uint32 bBufferFinished : 1;
+		uint32 bPlayedCachedBuffer : 1;
+		uint32 bIsSeeking : 1;
+		uint32 bLoopCallback : 1;
 
 	};
 }
diff --git a/Engine/Source/Runtime/AudioMixer/Private/AudioMixerSourceDecode.cpp b/Engine/Source/Runtime/AudioMixer/Private/AudioMixerSourceDecode.cpp
--- a/Engine/Source/Runtime/AudioMixer/Private/AudioMixerSourceDecode.cpp
+++ b/Engine/Source/Runtime/AudioMixer/Private/AudioMixerSourceDecode.cpp
@@ -94,7 +94,7 @@ public:
 
 			case EAudioTaskType::Decode:
 			{
-				int32 NumChannels = DecodeTaskData.MixerBuffer->GetNumChannels();
+				int32 NumChannels = DecodeTaskData.NumChannels;
 				int32 ByteSize = NumChannels * DecodeTaskData.NumFramesToDecode * sizeof(int16);
 
 				// Create a buffer to decode into that's of the appropriate size
@@ -104,13 +104,32 @@ public:
 				// skip the first buffers if we've already decoded them during Precache:
 				if (DecodeTaskData.bSkipFirstBuffer)
 				{
-					for (int32 NumberOfBuffersToSkip = 0; NumberOfBuffersToSkip < PLATFORM_NUM_AUDIODECOMPRESSION_PRECACHE_BUFFERS; NumberOfBuffersToSkip++)
+					const int32 kPCMBufferSize = NumChannels * DecodeTaskData.NumPrecacheFrames * sizeof(int16);
+					if (DecodeTaskData.BufferType == EBufferType::Streaming)
 					{
-						DecodeTaskData.MixerBuffer->ReadCompressedData(DecodeBuffer.GetData(), DecodeTaskData.NumPrecacheFrames, DecodeTaskData.bLoopingMode);
+						for (int32 NumberOfBuffersToSkip = 0; NumberOfBuffersToSkip < PLATFORM_NUM_AUDIODECOMPRESSION_PRECACHE_BUFFERS; NumberOfBuffersToSkip++)
+						{
+							DecodeTaskData.DecompressionState->StreamCompressedData(DecodeBuffer.GetData(), DecodeTaskData.bLoopingMode, kPCMBufferSize);
+						}
+					}
+					else
+					{
+						for (int32 NumberOfBuffersToSkip = 0; NumberOfBuffersToSkip < PLATFORM_NUM_AUDIODECOMPRESSION_PRECACHE_BUFFERS; NumberOfBuffersToSkip++)
+						{
+							DecodeTaskData.DecompressionState->ReadCompressedData(DecodeBuffer.GetData(), DecodeTaskData.bLoopingMode, kPCMBufferSize);
+						}
 					}
 				}
 
-				DecodeResult.bLooped = DecodeTaskData.MixerBuffer->ReadCompressedData(DecodeBuffer.GetData(), DecodeTaskData.NumFramesToDecode, DecodeTaskData.bLoopingMode);
+				const int32 kPCMBufferSize = NumChannels * DecodeTaskData.NumFramesToDecode * sizeof(int16);
+				if (DecodeTaskData.BufferType == EBufferType::Streaming)
+				{
+					DecodeResult.bLooped = DecodeTaskData.DecompressionState->StreamCompressedData(DecodeBuffer.GetData(), DecodeTaskData.bLoopingMode, kPCMBufferSize);
+				}
+				else
+				{
+					DecodeResult.bLooped = DecodeTaskData.DecompressionState->ReadCompressedData(DecodeBuffer.GetData(), DecodeTaskData.bLoopingMode, kPCMBufferSize);
+				}
 
 				// Convert the decoded PCM data into a float buffer while still in the async task
 				int32 SampleIndex = 0;
@@ -175,6 +194,18 @@ public:
 		}
 	}
 
+	virtual void CancelTask() override
+	{
+		if (Task)
+		{
+			// If Cancel returns false, it means we weren't able to cancel. So lets then fallback to ensure complete.
+			if (!Task->Cancel())
+			{
+				Task->EnsureCompletion();
+			}
+		}
+	}
+
 protected:
 
 	FAsyncTask<FAsyncDecodeWorker>* Task;
diff --git a/Engine/Source/Runtime/AudioMixer/Private/AudioMixerSourceDecode.h b/Engine/Source/Runtime/AudioMixer/Private/AudioMixerSourceDecode.h
--- a/Engine/Source/Runtime/AudioMixer/Private/AudioMixerSourceDecode.h
+++ b/Engine/Source/Runtime/AudioMixer/Private/AudioMixerSourceDecode.h
@@ -5,6 +5,8 @@
 #include "CoreMinimal.h"
 #include "Containers/Queue.h"
 #include "Sound/SoundWaveProcedural.h"
+#include "AudioDecompress.h"
+#include "AudioMixerBuffer.h"
 
 namespace Audio
 {
@@ -36,12 +38,18 @@ namespace Audio
 	// Data needed for a decode audio task
 	struct FDecodeAudioTaskData
 	{
-		// The mixer buffer buffer object which holds state to use for the decode operation
-		FMixerBuffer* MixerBuffer;
-
 		// A pointer to a buffer of audio which will be decoded to
 		float* AudioData;
 
+		// Decompression state for decoder
+		ICompressedAudioInfo* DecompressionState;
+
+		// The buffer type for the decoder
+		Audio::EBufferType::Type BufferType;
+
+		// Number of channels of the decoder
+		int32 NumChannels;
+
 		// The number of frames which are precached
 		int32 NumPrecacheFrames;
 
@@ -55,8 +63,10 @@ namespace Audio
 		bool bSkipFirstBuffer;
 
 		FDecodeAudioTaskData()
-			: MixerBuffer(nullptr)
-			, AudioData(nullptr)
+			: AudioData(nullptr)
+			, DecompressionState(nullptr)
+			, BufferType(Audio::EBufferType::Invalid)
+			, NumChannels(0)
 			, NumPrecacheFrames(0)
 			, NumFramesToDecode(0)
 			, bLoopingMode(false)
@@ -131,6 +141,9 @@ namespace Audio
 		// Ensures the completion of the decode operation.
 		virtual void EnsureCompletion() = 0;
 
+		// Cancel the decode operation
+		virtual void CancelTask() = 0;
+
 		// Returns the result of a procedural sound generate job
 		virtual void GetResult(FProceduralAudioTaskResults& OutResult) {};
 
diff --git a/Engine/Source/Runtime/AudioMixer/Private/DSP/SoundWaveDecoder.cpp b/Engine/Source/Runtime/AudioMixer/Private/DSP/SoundWaveDecoder.cpp
--- a/Engine/Source/Runtime/AudioMixer/Private/DSP/SoundWaveDecoder.cpp
+++ b/Engine/Source/Runtime/AudioMixer/Private/DSP/SoundWaveDecoder.cpp
@@ -65,7 +65,9 @@ namespace Audio
 							MixerBuffer->Seek(SeekTime);
 						}
 
-						MixerSourceBuffer.ReadMoreRealtimeData(0, EBufferReadMode::Asynchronous);
+						ICompressedAudioInfo* CompressedAudioInfo = MixerBuffer->GetDecompressionState(false);
+
+						MixerSourceBuffer.ReadMoreRealtimeData(CompressedAudioInfo, 0, EBufferReadMode::Asynchronous);
 
 						// not ready
 						return false;
diff --git a/Engine/Source/Runtime/Engine/Private/ADPCMAudioInfo.cpp b/Engine/Source/Runtime/Engine/Private/ADPCMAudioInfo.cpp
--- a/Engine/Source/Runtime/Engine/Private/ADPCMAudioInfo.cpp
+++ b/Engine/Source/Runtime/Engine/Private/ADPCMAudioInfo.cpp
@@ -272,11 +272,11 @@ int FADPCMAudioInfo::GetStreamBufferSize() const
 	return StreamBufferSize;
 }
 
-bool FADPCMAudioInfo::StreamCompressedInfo(USoundWave* Wave, struct FSoundQualityInfo* QualityInfo)
+bool FADPCMAudioInfo::StreamCompressedInfoInternal(USoundWave* Wave, struct FSoundQualityInfo* QualityInfo)
 {
 	check(QualityInfo);
 
-	StreamingSoundWave = Wave;
+	check(StreamingSoundWave == Wave);
 
 	// Get the first chunk of audio data (should already be loaded)
 	uint8 const* firstChunk = IStreamingManager::Get().GetAudioStreamingManager().GetLoadedChunk(Wave, 0, &CurrentChunkDataSize);
diff --git a/Engine/Source/Runtime/Engine/Private/AudioDecompress.cpp b/Engine/Source/Runtime/Engine/Private/AudioDecompress.cpp
--- a/Engine/Source/Runtime/Engine/Private/AudioDecompress.cpp
+++ b/Engine/Source/Runtime/Engine/Private/AudioDecompress.cpp
@@ -23,7 +23,6 @@ IStreamedCompressedInfo::IStreamedCompressedInfo()
 	, LastPCMByteSize(0)
 	, LastPCMOffset(0)
 	, bStoringEndOfFile(false)
-	, StreamingSoundWave(nullptr)
 	, CurrentChunkIndex(0)
 	, bPrintChunkFailMessage(true)
 	, SrcBufferPadding(0)
@@ -142,9 +141,9 @@ void IStreamedCompressedInfo::ExpandFile(uint8* DstBuffer, struct FSoundQualityI
 	}
 }
 
-bool IStreamedCompressedInfo::StreamCompressedInfo(USoundWave* Wave, struct FSoundQualityInfo* QualityInfo)
+bool IStreamedCompressedInfo::StreamCompressedInfoInternal(USoundWave* Wave, struct FSoundQualityInfo* QualityInfo)
 {
-	StreamingSoundWave = Wave;
+	check(StreamingSoundWave == Wave);
 
 	// Get the first chunk of audio data (should always be loaded)
 	CurrentChunkIndex = 0;
diff --git a/Engine/Source/Runtime/Engine/Private/AudioDevice.cpp b/Engine/Source/Runtime/Engine/Private/AudioDevice.cpp
--- a/Engine/Source/Runtime/Engine/Private/AudioDevice.cpp
+++ b/Engine/Source/Runtime/Engine/Private/AudioDevice.cpp
@@ -3910,6 +3910,11 @@ void FAudioDevice::Update(bool bGameTicking)
 void FAudioDevice::SendUpdateResultsToGameThread(const int32 FirstActiveIndex)
 {
 #if !UE_BUILD_SHIPPING
+	if (FirstActiveIndex == INDEX_NONE)
+	{
+		return;
+	}
+
 	TArray<FAudioStats::FStatSoundInfo> StatSoundInfos;
 	TArray<FAudioStats::FStatSoundMix> StatSoundMixes;
 	const FVector ListenerPosition = Listeners[0].Transform.GetTranslation();
diff --git a/Engine/Source/Runtime/Engine/Private/AudioStreaming.cpp b/Engine/Source/Runtime/Engine/Private/AudioStreaming.cpp
--- a/Engine/Source/Runtime/Engine/Private/AudioStreaming.cpp
+++ b/Engine/Source/Runtime/Engine/Private/AudioStreaming.cpp
@@ -16,6 +16,7 @@ AudioStreaming.cpp: Implementation of audio streaming classes.
 #include "Misc/ScopeLock.h"
 #include "HAL/IConsoleManager.h"
 #include "HAL/LowLevelMemTracker.h"
+#include "AudioDecompress.h"
 
 static int32 SpoofFailedStreamChunkLoad = 0;
 FAutoConsoleVariableRef CVarSpoofFailedStreamChunkLoad(
@@ -580,6 +581,33 @@ void FAudioStreamingManager::UpdateResourceStreaming(float DeltaTime, bool bProc
 			}
 		}
 	}
+
+	for (ICompressedAudioInfo* Decoder : CompressedAudioInfos)
+	{
+		USoundWave* SoundWave = Decoder->GetStreamingSoundWave();
+		if (SoundWave)
+		{
+			FStreamingWaveData** WaveDataPtr = StreamingSoundWaves.Find(SoundWave);
+			if (WaveDataPtr && (*WaveDataPtr)->PendingChunkChangeRequestStatus.GetValue() == AudioState_ReadyFor_Requests)
+			{
+				FStreamingWaveData* WaveData = *WaveDataPtr;
+				// Request the chunk the source is using and the one after that
+				FWaveRequest& WaveRequest = GetWaveRequest(SoundWave);
+				int32 SourceChunk = Decoder->GetCurrentChunkIndex();
+				if (SourceChunk >= 0 && SourceChunk < SoundWave->RunningPlatformData->NumChunks)
+				{
+					WaveRequest.RequiredIndices.AddUnique(SourceChunk);
+					WaveRequest.RequiredIndices.AddUnique((SourceChunk + 1) % SoundWave->RunningPlatformData->NumChunks);
+					WaveRequest.bPrioritiseRequest = true;
+				}
+				else
+				{
+					UE_LOG(LogAudio, Log, TEXT("Invalid chunk request curIndex=%d numChunks=%d\n"), SourceChunk, SoundWave->RunningPlatformData->NumChunks);
+				}
+			}
+		}
+	}
+
 	for (auto Iter = WaveRequests.CreateIterator(); Iter; ++Iter)
 	{
 		USoundWave* Wave = Iter.Key();
@@ -694,6 +722,18 @@ void FAudioStreamingManager::RemoveStreamingSoundWave(USoundWave* SoundWave)
 	WaveRequests.Remove(SoundWave);
 }
 
+void FAudioStreamingManager::AddDecoder(ICompressedAudioInfo* InCompressedAudioInfo)
+{
+	FScopeLock Lock(&CriticalSection);
+	CompressedAudioInfos.Add(InCompressedAudioInfo);
+}
+
+void FAudioStreamingManager::RemoveDecoder(ICompressedAudioInfo* InCompressedAudioInfo)
+{
+	FScopeLock Lock(&CriticalSection);
+	CompressedAudioInfos.Remove(InCompressedAudioInfo);
+}
+
 bool FAudioStreamingManager::IsManagedStreamingSoundWave(const USoundWave* SoundWave) const
 {
 	FScopeLock Lock(&CriticalSection);
diff --git a/Engine/Source/Runtime/Engine/Private/VorbisAudioInfo.cpp b/Engine/Source/Runtime/Engine/Private/VorbisAudioInfo.cpp
--- a/Engine/Source/Runtime/Engine/Private/VorbisAudioInfo.cpp
+++ b/Engine/Source/Runtime/Engine/Private/VorbisAudioInfo.cpp
@@ -98,7 +98,6 @@ FVorbisAudioInfo::FVorbisAudioInfo()
 	, SrcBufferDataSize(0)
 	, BufferOffset(0)
 	, CurrentBufferChunkOffset(0)
-	, StreamingSoundWave(NULL)
 	, CurrentStreamingChunkData(nullptr)
 	, CurrentStreamingChunkIndex(INDEX_NONE)
 	, NextStreamingChunkIndex(0)
@@ -529,11 +528,11 @@ void FVorbisAudioInfo::EnableHalfRate( bool HalfRate )
 	ov_halfrate(&VFWrapper->vf, int32(HalfRate));
 }
 
-bool FVorbisAudioInfo::StreamCompressedInfo(USoundWave* Wave, struct FSoundQualityInfo* QualityInfo)
+bool FVorbisAudioInfo::StreamCompressedInfoInternal(USoundWave* Wave, struct FSoundQualityInfo* QualityInfo)
 {
 	if (!bDllLoaded)
 	{
-		UE_LOG(LogAudio, Error, TEXT("FVorbisAudioInfo::StreamCompressedInfo failed to parse header due to vorbis DLL not being loaded for sound '%s'."), *Wave->GetName());
+		UE_LOG(LogAudio, Error, TEXT("FVorbisAudioInfo::StreamCompressedInfoInternal failed to parse header due to vorbis DLL not being loaded for sound '%s'."), *Wave->GetName());
 		return false;
 	}
 
@@ -543,7 +542,7 @@ bool FVorbisAudioInfo::StreamCompressedInfo(USoundWave* Wave, struct FSoundQuali
 
 	if (!VFWrapper)
 	{
-		UE_LOG(LogAudio, Error, TEXT("FVorbisAudioInfo::StreamCompressedInfo failed due to no vorbis wrapper for sound '%s'."), *Wave->GetName());
+		UE_LOG(LogAudio, Error, TEXT("FVorbisAudioInfo::StreamCompressedInfoInternal failed due to no vorbis wrapper for sound '%s'."), *Wave->GetName());
 		return false;
 	}
 
@@ -562,7 +561,7 @@ bool FVorbisAudioInfo::StreamCompressedInfo(USoundWave* Wave, struct FSoundQuali
 	bHeaderParsed = GetCompressedInfoCommon(&Callbacks, QualityInfo);
 	if (!bHeaderParsed)
 	{
-		UE_LOG(LogAudio, Error, TEXT("FVorbisAudioInfo::StreamCompressedInfo failed to parse header for '%s'."), *Wave->GetName());
+		UE_LOG(LogAudio, Error, TEXT("FVorbisAudioInfo::StreamCompressedInfoInternal failed to parse header for '%s'."), *Wave->GetName());
 	}
 	
 
diff --git a/Engine/Source/Runtime/Engine/Public/ADPCMAudioInfo.h b/Engine/Source/Runtime/Engine/Public/ADPCMAudioInfo.h
--- a/Engine/Source/Runtime/Engine/Public/ADPCMAudioInfo.h
+++ b/Engine/Source/Runtime/Engine/Public/ADPCMAudioInfo.h
@@ -105,7 +105,7 @@ public:
 
 	// Additional overrides for streaming
 	virtual bool SupportsStreaming() const override {return true;}
-	virtual bool StreamCompressedInfo(USoundWave* Wave, struct FSoundQualityInfo* QualityInfo) override;
+	virtual bool StreamCompressedInfoInternal(USoundWave* Wave, struct FSoundQualityInfo* QualityInfo) override;
 	virtual bool StreamCompressedData(uint8* Destination, bool bLooping, uint32 BufferSize) override;
 	virtual int32 GetCurrentChunkIndex() const override
 	{
@@ -135,7 +135,6 @@ public:
 	uint32			TotalSamplesPerChannel;			// Number of samples per channel, used to detect when an audio waveform has ended
 	uint32			SamplesPerBlock;				// The number of samples per block
 	uint32			FirstChunkSampleDataOffset;		// The size of the header in the first chunk, used to skip over it when looping or starting the sample over
-	USoundWave*		StreamingSoundWave;				// The current sound wave being streamed, this is used to fetch new chunks
 	const uint8*	CurCompressedChunkData;			// A pointer to the current chunk of data
 
 	uint32			CurrentCompressedBlockIndex;		// For non disk streaming - the current compressed block in the compressed source data
diff --git a/Engine/Source/Runtime/Engine/Public/AudioDecompress.h b/Engine/Source/Runtime/Engine/Public/AudioDecompress.h
--- a/Engine/Source/Runtime/Engine/Public/AudioDecompress.h
+++ b/Engine/Source/Runtime/Engine/Public/AudioDecompress.h
@@ -99,7 +99,18 @@ public:
 	* @param	Wave			Wave that will be read from to retrieve necessary chunk
 	* @param	QualityInfo		Quality Info (to be filled out)
 	*/
-	virtual bool StreamCompressedInfo(USoundWave* Wave, struct FSoundQualityInfo* QualityInfo) {return false;}
+	bool StreamCompressedInfo(USoundWave* Wave, struct FSoundQualityInfo* QualityInfo)
+	{
+		StreamingSoundWave = Wave;
+
+		return StreamCompressedInfoInternal(Wave, QualityInfo);
+	}
+
+protected:
+	/** Internal override implemented by subclasses. */
+	virtual bool StreamCompressedInfoInternal(USoundWave* Wave, struct FSoundQualityInfo* QualityInfo) = 0;
+
+public:
 
 	/**
 	* Decompresses streamed data to raw PCM data.
@@ -121,6 +132,12 @@ public:
 	 * Gets the offset into the chunk that was last read to (for Streaming Manager priority)
 	 */
 	virtual int32 GetCurrentChunkOffset() const {return -1;}
+
+	/** Return the streaming sound wave used by this decoder. Returns nullptr if there is not a streaming sound wave. */
+	virtual USoundWave* GetStreamingSoundWave() { return StreamingSoundWave; }
+
+protected:
+	USoundWave* StreamingSoundWave;
 };
 
 /** Struct used to store the results of a decode operation. **/
@@ -161,7 +178,7 @@ public:
 	virtual bool UsesVorbisChannelOrdering() const override { return false; }
 	virtual int GetStreamBufferSize() const override { return  MONO_PCM_BUFFER_SIZE; }
 	virtual bool SupportsStreaming() const override { return true; }
-	virtual bool StreamCompressedInfo(USoundWave* Wave, FSoundQualityInfo* QualityInfo) override;
+	virtual bool StreamCompressedInfoInternal(USoundWave* Wave, FSoundQualityInfo* QualityInfo) override;
 	virtual bool StreamCompressedData(uint8* Destination, bool bLooping, uint32 BufferSize) override;
 	virtual int32 GetCurrentChunkIndex() const override { return CurrentChunkIndex; }
 	virtual int32 GetCurrentChunkOffset() const override { return SrcBufferOffset; }
@@ -254,8 +271,6 @@ protected:
 	uint32 LastPCMOffset;
 	/** If we're currently reading the final buffer. */
 	bool bStoringEndOfFile;
-    /** Ptr to the sound wave currently streaming. */
-	USoundWave* StreamingSoundWave;
 	/** The current chunk index in the streamed chunks. */
 	int32 CurrentChunkIndex;
 	/** Whether or not to print the chunk fail message. */
diff --git a/Engine/Source/Runtime/Engine/Public/AudioStreaming.h b/Engine/Source/Runtime/Engine/Public/AudioStreaming.h
--- a/Engine/Source/Runtime/Engine/Public/AudioStreaming.h
+++ b/Engine/Source/Runtime/Engine/Public/AudioStreaming.h
@@ -17,6 +17,7 @@ AudioStreaming.h: Definitions of classes used for audio streaming.
 
 class FSoundSource;
 class USoundWave;
+class ICompressedAudioInfo;
 struct FAudioStreamingManager;
 struct FWaveInstance;
 
@@ -261,6 +262,8 @@ struct FAudioStreamingManager : public IAudioStreamingManager
 	// IAudioStreamingManager interface
 	virtual void AddStreamingSoundWave(USoundWave* SoundWave) override;
 	virtual void RemoveStreamingSoundWave(USoundWave* SoundWave) override;
+	virtual void AddDecoder(ICompressedAudioInfo* CompressedAudioInfo) override;
+	virtual void RemoveDecoder(ICompressedAudioInfo* CompressedAudioInfo) override;
 	virtual bool IsManagedStreamingSoundWave(const USoundWave* SoundWave) const override;
 	virtual bool IsStreamingInProgress(const USoundWave* SoundWave) override;
 	virtual bool CanCreateSoundSource(const FWaveInstance* WaveInstance) const override;
@@ -301,4 +304,7 @@ protected:
 
 	/** Critical section to protect usage of shared gamethread/audiothread members */
 	mutable FCriticalSection CriticalSection;
+
+	/** Compressed audio info objects which are used to avoid deleting chunks with in-fligth decodes. */
+	TArray<ICompressedAudioInfo*> CompressedAudioInfos;
 };
diff --git a/Engine/Source/Runtime/Engine/Public/ContentStreaming.h b/Engine/Source/Runtime/Engine/Public/ContentStreaming.h
--- a/Engine/Source/Runtime/Engine/Public/ContentStreaming.h
+++ b/Engine/Source/Runtime/Engine/Public/ContentStreaming.h
@@ -13,6 +13,7 @@ class AActor;
 class FSoundSource;
 class UPrimitiveComponent;
 class USoundWave;
+class ICompressedAudioInfo;
 class UTexture2D;
 struct FRenderAssetStreamingManager;
 struct FWaveInstance;
@@ -409,6 +410,12 @@ struct IAudioStreamingManager : public IStreamingManager
 	/** Removes a Sound Wave from the streaming manager. */
 	virtual void RemoveStreamingSoundWave(USoundWave* SoundWave) = 0;
 
+	/** Adds the decoder to the streaming manager to prevent stream chunks from getting reaped from underneath it */
+	virtual void AddDecoder(ICompressedAudioInfo* CompressedAudioInfo) = 0;
+
+	/** Removes the decoder from the streaming manager. */
+	virtual void RemoveDecoder(ICompressedAudioInfo* CompressedAudioInfo) = 0;
+
 	/** Returns true if this is a Sound Wave that is managed by the streaming manager. */
 	virtual bool IsManagedStreamingSoundWave(const USoundWave* SoundWave) const = 0;
 
diff --git a/Engine/Source/Runtime/Engine/Public/VorbisAudioInfo.h b/Engine/Source/Runtime/Engine/Public/VorbisAudioInfo.h
--- a/Engine/Source/Runtime/Engine/Public/VorbisAudioInfo.h
+++ b/Engine/Source/Runtime/Engine/Public/VorbisAudioInfo.h
@@ -91,12 +91,14 @@ public:
 
 	// Additional overrides for streaming
 	virtual bool SupportsStreaming() const override {return true;}
-	virtual bool StreamCompressedInfo(USoundWave* Wave, struct FSoundQualityInfo* QualityInfo) override;
 	virtual bool StreamCompressedData(uint8* InDestination, bool bLooping, uint32 BufferSize) override;
 	virtual int32 GetCurrentChunkIndex() const override {return CurrentStreamingChunkIndex;}
 	virtual int32 GetCurrentChunkOffset() const override {return BufferOffset % CurrentStreamingChunksSize;}
 	// End of ICompressedAudioInfo Interface
 
+protected:
+	virtual bool StreamCompressedInfoInternal(USoundWave* Wave, struct FSoundQualityInfo* QualityInfo) override;
+
 private:
 
 	struct FVorbisFileWrapper* VFWrapper;
@@ -108,7 +110,6 @@ private:
 	/** Critical section used to prevent multiple threads accessing same ogg-vorbis file handles at the same time */
 	FCriticalSection VorbisCriticalSection;
 
-	USoundWave* StreamingSoundWave;				// The current sound wave being streamed, this is used to fetch new chunks
 	uint8 const* CurrentStreamingChunkData;
 	int32 CurrentStreamingChunkIndex;
 	int32 NextStreamingChunkIndex;
