Fix for hotfix dynamic volume not working with new audio mixer

diff --git a/Engine/Source/Runtime/AudioMixer/Private/AudioMixerSource.cpp b/Engine/Source/Runtime/AudioMixer/Private/AudioMixerSource.cpp
--- a/Engine/Source/Runtime/AudioMixer/Private/AudioMixerSource.cpp
+++ b/Engine/Source/Runtime/AudioMixer/Private/AudioMixerSource.cpp
@@ -759,6 +759,7 @@ namespace Audio
 			CurrentVolume = WaveInstance->GetVolume();
 			CurrentVolume *= WaveInstance->GetVolumeApp();
 			CurrentVolume *= AudioDevice->GetPlatformAudioHeadroom();
+			CurrentVolume *= WaveInstance->GetDynamicVolume();
 			CurrentVolume = FMath::Clamp<float>(GetDebugVolume(CurrentVolume), 0.0f, MAX_VOLUME);
 		}
 
