Fix for reverb on console/surround sound systems.

- Fix is to allow quad reverb channel flipping to be enabled via cvar.

diff --git a/Engine/Source/Runtime/AudioMixer/Private/Effects/AudioMixerSubmixEffectReverb.cpp b/Engine/Source/Runtime/AudioMixer/Private/Effects/AudioMixerSubmixEffectReverb.cpp
--- a/Engine/Source/Runtime/AudioMixer/Private/Effects/AudioMixerSubmixEffectReverb.cpp
+++ b/Engine/Source/Runtime/AudioMixer/Private/Effects/AudioMixerSubmixEffectReverb.cpp
@@ -14,6 +14,14 @@ FAutoConsoleVariableRef CVarDisableSubmixReverb(
 	TEXT("0: Not Disabled, 1: Disabled"),
 	ECVF_Default);
 
+static int32 EnableReverbStereoFlipForQuadCvar = 0;
+FAutoConsoleVariableRef CVarReverbStereoFlipForQuad(
+	TEXT("au.EnableReverbStereoFlipForQuad"),
+	EnableReverbStereoFlipForQuadCvar,
+	TEXT("Enables doing a stereo flip for quad reverb when in surround.\n")
+	TEXT("0: Not Enabled, 1: Enabled"),
+	ECVF_Default);
+
 class UReverbEffect;
 
 FSubmixEffectReverb::FSubmixEffectReverb()
@@ -115,9 +123,17 @@ void FSubmixEffectReverb::OnProcessAudio(const FSoundEffectSubmixInputData& InDa
 			PlateReverb.ProcessAudioFrame(&AudioData[InSampleIndex], InData.NumChannels, &OutAudioData[OutSampleIndex], InData.NumChannels);
 			// Now do a cross-over to the back-left/back-right speakers from the front-left and front-right
 			
-			// Using standard speaker map order map the right output to the BackLeft channel
-			OutAudioData[OutSampleIndex + EAudioMixerChannel::BackRight] = OutAudioData[OutSampleIndex + EAudioMixerChannel::FrontLeft];
-			OutAudioData[OutSampleIndex + EAudioMixerChannel::BackLeft] = OutAudioData[OutSampleIndex + EAudioMixerChannel::FrontRight];
+			if (EnableReverbStereoFlipForQuadCvar == 1)
+			{
+				// Using standard speaker map order map the right output to the BackLeft channel
+				OutAudioData[OutSampleIndex + EAudioMixerChannel::BackRight] = OutAudioData[OutSampleIndex + EAudioMixerChannel::FrontLeft];
+				OutAudioData[OutSampleIndex + EAudioMixerChannel::BackLeft] = OutAudioData[OutSampleIndex + EAudioMixerChannel::FrontRight];
+			}
+			else
+			{
+				OutAudioData[OutSampleIndex + EAudioMixerChannel::BackLeft] = OutAudioData[OutSampleIndex + EAudioMixerChannel::FrontLeft];
+				OutAudioData[OutSampleIndex + EAudioMixerChannel::BackRight] = OutAudioData[OutSampleIndex + EAudioMixerChannel::FrontRight];
+			}
 		}
 	}
 }
