Removing check which was invalid.

diff --git a/Engine/Source/Runtime/AudioMixer/Private/AudioMixerSourceBuffer.cpp b/Engine/Source/Runtime/AudioMixer/Private/AudioMixerSourceBuffer.cpp
--- a/Engine/Source/Runtime/AudioMixer/Private/AudioMixerSourceBuffer.cpp
+++ b/Engine/Source/Runtime/AudioMixer/Private/AudioMixerSourceBuffer.cpp
@@ -196,8 +196,6 @@ namespace Audio
 	{
 		CurrentBuffer = 0;
 
-		check(RawPCMDataBuffer.Data != nullptr);
-
 		RawPCMDataBuffer.NumSamples = RawPCMDataBuffer.DataSize / sizeof(int16);
 		RawPCMDataBuffer.CurrentSample = 0;
 
