[NO NON-PII COMMIT MESSAGE]

diff --git a/Engine/Source/Runtime/AudioMixer/Private/AudioMixerBus.cpp b/Engine/Source/Runtime/AudioMixer/Private/AudioMixerBus.cpp
--- a/Engine/Source/Runtime/AudioMixer/Private/AudioMixerBus.cpp
+++ b/Engine/Source/Runtime/AudioMixer/Private/AudioMixerBus.cpp
@@ -20,14 +20,30 @@ namespace Audio
 		}
 	}
 
+	void FMixerBus::SetNumOutputChannels(int32 InNumOutputChannels)
+	{
+		NumChannels = InNumOutputChannels;
+		const int32 NumSamples = NumChannels * NumFrames;
+		for (int32 i = 0; i < 2; ++i)
+		{
+			MixedSourceData[i].Reset();
+			MixedSourceData[i].AddZeroed(NumSamples);
+		}
+	}
+
 	void FMixerBus::Update()
 	{
 		CurrentBufferIndex = !CurrentBufferIndex;
 	}
 
-	void FMixerBus::AddInstanceId(const int32 InSourceId)
+	void FMixerBus::AddInstanceId(const int32 InSourceId, int32 InNumOutputChannels)
 	{
 		InstanceIds.Add(InSourceId);
+
+		if (NumChannels != InNumOutputChannels)
+		{
+			SetNumOutputChannels(InNumOutputChannels);
+		}
 	}
 
 	bool FMixerBus::RemoveInstanceId(const int32 InSourceId)
diff --git a/Engine/Source/Runtime/AudioMixer/Private/AudioMixerSourceManager.cpp b/Engine/Source/Runtime/AudioMixer/Private/AudioMixerSourceManager.cpp
--- a/Engine/Source/Runtime/AudioMixer/Private/AudioMixerSourceManager.cpp
+++ b/Engine/Source/Runtime/AudioMixer/Private/AudioMixerSourceManager.cpp
@@ -652,14 +652,14 @@ namespace Audio
 					if (Bus)
 					{
 						// If this bus is already registered, add this as a source id
-						Bus->AddInstanceId(SourceId);
+						Bus->AddInstanceId(SourceId, InitParams.NumInputChannels);
 					}
 					else
 					{
 						// If the bus is not registered, make a new entry
 						FMixerBus NewBusData(this, InitParams.NumInputChannels, NumOutputFrames);
 
-						NewBusData.AddInstanceId(SourceId);
+						NewBusData.AddInstanceId(SourceId, InitParams.NumInputChannels);
 
 						Buses.Add(InitParams.BusId, NewBusData);
 					}
diff --git a/Engine/Source/Runtime/AudioMixer/Public/AudioMixerBus.h b/Engine/Source/Runtime/AudioMixer/Public/AudioMixerBus.h
--- a/Engine/Source/Runtime/AudioMixer/Public/AudioMixerBus.h
+++ b/Engine/Source/Runtime/AudioMixer/Public/AudioMixerBus.h
@@ -26,7 +26,7 @@ namespace Audio
 		void Update();
 
 		// Adds a source id for instances of this bus
-		void AddInstanceId(const int32 InSourceInstanceId);
+		void AddInstanceId(const int32 InSourceInstanceId, int32 InNumOutputChannels);
 
 		// Removes the source id from this bus. Returns true if there are no more instances or sends.
 		bool RemoveInstanceId(const int32 InSourceId);
@@ -48,6 +48,9 @@ namespace Audio
 
 	private:
 
+		// If this bus was constructed before
+		void SetNumOutputChannels(int32 InNumOutputChannels);
+
 		// Array of instance ids. These are sources which are instances of this.
 		// It's possible for this data to have bus sends but no instance ids.
 		// This means a source would send its audio to the bus if the bus had an instance.
