Crash fixes -- couple of crashes are due to the same issue.

- USoundWave ptr is getting GC'd from underneath a FMixerSourceBuffer. Fix is to remove the USoundWave* ptr in non-procedural sound wave cases and copy needed data over. USoundWave can get GC'd safely while sound is stopping just fine.

- Crashes also due to not handling a case where a decoder state is not properly added/removed from streaming manager so streaming manager freeing a loaded chunk can cause a crash

- A streaming manager crash can happen when a streaming sound source is removed from the streaming manager and has async IOs pending. The Async tasks are blocked till completed, but the data which the tasks append to is later queried after the streaming sound source is deleted.

diff --git a/Engine/Source/Runtime/AudioMixer/Private/AudioMixerBuffer.cpp b/Engine/Source/Runtime/AudioMixer/Private/AudioMixerBuffer.cpp
--- a/Engine/Source/Runtime/AudioMixer/Private/AudioMixerBuffer.cpp
+++ b/Engine/Source/Runtime/AudioMixer/Private/AudioMixerBuffer.cpp
@@ -4,6 +4,8 @@
 #include "AudioMixerDevice.h"
 #include "Interfaces/IAudioFormat.h"
 #include "AudioMixerSourceDecode.h"
+#include "AudioStreaming.h"
+
 namespace Audio
 {
 
@@ -42,6 +44,11 @@ namespace Audio
 
 		if (DecompressionState)
 		{
+			if (BufferType == EBufferType::Streaming)
+			{
+				IStreamingManager::Get().GetAudioStreamingManager().RemoveDecoder(DecompressionState);
+			}
+
 			delete DecompressionState;
 			DecompressionState = nullptr;
 		}
diff --git a/Engine/Source/Runtime/AudioMixer/Private/AudioMixerSource.cpp b/Engine/Source/Runtime/AudioMixer/Private/AudioMixerSource.cpp
--- a/Engine/Source/Runtime/AudioMixer/Private/AudioMixerSource.cpp
+++ b/Engine/Source/Runtime/AudioMixer/Private/AudioMixerSource.cpp
@@ -58,13 +58,20 @@ namespace Audio
 
 		FSoundSource::InitCommon();
 
-		// Get the number of frames before creating the buffer
-		int32 NumFrames = INDEX_NONE;
+		check(WaveInstance->WaveData);
 
-		AUDIO_MIXER_CHECK(WaveInstance->WaveData);
+		// Prevent double-triggering procedural soundwaves
+		if (WaveInstance->WaveData->bProcedural && WaveInstance->WaveData->IsGenerating())
+		{
+			UE_LOG(LogAudioMixer, Warning, TEXT("Procedural sound wave is reinitializing even though it is currently actively generating audio. Please stop sound before trying to play it again."));
+			return false;
+		}
 
+		// Get the number of frames before creating the buffer
+		int32 NumFrames = INDEX_NONE;
 		if (WaveInstance->WaveData->DecompressionType != DTYPE_Procedural)
 		{
+			//NumFrames = MixerBuffer->GetNumFrames();
 			check(!InWaveInstance->WaveData->RawPCMData || InWaveInstance->WaveData->RawPCMDataSize);
 			const int32 NumBytes = WaveInstance->WaveData->RawPCMDataSize;
 			NumFrames = NumBytes / (WaveInstance->WaveData->NumChannels * sizeof(int16));
@@ -282,7 +289,34 @@ namespace Audio
 			// Update the buffer sample rate to the wave instance sample rate in case it was serialized incorrectly
 			MixerBuffer->InitSampleRate(WaveInstance->WaveData->GetSampleRateForCurrentPlatform());
 
-			// Pass the decompression state off to the mixer source buffer
+			// Retrieve the raw pcm buffer data and the precached buffers before initializing so we can avoid having USoundWave ptrs in audio renderer thread
+			EBufferType::Type BufferType = MixerBuffer->GetType();
+			if (BufferType == EBufferType::PCM || BufferType == EBufferType::PCMPreview)
+			{
+				FRawPCMDataBuffer RawPCMDataBuffer;
+				MixerBuffer->GetPCMData(&RawPCMDataBuffer.Data, &RawPCMDataBuffer.DataSize);
+				MixerSourceBuffer->SetPCMData(RawPCMDataBuffer);
+			}
+#if PLATFORM_NUM_AUDIODECOMPRESSION_PRECACHE_BUFFERS > 0
+			else if (BufferType == EBufferType::PCMRealTime || EBufferType::Streaming)
+			{
+				USoundWave* WaveData = WaveInstance->WaveData;
+				if (WaveData->CachedRealtimeFirstBuffer)
+				{
+					const uint32 NumPrecacheSamples = (uint32)(WaveData->NumPrecacheFrames * WaveData->NumChannels);
+					const uint32 BufferSize = NumPrecacheSamples * sizeof(int16) * PLATFORM_NUM_AUDIODECOMPRESSION_PRECACHE_BUFFERS;
+
+					TArray<uint8> PrecacheBufferCopy;
+					PrecacheBufferCopy.AddUninitialized(BufferSize);
+
+					FMemory::Memcpy(PrecacheBufferCopy.GetData(), WaveData->CachedRealtimeFirstBuffer, BufferSize);
+
+					MixerSourceBuffer->SetCachedRealtimeFirstBuffers(MoveTemp(PrecacheBufferCopy));
+				}
+			}
+#endif
+
+			// Pass the decompression state off to the mixer source buffer if it hasn't already done so
 			ICompressedAudioInfo* Decoder = MixerBuffer->GetDecompressionState(true);
 			MixerSourceBuffer->SetDecoder(Decoder);
 
@@ -437,7 +471,13 @@ namespace Audio
 							}
 
 							check(MixerSourceBuffer.IsValid());
+
 							ICompressedAudioInfo* Decoder = MixerBuffer->GetDecompressionState(false);
+							if (BufferType == EBufferType::Streaming)
+							{
+								IStreamingManager::Get().GetAudioStreamingManager().AddDecoder(Decoder);
+							}
+
 							MixerSourceBuffer->ReadMoreRealtimeData(Decoder, 0, EBufferReadMode::Asynchronous);
 
 							// not ready
diff --git a/Engine/Source/Runtime/AudioMixer/Private/AudioMixerSourceBuffer.cpp b/Engine/Source/Runtime/AudioMixer/Private/AudioMixerSourceBuffer.cpp
--- a/Engine/Source/Runtime/AudioMixer/Private/AudioMixerSourceBuffer.cpp
+++ b/Engine/Source/Runtime/AudioMixer/Private/AudioMixerSourceBuffer.cpp
@@ -63,11 +63,14 @@ namespace Audio
 		, LoopingMode(ELoopingMode::LOOP_Never)
 		, NumChannels(0)
 		, BufferType(Audio::EBufferType::Invalid)
+		, NumPrecacheFrames(0)
 		, bInitialized(false)
 		, bBufferFinished(false)
 		, bPlayedCachedBuffer(false)
 		, bIsSeeking(false)
 		, bLoopCallback(false)
+		, bProcedural(false)
+		, bIsBus(false)
 	{
 	}
 
@@ -86,7 +89,10 @@ namespace Audio
 		// Clean up decompression state after things have been finished using it
 		if (DecompressionState)
 		{
-			IStreamingManager::Get().GetAudioStreamingManager().RemoveDecoder(DecompressionState);
+			if (BufferType == EBufferType::Streaming)
+			{
+				IStreamingManager::Get().GetAudioStreamingManager().RemoveDecoder(DecompressionState);
+			}
 
 			delete DecompressionState;
 			DecompressionState = nullptr;
@@ -96,23 +102,24 @@ namespace Audio
 	bool FMixerSourceBuffer::PreInit(FMixerBuffer* InBuffer, USoundWave* InWave, ELoopingMode InLoopingMode, bool bInIsSeeking)
 	{
 		LLM_SCOPE(ELLMTag::AudioMixer);
+		if (!InWave)
+		{
+			return false;
+		}
 
 		NumChannels = InBuffer->NumChannels;
 		BufferType = InBuffer->GetType();
-
-		// If this is a PCM buffer type, we need to take ownership of the raw pcm data
-		if (BufferType == EBufferType::PCM || BufferType == EBufferType::PCMPreview)
+		bIsBus = InWave->bIsBus;
+		bProcedural = InWave->bProcedural;
+		NumPrecacheFrames = InWave->NumPrecacheFrames;
+		
+		check(SoundWave == nullptr);
+		// Only need to have a handle to a USoundWave for procedural sound waves
+		if (bProcedural)
 		{
-			check(RawPCMDataBuffer.Data == nullptr);
-			check(RawPCMDataBuffer.DataSize == 0);
-
-			// GetPCMData will retrieve ownership of the raw PCM data
-			InBuffer->GetPCMData(&RawPCMDataBuffer.Data, &RawPCMDataBuffer.DataSize);
+			SoundWave = InWave;
 		}
 
-		// May or may not be nullptr
-		SoundWave = InWave;
-
 		LoopingMode = InLoopingMode;
 		bIsSeeking = bInIsSeeking;
 		bLoopCallback = false;
@@ -134,22 +141,29 @@ namespace Audio
 
 	void FMixerSourceBuffer::SetDecoder(ICompressedAudioInfo* InCompressedAudioInfo)
 	{
-		DecompressionState = InCompressedAudioInfo;
-		if (BufferType == EBufferType::Streaming)
+		if (DecompressionState == nullptr)
 		{
-			IStreamingManager::Get().GetAudioStreamingManager().AddDecoder(DecompressionState);
+			DecompressionState = InCompressedAudioInfo;
+			if (BufferType == EBufferType::Streaming)
+			{
+				IStreamingManager::Get().GetAudioStreamingManager().AddDecoder(DecompressionState);
+			}
 		}
 	}
 
-	bool FMixerSourceBuffer::Init()
+	void FMixerSourceBuffer::SetPCMData(const FRawPCMDataBuffer& InPCMDataBuffer)
 	{
-		check(SoundWave);
-		if (SoundWave->bProcedural && SoundWave->IsGenerating())
-		{
-			UE_LOG(LogAudioMixer, Warning, TEXT("Procedural sound wave is reinitializing even though it is currently actively generating audio. Please stop sound before trying to play it again."));
-			return false;
-		}
+		check(BufferType == EBufferType::PCM || EBufferType::PCMPreview);
+		RawPCMDataBuffer = InPCMDataBuffer;
+	}
+
+	void FMixerSourceBuffer::SetCachedRealtimeFirstBuffers(TArray<uint8>&& InPrecachedBuffers)
+	{
+		CachedRealtimeFirstBuffer = MoveTemp(InPrecachedBuffers);
+	}
 
+	bool FMixerSourceBuffer::Init()
+	{
 		// We have successfully initialized which means our SoundWave has been flagged as bIsActive
 		// GC can run between PreInit and Init so when cleaning up FMixerSourceBuffer, we don't want to touch SoundWave unless bInitailized is true.
 		// SoundWave->bIsSoundActive will prevent GC until it is released in audio render thread
@@ -176,7 +190,7 @@ namespace Audio
 
 	void FMixerSourceBuffer::OnBufferEnd()
 	{
-		if ((NumBuffersQeueued == 0 && bBufferFinished) || !SoundWave)
+		if ((NumBuffersQeueued == 0 && bBufferFinished) && (bProcedural && !SoundWave))
 		{
 			return;
 		}
@@ -229,11 +243,10 @@ namespace Audio
 		CurrentBuffer = 0;
 
 		bPlayedCachedBuffer = false;
-		if (!bIsSeeking && SoundWave && SoundWave->CachedRealtimeFirstBuffer)
+		if (!bIsSeeking && CachedRealtimeFirstBuffer.Num() > 0)
 		{
 			bPlayedCachedBuffer = true;
 
-			const int32 NumPrecacheFrames = SoundWave->NumPrecacheFrames;
 			const uint32 NumSamples = NumPrecacheFrames * NumChannels;
 			const uint32 BufferSize = NumSamples * sizeof(int16);
 
@@ -247,8 +260,8 @@ namespace Audio
 					SourceVoiceBuffers[i]->AudioData.AddUninitialized(NumSamples);
 				}
 
-				int16* CachedBufferPtr0 = (int16*)SoundWave->CachedRealtimeFirstBuffer;
-				int16* CachedBufferPtr1 = (int16*)(SoundWave->CachedRealtimeFirstBuffer + BufferSize);
+				int16* CachedBufferPtr0 = (int16*)CachedRealtimeFirstBuffer.GetData();
+				int16* CachedBufferPtr1 = (int16*)(CachedRealtimeFirstBuffer.GetData() + BufferSize);
 				float* AudioData0 = SourceVoiceBuffers[0]->AudioData.GetData();
 				float* AudioData1 = SourceVoiceBuffers[1]->AudioData.GetData();
 				for (uint32 Sample = 0; Sample < NumSamples; ++Sample)
@@ -272,7 +285,7 @@ namespace Audio
 				SourceVoiceBuffers[0]->AudioData.Reset();
 				SourceVoiceBuffers[0]->AudioData.AddUninitialized(NumSamples);
 
-				int16* CachedBufferPtr0 = (int16*)SoundWave->CachedRealtimeFirstBuffer;
+				int16* CachedBufferPtr0 = (int16*)CachedRealtimeFirstBuffer.GetData();
 
 				float* AudioData0 = SourceVoiceBuffers[0]->AudioData.GetData();
 				for (uint32 Sample = 0; Sample < NumSamples; ++Sample)
@@ -287,7 +300,7 @@ namespace Audio
 			}
 #endif
 		}
-		else if (SoundWave && !SoundWave->bIsBus)
+		else if (!bIsBus)
 		{
 			// We should have already kicked off and finished a task. 
 			check(AsyncRealtimeAudioTask != nullptr);
@@ -302,8 +315,9 @@ namespace Audio
 		SourceVoiceBuffers[BufferIndex]->AudioData.Reset();
 		SourceVoiceBuffers[BufferIndex]->AudioData.AddUninitialized(MaxSamples);
 
-		if (SoundWave && SoundWave->bProcedural)
+		if (bProcedural)
 		{
+			check(SoundWave && SoundWave->bProcedural);
 			FProceduralAudioTaskData NewTaskData;
 			NewTaskData.ProceduralSoundWave = SoundWave;
 			NewTaskData.AudioData = SourceVoiceBuffers[BufferIndex]->AudioData.GetData();
@@ -333,15 +347,7 @@ namespace Audio
 		NewTaskData.bLoopingMode = LoopingMode != LOOP_Never;
 		NewTaskData.bSkipFirstBuffer = (BufferReadMode == EBufferReadMode::AsynchronousSkipFirstFrame);
 		NewTaskData.NumFramesToDecode = MONO_PCM_BUFFER_SAMPLES;
-
-		if (SoundWave)
-		{
-			NewTaskData.NumPrecacheFrames = SoundWave->NumPrecacheFrames;
-		}
-		else
-		{
-			NewTaskData.NumPrecacheFrames = MONO_PCM_BUFFER_SAMPLES;
-		}
+		NewTaskData.NumPrecacheFrames = NumPrecacheFrames;
 
 		check(!AsyncRealtimeAudioTask);
 		AsyncRealtimeAudioTask = CreateAudioTask(NewTaskData);
@@ -476,7 +482,12 @@ namespace Audio
 
 	bool FMixerSourceBuffer::IsBeginDestroy()
 	{
-		return SoundWave && SoundWave->bIsBeginDestroy;
+		if (bProcedural)
+		{
+			check(SoundWave && SoundWave->bProcedural);
+			return SoundWave->bIsBeginDestroy;
+		}
+		return false;
 	}
 
 	void FMixerSourceBuffer::ClearSoundWave()
@@ -487,14 +498,11 @@ namespace Audio
 
 	void FMixerSourceBuffer::OnBeginGenerate()
 	{
-		if (SoundWave)
+		if (bProcedural)
 		{
-			if (SoundWave->bProcedural)
-			{
-				SoundWave->SetGenerating(true);
-				SoundWave->OnBeginGenerate();
-			}
-
+			check(SoundWave && SoundWave->bProcedural);
+			SoundWave->SetGenerating(true);
+			SoundWave->OnBeginGenerate();
 		}
 	}
 
@@ -504,14 +512,11 @@ namespace Audio
 		EnsureAsyncTaskFinishes();
 
 		// Only need to call OnEndGenerate and access SoundWave here if we successfully initialized
-		if (bInitialized && SoundWave)
+		if (bInitialized && bProcedural)
 		{
-			if (SoundWave->bProcedural)
-			{
-				SoundWave->OnEndGenerate();
-				SoundWave->SetGenerating(false);
-			}
-
+			check(SoundWave && SoundWave->bProcedural);
+			SoundWave->OnEndGenerate();
+			SoundWave->SetGenerating(false);
 			SoundWave = nullptr;
 		}
 	}
diff --git a/Engine/Source/Runtime/AudioMixer/Private/AudioMixerSourceBuffer.h b/Engine/Source/Runtime/AudioMixer/Private/AudioMixerSourceBuffer.h
--- a/Engine/Source/Runtime/AudioMixer/Private/AudioMixerSourceBuffer.h
+++ b/Engine/Source/Runtime/AudioMixer/Private/AudioMixerSourceBuffer.h
@@ -55,6 +55,12 @@ namespace Audio
 		// Sets the decoder to use for realtime async decoding
 		void SetDecoder(ICompressedAudioInfo* InCompressedAudioInfo);
 
+		// Sets the raw PCM data buffer to use for the source buffer
+		void SetPCMData(const FRawPCMDataBuffer& InPCMDataBuffer);
+
+		// Sets the precached buffers
+		void SetCachedRealtimeFirstBuffers(TArray<uint8>&& InPrecachedBuffer);
+
 		// Called by source manager when needing more buffers
 		void OnBufferEnd();
 
@@ -107,17 +113,23 @@ namespace Audio
 		TArray<TSharedPtr<FMixerSourceVoiceBuffer>> SourceVoiceBuffers;
 		TQueue<TSharedPtr<FMixerSourceVoiceBuffer>> BufferQueue;
 		int32 CurrentBuffer;
+		// SoundWaves are only set for procedural sound waves
 		USoundWave* SoundWave;
 		IAudioTask* AsyncRealtimeAudioTask;
 		ICompressedAudioInfo* DecompressionState;
 		ELoopingMode LoopingMode;
 		int32 NumChannels;
 		Audio::EBufferType::Type BufferType;
+		int32 NumPrecacheFrames;
+		TArray<uint8> CachedRealtimeFirstBuffer;
+
 		uint32 bInitialized : 1;
 		uint32 bBufferFinished : 1;
 		uint32 bPlayedCachedBuffer : 1;
 		uint32 bIsSeeking : 1;
 		uint32 bLoopCallback : 1;
+		uint32 bProcedural : 1;
+		uint32 bIsBus : 1;
 
 	};
 }
diff --git a/Engine/Source/Runtime/Engine/Private/AudioStreaming.cpp b/Engine/Source/Runtime/Engine/Private/AudioStreaming.cpp
--- a/Engine/Source/Runtime/Engine/Private/AudioStreaming.cpp
+++ b/Engine/Source/Runtime/Engine/Private/AudioStreaming.cpp
@@ -89,6 +89,11 @@ FStreamingWaveData::FStreamingWaveData()
 }
 
 FStreamingWaveData::~FStreamingWaveData()
+{
+	check(IORequestHandle == nullptr);
+}
+
+void FStreamingWaveData::FreeResources()
 {
 	// Make sure there are no pending requests in flight.
 	for (int32 Pass = 0; Pass < 3; Pass++)
@@ -100,6 +105,7 @@ FStreamingWaveData::~FStreamingWaveData()
 		}
 		check(Pass < 2); // we should be done after two passes. Pass 0 will start anything we need and pass 1 will complete those requests
 	}
+
 	for (FLoadedAudioChunk& LoadedChunk : LoadedChunks)
 	{
 		FreeLoadedChunk(LoadedChunk);
@@ -717,6 +723,25 @@ void FAudioStreamingManager::RemoveStreamingSoundWave(USoundWave* SoundWave)
 	if (WaveData)
 	{
 		StreamingSoundWaves.Remove(SoundWave);
+
+		// Free the resources of the streaming wave data. This blocks pending IO requests
+		WaveData->FreeResources();
+
+		{
+			// Then we need to remove any results from those pending requests before we delete so that we don't process them
+			FScopeLock StreamChunkResults(&ChunkResultCriticalSection);
+			for (int32 i = AsyncAudioStreamChunkResults.Num() - 1 ; i >= 0; --i)
+			{
+				FASyncAudioChunkLoadResult* LoadResult = AsyncAudioStreamChunkResults[i];
+				FStreamingWaveData* StreamingWaveData = LoadResult->StreamingWaveData;
+				if (StreamingWaveData == WaveData)
+				{
+					delete LoadResult;
+					
+					AsyncAudioStreamChunkResults.RemoveAtSwap(i, 1, false);
+				}
+			}
+		}
 		delete WaveData;
 	}
 	WaveRequests.Remove(SoundWave);
@@ -725,7 +750,7 @@ void FAudioStreamingManager::RemoveStreamingSoundWave(USoundWave* SoundWave)
 void FAudioStreamingManager::AddDecoder(ICompressedAudioInfo* InCompressedAudioInfo)
 {
 	FScopeLock Lock(&CriticalSection);
-	CompressedAudioInfos.Add(InCompressedAudioInfo);
+	CompressedAudioInfos.AddUnique(InCompressedAudioInfo);
 }
 
 void FAudioStreamingManager::RemoveDecoder(ICompressedAudioInfo* InCompressedAudioInfo)
diff --git a/Engine/Source/Runtime/Engine/Public/AudioStreaming.h b/Engine/Source/Runtime/Engine/Public/AudioStreaming.h
--- a/Engine/Source/Runtime/Engine/Public/AudioStreaming.h
+++ b/Engine/Source/Runtime/Engine/Public/AudioStreaming.h
@@ -128,6 +128,9 @@ struct FStreamingWaveData final
 	FStreamingWaveData();
 	~FStreamingWaveData();
 
+	// Frees streaming wave data resources, blocks pending async IO requests
+	void FreeResources();
+
 	/**
 	 * Sets up the streaming wave data and loads the first chunk of audio for instant play
 	 *
