Further fix for reverb -- reduce volume of quad verb by 0.5 to maintain power.
Add ability to disable quad reverb.

diff --git a/Engine/Source/Runtime/AudioMixer/Private/Effects/AudioMixerSubmixEffectReverb.cpp b/Engine/Source/Runtime/AudioMixer/Private/Effects/AudioMixerSubmixEffectReverb.cpp
--- a/Engine/Source/Runtime/AudioMixer/Private/Effects/AudioMixerSubmixEffectReverb.cpp
+++ b/Engine/Source/Runtime/AudioMixer/Private/Effects/AudioMixerSubmixEffectReverb.cpp
@@ -22,6 +22,15 @@ FAutoConsoleVariableRef CVarReverbStereoFlipForQuad(
 	TEXT("0: Not Enabled, 1: Enabled"),
 	ECVF_Default);
 
+
+static int32 DisableQuadReverbCVar = 0;
+FAutoConsoleVariableRef CVarDisableQuadReverbCVar(
+	TEXT("au.DisableQuadReverb"),
+	DisableQuadReverbCVar,
+	TEXT("Disables quad reverb in surround.\n")
+	TEXT("0: Not Disabled, 1: Disabled"),
+	ECVF_Default);
+
 class UReverbEffect;
 
 FSubmixEffectReverb::FSubmixEffectReverb()
@@ -117,20 +126,40 @@ void FSubmixEffectReverb::OnProcessAudio(const FSoundEffectSubmixInputData& InDa
 	// 5.1 or higher surround sound. Map stereo output to quad output
 	else if(OutData.NumChannels > 5)
 	{
-		for (int32 InSampleIndex = 0, OutSampleIndex = 0; InSampleIndex < InData.AudioBuffer->Num(); InSampleIndex += InData.NumChannels, OutSampleIndex += OutData.NumChannels)
+		if (DisableQuadReverbCVar == 1)
+		{
+			for (int32 InSampleIndex = 0, OutSampleIndex = 0; InSampleIndex < InData.AudioBuffer->Num(); InSampleIndex += InData.NumChannels, OutSampleIndex += OutData.NumChannels)
+			{
+				PlateReverb.ProcessAudioFrame(&AudioData[InSampleIndex], InData.NumChannels, &OutAudioData[OutSampleIndex], InData.NumChannels);
+			}
+		}
+		else if (EnableReverbStereoFlipForQuadCvar == 1)
 		{
-			// Processed downmixed audio frame
-			PlateReverb.ProcessAudioFrame(&AudioData[InSampleIndex], InData.NumChannels, &OutAudioData[OutSampleIndex], InData.NumChannels);
-			// Now do a cross-over to the back-left/back-right speakers from the front-left and front-right
-			
-			if (EnableReverbStereoFlipForQuadCvar == 1)
+			for (int32 InSampleIndex = 0, OutSampleIndex = 0; InSampleIndex < InData.AudioBuffer->Num(); InSampleIndex += InData.NumChannels, OutSampleIndex += OutData.NumChannels)
 			{
+				PlateReverb.ProcessAudioFrame(&AudioData[InSampleIndex], InData.NumChannels, &OutAudioData[OutSampleIndex], InData.NumChannels);
+
+				// Reduce volume of stereo fronts when doing quad reverb to maintain power
+				OutAudioData[OutSampleIndex + EAudioMixerChannel::FrontLeft] *= 0.5f;
+				OutAudioData[OutSampleIndex + EAudioMixerChannel::FrontRight] *= 0.5f;
+
 				// Using standard speaker map order map the right output to the BackLeft channel
 				OutAudioData[OutSampleIndex + EAudioMixerChannel::BackRight] = OutAudioData[OutSampleIndex + EAudioMixerChannel::FrontLeft];
 				OutAudioData[OutSampleIndex + EAudioMixerChannel::BackLeft] = OutAudioData[OutSampleIndex + EAudioMixerChannel::FrontRight];
 			}
-			else
+		}
+		else
+		{
+			for (int32 InSampleIndex = 0, OutSampleIndex = 0; InSampleIndex < InData.AudioBuffer->Num(); InSampleIndex += InData.NumChannels, OutSampleIndex += OutData.NumChannels)
 			{
+				// Processed downmixed audio frame
+				PlateReverb.ProcessAudioFrame(&AudioData[InSampleIndex], InData.NumChannels, &OutAudioData[OutSampleIndex], InData.NumChannels);
+
+				// Reduce volume of stereo fronts when doing quad reverb to maintain power
+				OutAudioData[OutSampleIndex + EAudioMixerChannel::FrontLeft] *= 0.5f;
+				OutAudioData[OutSampleIndex + EAudioMixerChannel::FrontRight] *= 0.5f;
+
+				// Using standard speaker map order map the right output to the BackLeft channel
 				OutAudioData[OutSampleIndex + EAudioMixerChannel::BackLeft] = OutAudioData[OutSampleIndex + EAudioMixerChannel::FrontLeft];
 				OutAudioData[OutSampleIndex + EAudioMixerChannel::BackRight] = OutAudioData[OutSampleIndex + EAudioMixerChannel::FrontRight];
 			}
