[NO NON-PII COMMIT MESSAGE]

diff --git a/Engine/Source/Runtime/AudioMixer/Private/AudioMixerSource.cpp b/Engine/Source/Runtime/AudioMixer/Private/AudioMixerSource.cpp
--- a/Engine/Source/Runtime/AudioMixer/Private/AudioMixerSource.cpp
+++ b/Engine/Source/Runtime/AudioMixer/Private/AudioMixerSource.cpp
@@ -66,7 +66,10 @@ namespace Audio
 		{
 			check(!InWaveInstance->WaveData->RawPCMData || InWaveInstance->WaveData->RawPCMDataSize);
 			const int32 NumBytes = WaveInstance->WaveData->RawPCMDataSize;
-			NumFrames = NumBytes / (WaveInstance->WaveData->NumChannels * sizeof(int16));
+			if (WaveInstance->WaveData->NumChannels > 0)
+			{
+				NumFrames = NumBytes / (WaveInstance->WaveData->NumChannels * sizeof(int16));
+			}
 		}
 
 		// Unfortunately, we need to know if this is a vorbis source since channel maps are different for 5.1 vorbis files
