Fix for potential div by 0

diff --git a/Engine/Source/Runtime/AudioMixer/Private/AudioMixerSourceManager.cpp b/Engine/Source/Runtime/AudioMixer/Private/AudioMixerSourceManager.cpp
--- a/Engine/Source/Runtime/AudioMixer/Private/AudioMixerSourceManager.cpp
+++ b/Engine/Source/Runtime/AudioMixer/Private/AudioMixerSourceManager.cpp
@@ -1173,7 +1173,7 @@ namespace Audio
 			}
 
 			// If we have audio in our queue, we're still playing
-			if (SourceInfo.MixerSourceBuffer->GetNumBuffersQueued() > 0)
+			if (SourceInfo.MixerSourceBuffer->GetNumBuffersQueued() > 0 && NumChannels > 0)
 			{
 				SourceInfo.CurrentPCMBuffer = SourceInfo.MixerSourceBuffer->GetNextBuffer();
 				SourceInfo.CurrentAudioChunkNumFrames = SourceInfo.CurrentPCMBuffer->AudioData.Num() / NumChannels;
