Fix for audio mixer and HTML5 (and single-threaded mode).

- Bring back non-thread-event-based command queue pumping system for single-threaded mode.

diff --git a/Engine/Source/Runtime/AudioMixer/Private/AudioMixerSourceManager.cpp b/Engine/Source/Runtime/AudioMixer/Private/AudioMixerSourceManager.cpp
--- a/Engine/Source/Runtime/AudioMixer/Private/AudioMixerSourceManager.cpp
+++ b/Engine/Source/Runtime/AudioMixer/Private/AudioMixerSourceManager.cpp
@@ -264,6 +264,7 @@ namespace Audio
 		}
 
 		bInitialized = true;
+		bPumpQueue = false;
 	}
 
 	void FMixerSourceManager::Update()
@@ -281,22 +282,46 @@ namespace Audio
 		}
 #endif
 
-		// If the command was triggered, then we want to do a swap of command buffers
-		if (CommandsProcessedEvent->Wait(0))
+		if (FPlatformProcess::SupportsMultithreading())
 		{
-			int32 CurrentGameIndex = AudioThreadCommandBufferIndex.GetValue();
+			// If the command was triggered, then we want to do a swap of command buffers
+			if (CommandsProcessedEvent->Wait(0))
+			{
+				int32 CurrentGameIndex = AudioThreadCommandBufferIndex.GetValue();
+
+				// This flags the audio render thread to be able to pump the next batch of commands
+				// And will allow the audio thread to write to a new command slot
+				const int32 NextIndex = (CurrentGameIndex + 1) & 1;
+
+				// Make sure we've actually emptied the command queue from the render thread before writing to it
+				check(CommandBuffers[NextIndex].SourceCommandQueue.Num() == 0);
+				AudioThreadCommandBufferIndex.Set(NextIndex);
+				RenderThreadCommandBufferIndex.Set(CurrentGameIndex);
 
-			// This flags the audio render thread to be able to pump the next batch of commands
-			// And will allow the audio thread to write to a new command slot
-			const int32 NextIndex = (CurrentGameIndex + 1) & 1;
+				CommandsProcessedEvent->Reset();
+			}
+		}
+		else
+		{
+			int32 CurrentRenderIndex = RenderThreadCommandBufferIndex.GetValue();
+			int32 CurrentGameIndex = AudioThreadCommandBufferIndex.GetValue();
+			check(CurrentGameIndex == 0 || CurrentGameIndex == 1);
+			check(CurrentRenderIndex == 0 || CurrentRenderIndex == 1);
 
-			// Make sure we've actually emptied the command queue from the render thread before writing to it
-			check(CommandBuffers[NextIndex].SourceCommandQueue.Num() == 0);
-			AudioThreadCommandBufferIndex.Set(NextIndex);
-			RenderThreadCommandBufferIndex.Set(CurrentGameIndex);
+			// If these values are the same, that means the audio render thread has finished the last buffer queue so is ready for the next block
+			if (CurrentRenderIndex == CurrentGameIndex)
+			{
+				// This flags the audio render thread to be able to pump the next batch of commands
+				// And will allow the audio thread to write to a new command slot
+				const int32 NextIndex = !CurrentGameIndex;
 
-			CommandsProcessedEvent->Reset();
+				// Make sure we've actually emptied the command queue from the render thread before writing to it
+				check(CommandBuffers[NextIndex].SourceCommandQueue.Num() == 0);
+				AudioThreadCommandBufferIndex.Set(NextIndex);
+				bPumpQueue = true;
+			}
 		}
+
 	}
 
 	void FMixerSourceManager::ReleaseSource(const int32 SourceId)
@@ -2314,8 +2339,16 @@ namespace Audio
 
 		SCOPE_CYCLE_COUNTER(STAT_AudioMixerSourceManagerUpdate);
 
-		// Get the this blocks commands before rendering audio
-		PumpCommandQueue();
+		if (FPlatformProcess::SupportsMultithreading())
+		{
+			// Get the this blocks commands before rendering audio
+			PumpCommandQueue();
+		}
+		else if (bPumpQueue)
+		{
+			bPumpQueue = false;
+			PumpCommandQueue();
+		}
 
 		// Update pending tasks and release them if they're finished
 		UpdatePendingReleaseData();
@@ -2384,9 +2417,12 @@ namespace Audio
 	void FMixerSourceManager::PumpCommandQueue()
 	{
 		// If we're already triggered, we need to wait for the audio thread to reset it before pumping
-		if (CommandsProcessedEvent->Wait(0))
+		if (FPlatformProcess::SupportsMultithreading())
 		{
-			return;
+			if (CommandsProcessedEvent->Wait(0))
+			{
+				return;
+			}
 		}
 
 		int32 CurrentRenderThreadIndex = RenderThreadCommandBufferIndex.GetValue();
@@ -2402,8 +2438,16 @@ namespace Audio
 
 		Commands.SourceCommandQueue.Reset();
 
-		check(CommandsProcessedEvent != nullptr);
-		CommandsProcessedEvent->Trigger();
+		if (FPlatformProcess::SupportsMultithreading())
+		{
+			check(CommandsProcessedEvent != nullptr);
+			CommandsProcessedEvent->Trigger();
+		}
+		else
+		{
+			RenderThreadCommandBufferIndex.Set(!CurrentRenderThreadIndex);
+		}
+
 	}
 
 	void FMixerSourceManager::FlushCommandQueue()
diff --git a/Engine/Source/Runtime/AudioMixer/Private/AudioMixerSourceManager.h b/Engine/Source/Runtime/AudioMixer/Private/AudioMixerSourceManager.h
--- a/Engine/Source/Runtime/AudioMixer/Private/AudioMixerSourceManager.h
+++ b/Engine/Source/Runtime/AudioMixer/Private/AudioMixerSourceManager.h
@@ -549,6 +549,9 @@ namespace Audio
 		uint8 bInitialized : 1;
 		uint8 bUsingSpatializationPlugin : 1;
 
+		// Set to true when the audio source manager should pump the command queue
+		FThreadSafeBool bPumpQueue;
+
 		friend class FMixerSourceVoice;
 	};
 }
