[NO NON-PII COMMIT MESSAGE]

diff --git a/Engine/Source/Runtime/AudioMixer/Private/AudioMixerSourceManager.h b/Engine/Source/Runtime/AudioMixer/Private/AudioMixerSourceManager.h
--- a/Engine/Source/Runtime/AudioMixer/Private/AudioMixerSourceManager.h
+++ b/Engine/Source/Runtime/AudioMixer/Private/AudioMixerSourceManager.h
@@ -152,6 +152,7 @@ namespace Audio
 		{
 			checkSlow(InNumInChannels <= AUDIO_MIXER_MAX_OUTPUT_CHANNELS);
 			checkSlow(InNumOutChannels <= AUDIO_MIXER_MAX_OUTPUT_CHANNELS);
+			FMemory::Memzero(ChannelStartGains, CopySize);
 		}
 
 		FORCEINLINE void Reset(int32 InNumInChannels, int32 InNumOutChannels)
@@ -160,6 +161,7 @@ namespace Audio
 			checkSlow(InNumOutChannels <= AUDIO_MIXER_MAX_OUTPUT_CHANNELS);
 
 			CopySize = InNumInChannels * InNumOutChannels * sizeof(float);
+			FMemory::Memzero(ChannelStartGains, CopySize);
 			FMemory::Memzero(ChannelDestinationGains, CopySize);
 		}
 
