[NO NON-PII COMMIT MESSAGE]

diff --git a/Engine/Source/Editor/AudioEditor/Private/AudioEditorModule.cpp b/Engine/Source/Editor/AudioEditor/Private/AudioEditorModule.cpp
--- a/Engine/Source/Editor/AudioEditor/Private/AudioEditorModule.cpp
+++ b/Engine/Source/Editor/AudioEditor/Private/AudioEditorModule.cpp
@@ -92,7 +92,10 @@ public:
 		SetupIcons();
 
 #if WITH_SNDFILE_IO
-		Audio::InitSoundFileIOManager();
+		if (!Audio::InitSoundFileIOManager())
+		{
+			UE_LOG(LogAudioEditor, Error, TEXT("LibSoundFile failed to load. Importing audio will not work correctly."));
+		}
 #endif // WITH_SNDFILE_IO
 	}
 
diff --git a/Engine/Source/Runtime/AudioMixer/Private/SoundFileIO/SoundFileIOManagerImpl.cpp b/Engine/Source/Runtime/AudioMixer/Private/SoundFileIO/SoundFileIOManagerImpl.cpp
--- a/Engine/Source/Runtime/AudioMixer/Private/SoundFileIO/SoundFileIOManagerImpl.cpp
+++ b/Engine/Source/Runtime/AudioMixer/Private/SoundFileIO/SoundFileIOManagerImpl.cpp
@@ -282,9 +282,23 @@ namespace Audio
 		}
 
 		// open a sound file handle to get the description
-		*OutFileHandle = SoundFileOpen(TCHAR_TO_ANSI(*FilePath), ESoundFileOpenMode::READING, &OutputDescription);
+		if (SoundFileOpen != nullptr)
+		{
+			*OutFileHandle = SoundFileOpen(TCHAR_TO_ANSI(*FilePath), ESoundFileOpenMode::READING, &OutputDescription);
+		}
+		else
+		{
+			UE_LOG(LogAudioMixer, Error, TEXT("LibSoundFile failed to load symbols for SoundFileOpen."));
+			*OutFileHandle = nullptr;
+		}
+
 		if (!*OutFileHandle)
 		{
+			if (!SoundFileStrError)
+			{
+				return ESoundFileError::Type::INVALID_DATA;
+			}
+
 			FString StrError = FString(SoundFileStrError(nullptr));
 			UE_LOG(LogAudioMixer, Error, TEXT("Failed to open sound file %s: %s"), *FilePath, *StrError);
 			return ESoundFileError::Type::FAILED_TO_OPEN;
@@ -294,7 +308,15 @@ namespace Audio
 		int32 NumChannels = OutputDescription.NumChannels;
 		OutChannelMap.Init(ESoundFileChannelMap::Type::INVALID, NumChannels);
 
-		int32 Result = SoundFileCommand(*OutFileHandle, GET_CHANNEL_MAP_INFO, (int32*)OutChannelMap.GetData(), sizeof(int32)*NumChannels);
+		int32 Result = 0;
+		if (SoundFileCommand)
+		{
+			Result = SoundFileCommand(*OutFileHandle, GET_CHANNEL_MAP_INFO, (int32*)OutChannelMap.GetData(), sizeof(int32)*NumChannels);
+		}
+		else
+		{
+			UE_LOG(LogAudioMixer, Error, TEXT("LibSoundFile wasn't properly loaded with symbols for SoundFileCommand."));
+		}
 
 		// If we failed to get the file's channel map definition, then we set the default based on the number of channels
 		if (Result == 0)
@@ -522,6 +544,13 @@ namespace Audio
 
 		ESoundFileError::Type Release() override
 		{
+			if (!SoundFileClose)
+			{
+				UE_LOG(LogAudioMixer, Error, TEXT("LibSoundFile failed to load symbols for SoundFileClose"));
+				FileHandle = nullptr;
+				return ESoundFileError::Type::INVALID_STATE;
+			}
+
 			if (FileHandle)
 			{
 				SoundFileClose(FileHandle);
@@ -540,9 +569,20 @@ namespace Audio
 
 		ESoundFileError::Type SeekFrames(SoundFileCount Offset, ESoundFileSeekMode::Type SeekMode, SoundFileCount& OutOffset) override
 		{
+			if (!SoundFileSeek)
+			{
+				UE_LOG(LogAudioMixer, Error, TEXT("LibSoundFile failed to load symbols for SoundFileSeek"));
+				return ESoundFileError::Type::INVALID_STATE;
+			}
+
 			SoundFileCount Pos = SoundFileSeek(FileHandle, Offset, (int32)SeekMode);
 			if (Pos == -1)
 			{
+				if (!SoundFileStrError)
+				{
+					return SetError(ESoundFileError::Type::INVALID_STATE);
+				}
+
 				FString StrErr = SoundFileStrError(FileHandle);
 				UE_LOG(LogAudioMixer, Error, TEXT("Failed to seek file: %s"), *StrErr);
 				return SetError(ESoundFileError::Type::FAILED_TO_SEEK);
@@ -552,24 +592,48 @@ namespace Audio
 
 		ESoundFileError::Type ReadFrames(float* DataPtr, SoundFileCount NumFrames, SoundFileCount& OutNumFramesRead) override
 		{
+			if (!SoundFileReadFramesFloat)
+			{
+				UE_LOG(LogAudioMixer, Error, TEXT("LibSoundFile failed to load symbols for SoundFileReadFramesFloat"));
+				return ESoundFileError::Type::INVALID_STATE;
+			}
+
 			OutNumFramesRead = SoundFileReadFramesFloat(FileHandle, DataPtr, NumFrames);
 			return ESoundFileError::Type::NONE;
 		}
 
 		ESoundFileError::Type ReadFrames(double* DataPtr, SoundFileCount NumFrames, SoundFileCount& OutNumFramesRead) override
 		{
+			if (!SoundFileReadFramesDouble)
+			{
+				UE_LOG(LogAudioMixer, Error, TEXT("LibSoundFile failed to load symbols for SoundFileReadFramesDouble"));
+				return ESoundFileError::Type::INVALID_STATE;
+			}
+
 			OutNumFramesRead = SoundFileReadFramesDouble(FileHandle, DataPtr, NumFrames);
 			return ESoundFileError::Type::NONE;
 		}
 
 		ESoundFileError::Type ReadSamples(float* DataPtr, SoundFileCount NumSamples, SoundFileCount& OutNumSamplesRead) override
 		{
+			if (!SoundFileReadSamplesFloat)
+			{
+				UE_LOG(LogAudioMixer, Error, TEXT("LibSoundFile failed to load symbols for SoundFileReadSamplesFloat"));
+				return ESoundFileError::Type::INVALID_STATE;
+			}
+
 			OutNumSamplesRead = SoundFileReadSamplesFloat(FileHandle, DataPtr, NumSamples);
 			return ESoundFileError::Type::NONE;
 		}
 
 		ESoundFileError::Type ReadSamples(double* DataPtr, SoundFileCount NumSamples, SoundFileCount& OutNumSamplesRead) override
 		{
+			if (!SoundFileReadSamplesDouble)
+			{
+				UE_LOG(LogAudioMixer, Error, TEXT("LibSoundFile failed to load symbols for SoundFileReadSamplesDouble"));
+				return ESoundFileError::Type::INVALID_STATE;
+			}
+
 			OutNumSamplesRead = SoundFileReadSamplesDouble(FileHandle, DataPtr, NumSamples);
 			return ESoundFileError::Type::NONE;
 		}
@@ -624,14 +688,35 @@ namespace Audio
 
 			FSoundFileDescription Description;
 			SoundFileData->GetDescription(Description);
+
+			if (!SoundFileFormatCheck)
+			{
+				UE_LOG(LogAudioMixer, Error, TEXT("LibSoundFile failed to load symbols for SoundFileFormatCheck"));
+				return SetError(ESoundFileError::Type::INVALID_STATE);
+			}
+
 			if (!SoundFileFormatCheck(&Description))
 			{
 				return SetError(ESoundFileError::Type::INVALID_INPUT_FORMAT);
 			}
 
-			FileHandle = SoundFileOpenVirtual(&VirtualSoundFileInfo, ESoundFileOpenMode::READING, &Description, (void*)this);
+			if (!SoundFileOpenVirtual)
+			{
+				UE_LOG(LogAudioMixer, Error, TEXT("LibSoundFile failed to load symbols for SoundFileOpenVirtual"));
+				FileHandle = nullptr;
+			}
+			else
+			{
+				FileHandle = SoundFileOpenVirtual(&VirtualSoundFileInfo, ESoundFileOpenMode::READING, &Description, (void*)this);
+			}
+			
 			if (!FileHandle)
 			{
+				if (!SoundFileStrError)
+				{
+					return SetError(ESoundFileError::Type::INVALID_DATA);
+				}
+
 				FString StrErr = SoundFileStrError(nullptr);
 				UE_LOG(LogAudioMixer, Error, TEXT("Failed to intitialize sound file: %s"), *StrErr);
 				return SetError(ESoundFileError::Type::FAILED_TO_OPEN);
@@ -859,9 +944,23 @@ namespace Audio
 			VirtualSoundFileInfo.VirtualSoundFileWrite = OnSoundFileWriteBytes;
 			VirtualSoundFileInfo.VirtualSoundFileTell = OnSoundFileTell;
 
-			FileHandle = SoundFileOpenVirtual(&VirtualSoundFileInfo, ESoundFileOpenMode::READING, &Description, (void*)this);
+			if (SoundFileOpenVirtual)
+			{
+				FileHandle = SoundFileOpenVirtual(&VirtualSoundFileInfo, ESoundFileOpenMode::READING, &Description, (void*)this);
+			}
+			else
+			{
+				UE_LOG(LogAudioMixer, Error, TEXT("LibSoundFile failed to load symbols for SoundFileOpenVirtual."));
+				FileHandle = nullptr;
+			}
+
 			if (!FileHandle)
 			{
+				if (!SoundFileStrError)
+				{
+					return SetError(ESoundFileError::Type::INVALID_DATA);
+				}
+
 				FString StrErr = SoundFileStrError(nullptr);
 				UE_LOG(LogAudioMixer, Error, TEXT("Failed to initialize sound file: %s"), *StrErr);
 				return SetError(ESoundFileError::Type::FAILED_TO_OPEN);
@@ -870,8 +969,16 @@ namespace Audio
 			// Try to get a channel mapping
 			int32 NumChannels = Description.NumChannels;
 			ChannelMap.Init(ESoundFileChannelMap::Type::INVALID, NumChannels);
-
-			int32 Result = SoundFileCommand(FileHandle, GET_CHANNEL_MAP_INFO, (int32*)ChannelMap.GetData(), sizeof(int32)*NumChannels);
+			
+			int32 Result = 0;
+			if (SoundFileCommand)
+			{
+				Result = SoundFileCommand(FileHandle, GET_CHANNEL_MAP_INFO, (int32*)ChannelMap.GetData(), sizeof(int32)*NumChannels);
+			}
+			else
+			{
+				UE_LOG(LogAudioMixer, Error, TEXT("LibSoundFile failed to load symbols for SoundFileCommand."));
+			}
 
 			// If we failed to get the file's channel map definition, then we set the default based on the number of channels
 			if (Result == 0)
@@ -920,9 +1027,20 @@ namespace Audio
 
 		ESoundFileError::Type SeekFrames(SoundFileCount Offset, ESoundFileSeekMode::Type SeekMode, SoundFileCount& OutOffset) override
 		{
+			if (!SoundFileSeek)
+			{
+				UE_LOG(LogAudioMixer, Error, TEXT("LibSoundFile failed to load."));
+				return ESoundFileError::Type::INVALID_STATE;
+			}
+
 			SoundFileCount Pos = SoundFileSeek(FileHandle, Offset, (int32)SeekMode);
 			if (Pos == -1)
 			{
+				if (!SoundFileStrError)
+				{
+					return SetError(ESoundFileError::Type::INVALID_DATA);
+				}
+
 				FString StrErr = SoundFileStrError(FileHandle);
 				UE_LOG(LogAudioMixer, Error, TEXT("Failed to seek file: %s"), *StrErr);
 				return SetError(ESoundFileError::Type::FAILED_TO_SEEK);
@@ -932,24 +1050,48 @@ namespace Audio
 
 		ESoundFileError::Type ReadFrames(float* DataPtr, SoundFileCount NumFrames, SoundFileCount& OutNumFramesRead) override
 		{
+			if (!SoundFileReadFramesFloat)
+			{
+				UE_LOG(LogAudioMixer, Error, TEXT("LibSoundFile failed to load symbols for SoundFileReadFramesFloat"));
+				return ESoundFileError::Type::INVALID_STATE;
+			}
+
 			OutNumFramesRead = SoundFileReadFramesFloat(FileHandle, DataPtr, NumFrames);
 			return ESoundFileError::Type::NONE;
 		}
 
 		ESoundFileError::Type ReadFrames(double* DataPtr, SoundFileCount NumFrames, SoundFileCount& OutNumFramesRead) override
 		{
+			if (!SoundFileReadFramesDouble)
+			{
+				UE_LOG(LogAudioMixer, Error, TEXT("LibSoundFile failed to load symbols for SoundFileReadFramesDouble"));
+				return ESoundFileError::Type::INVALID_STATE;
+			}
+
 			OutNumFramesRead = SoundFileReadFramesDouble(FileHandle, DataPtr, NumFrames);
 			return ESoundFileError::Type::NONE;
 		}
 
 		ESoundFileError::Type ReadSamples(float* DataPtr, SoundFileCount NumSamples, SoundFileCount& OutNumSamplesRead) override
 		{
+			if (!SoundFileReadSamplesFloat)
+			{
+				UE_LOG(LogAudioMixer, Error, TEXT("LibSoundFile failed to load symbols for SoundFileReadSamplesFloat"));
+				return ESoundFileError::Type::INVALID_STATE;
+			}
+
 			OutNumSamplesRead = SoundFileReadSamplesFloat(FileHandle, DataPtr, NumSamples);
 			return ESoundFileError::Type::NONE;
 		}
 
 		ESoundFileError::Type ReadSamples(double* DataPtr, SoundFileCount NumSamples, SoundFileCount& OutNumSamplesRead) override
 		{
+			if (!SoundFileReadSamplesDouble)
+			{
+				UE_LOG(LogAudioMixer, Error, TEXT("LibSoundFile failed to load symbols for SoundFileReadSamplesDouble"));
+				return ESoundFileError::Type::INVALID_STATE;
+			}
+
 			OutNumSamplesRead = SoundFileReadSamplesDouble(FileHandle, DataPtr, NumSamples);
 			return ESoundFileError::Type::NONE;
 		}
@@ -1113,6 +1255,12 @@ namespace Audio
 			ChannelMap = InChannelMap;
 			EncodingQuality = InEncodingQuality;
 
+			if (!SoundFileFormatCheck)
+			{
+				UE_LOG(LogAudioMixer, Error, TEXT("LibSoundFile failed to load symbols for SoundFileFormatCheck"));
+				return SetError(ESoundFileError::Type::INVALID_STATE);
+			}
+
 			// First check the input format to make sure it's valid
 			if (!SoundFileFormatCheck(&InDescription))
 			{
@@ -1139,17 +1287,45 @@ namespace Audio
 			VirtualSoundFileInfo.VirtualSoundFileWrite = OnSoundFileWriteBytes;
 			VirtualSoundFileInfo.VirtualSoundFileTell = OnSoundFileTell;
 
-			FileHandle = SoundFileOpenVirtual(&VirtualSoundFileInfo, ESoundFileOpenMode::WRITING, &Description, (void*)this);
+			if (SoundFileOpenVirtual)
+			{
+				FileHandle = SoundFileOpenVirtual(&VirtualSoundFileInfo, ESoundFileOpenMode::WRITING, &Description, (void*)this);
+			}
+			else
+			{
+				UE_LOG(LogAudioMixer, Error, TEXT("LibSoundFile failed to load symbols for SoundFileOpenVirtual"));
+				FileHandle = nullptr;
+			}
+
 			if (!FileHandle)
 			{
+				if (!SoundFileStrError)
+				{
+					return SetError(ESoundFileError::Type::INVALID_DATA);
+				}
+
 				FString StrErr = FString(SoundFileStrError(nullptr));
 				UE_LOG(LogAudioMixer, Error, TEXT("Failed to open empty sound file: %s"), *StrErr);
 				return SetError(ESoundFileError::Type::FAILED_TO_OPEN);
 			}
 
-			int32 Result = SoundFileCommand(FileHandle, SET_CHANNEL_MAP_INFO, (int32*)InChannelMap.GetData(), sizeof(int32)*Description.NumChannels);
+			int32 Result = 0;
+			if (SoundFileCommand)
+			{
+				SoundFileCommand(FileHandle, SET_CHANNEL_MAP_INFO, (int32*)InChannelMap.GetData(), sizeof(int32)*Description.NumChannels);
+			}
+			else
+			{
+				UE_LOG(LogAudioMixer, Error, TEXT("LibSoundFile failed to load symbols for SoundFileCommand"));
+			}
+
 			if (Result != 1)
 			{
+				if (!SoundFileStrError)
+				{
+					return ESoundFileError::Type::INVALID_DATA;
+				}
+
 				// The result is returning 0 (false), however 'No Error'
 				// is provided and the file mapping is correct.
 				FString StrErr = SoundFileStrError(nullptr);
@@ -1162,9 +1338,24 @@ namespace Audio
 
 			if ((Description.FormatFlags & ESoundFileFormat::MAJOR_FORMAT_MASK) == ESoundFileFormat::OGG)
 			{
-				int32 Result2 = SoundFileCommand(FileHandle, SET_ENCODING_QUALITY, &EncodingQuality, sizeof(double));
+				int32 Result2 = 0;
+				if (SoundFileCommand)
+				{
+					SoundFileCommand(FileHandle, SET_ENCODING_QUALITY, &EncodingQuality, sizeof(double));
+				}
+				else
+				{
+					UE_LOG(LogAudioMixer, Error, TEXT("LibSoundFile failed to load symbols for SoundFileCommand"));
+				}
+
 				if (Result2 != 1)
 				{
+					if (!SoundFileStrError)
+					{
+						UE_LOG(LogAudioMixer, Error, TEXT("LibSoundFile failed to load symbols for SoundFileStrError"));
+						return ESoundFileError::Type::INVALID_DATA;
+					}
+
 					FString StrErr = SoundFileStrError(FileHandle);
 					UE_LOG(LogAudioMixer, Error, TEXT("Failed to set encoding quality: %s"), *StrErr);
 					return SetError(ESoundFileError::Type::BAD_ENCODING_QUALITY);
@@ -1176,6 +1367,13 @@ namespace Audio
 
 		ESoundFileError::Type Release() override
 		{
+			if (!SoundFileClose)
+			{
+				UE_LOG(LogAudioMixer, Error, TEXT("LibSoundFile failed to load symbols for SoundFileClose"));
+				FileHandle = nullptr;
+				return ESoundFileError::Type::INVALID_STATE;
+			}
+
 			if (FileHandle)
 			{
 				SoundFileClose(FileHandle);
@@ -1186,9 +1384,20 @@ namespace Audio
 
 		ESoundFileError::Type SeekFrames(SoundFileCount Offset, ESoundFileSeekMode::Type SeekMode, SoundFileCount& OutOffset) override
 		{
+			if (!SoundFileSeek)
+			{
+				UE_LOG(LogAudioMixer, Error, TEXT("LibSoundFile failed to load symbols for SoundFileSeek"));
+				return ESoundFileError::Type::INVALID_STATE;
+			}
+
 			SoundFileCount Pos = SoundFileSeek(FileHandle, Offset, (int32)SeekMode);
 			if (Pos == -1)
 			{
+				if (!SoundFileStrError)
+				{
+					return ESoundFileError::Type::INVALID_DATA;
+				}
+
 				FString StrErr = SoundFileStrError(FileHandle);
 				UE_LOG(LogAudioMixer, Error, TEXT("Failed to seek file: %s"), *StrErr);
 				return SetError(ESoundFileError::Type::FAILED_TO_SEEK);
@@ -1198,24 +1407,48 @@ namespace Audio
 
 		ESoundFileError::Type WriteFrames(const float* DataPtr, SoundFileCount NumFrames, SoundFileCount& OutNumFramesWritten) override
 		{
+			if (!SoundFileWriteFramesFloat)
+			{
+				UE_LOG(LogAudioMixer, Error, TEXT("LibSoundFile failed to load symbols for SoundFileWriteFramesFloat"));
+				return ESoundFileError::Type::INVALID_STATE;
+			}
+
 			OutNumFramesWritten = SoundFileWriteFramesFloat(FileHandle, DataPtr, NumFrames);
 			return ESoundFileError::Type::NONE;
 		}
 
 		ESoundFileError::Type WriteFrames(const double* DataPtr, SoundFileCount NumFrames, SoundFileCount& OutNumFramesWritten) override
 		{
+			if (!SoundFileWriteFramesDouble)
+			{
+				UE_LOG(LogAudioMixer, Error, TEXT("LibSoundFile failed to load symbols for SoundFileWriteFramesDouble"));
+				return ESoundFileError::Type::INVALID_STATE;
+			}
+
 			OutNumFramesWritten = SoundFileWriteFramesDouble(FileHandle, DataPtr, NumFrames);
 			return ESoundFileError::Type::NONE;
 		}
 
 		ESoundFileError::Type WriteSamples(const float* DataPtr, SoundFileCount NumSamples, SoundFileCount& OutNumSampleWritten) override
 		{
+			if (!SoundFileWriteSamplesFloat)
+			{
+				UE_LOG(LogAudioMixer, Error, TEXT("LibSoundFile failed to load symbols for SoundFileWriteSamplesFloat"));
+				return ESoundFileError::Type::INVALID_STATE;
+			}
+
 			OutNumSampleWritten = SoundFileWriteSamplesFloat(FileHandle, DataPtr, NumSamples);
 			return ESoundFileError::Type::NONE;
 		}
 
 		ESoundFileError::Type WriteSamples(const double* DataPtr, SoundFileCount NumSamples, SoundFileCount& OutNumSampleWritten) override
 		{
+			if (!SoundFileWriteSamplesDouble)
+			{
+				UE_LOG(LogAudioMixer, Error, TEXT("LibSoundFile failed to load symbols for SoundFileWriteSamplesDouble"));
+				return ESoundFileError::Type::INVALID_STATE;
+			}
+
 			OutNumSampleWritten = SoundFileWriteSamplesDouble(FileHandle, DataPtr, NumSamples);
 			return ESoundFileError::Type::NONE;
 		}
@@ -1293,6 +1526,13 @@ namespace Audio
 		if (Error == ESoundFileError::Type::NONE)
 		{
 			check(FileHandle != nullptr);
+
+			if (!SoundFileClose)
+			{
+				UE_LOG(LogAudioMixer, Error, TEXT("LibSoundFile failed to load symbols for SoundFileClose"));
+				return false;
+			}
+
 			SoundFileClose(FileHandle);
 			return true;
 		}
@@ -1338,6 +1578,12 @@ namespace Audio
 		ESoundFileError::Type Error = GetSoundDesriptionInternal(&FileHandle, FilePath, Description, ChannelMap);
 		if (FileHandle)
 		{
+			if (!SoundFileClose)
+			{
+				UE_LOG(LogAudioMixer, Error, TEXT("LibSoundFile failed to load symbols for SoundFileClose"));
+				return Error;
+			}
+
 			SoundFileClose(FileHandle);
 		}
 		return Error;
