Fix static analysis warnings for PVSStudio. #jira pending #rb none #fyi rob.gay

diff --git a/Engine/Source/Runtime/AudioMixer/Private/SoundFileIO/SoundFileIOManagerImpl.cpp b/Engine/Source/Runtime/AudioMixer/Private/SoundFileIO/SoundFileIOManagerImpl.cpp
--- a/Engine/Source/Runtime/AudioMixer/Private/SoundFileIO/SoundFileIOManagerImpl.cpp
+++ b/Engine/Source/Runtime/AudioMixer/Private/SoundFileIO/SoundFileIOManagerImpl.cpp
@@ -1312,7 +1312,7 @@ namespace Audio
 			int32 Result = 0;
 			if (SoundFileCommand)
 			{
-				SoundFileCommand(FileHandle, SET_CHANNEL_MAP_INFO, (int32*)InChannelMap.GetData(), sizeof(int32)*Description.NumChannels);
+				Result = SoundFileCommand(FileHandle, SET_CHANNEL_MAP_INFO, (int32*)InChannelMap.GetData(), sizeof(int32)*Description.NumChannels);
 			}
 			else
 			{
@@ -1341,7 +1341,7 @@ namespace Audio
 				int32 Result2 = 0;
 				if (SoundFileCommand)
 				{
-					SoundFileCommand(FileHandle, SET_ENCODING_QUALITY, &EncodingQuality, sizeof(double));
+					Result2 = SoundFileCommand(FileHandle, SET_ENCODING_QUALITY, &EncodingQuality, sizeof(double));
 				}
 				else
 				{
