Audio - Discussed with Aaron; Hiding the bUseSpatialization behind a getter and setter so that we can control all spatialization with a cvar, the old behavior wasn't actually working, as several things could override the bUseSpatialization behavior of the WaveInstance.

Rob.Gay, Aaron.McLeran

diff --git a/Engine/Source/Runtime/ALAudio/Private/ALAudioSource.cpp b/Engine/Source/Runtime/ALAudio/Private/ALAudioSource.cpp
--- a/Engine/Source/Runtime/ALAudio/Private/ALAudioSource.cpp
+++ b/Engine/Source/Runtime/ALAudio/Private/ALAudioSource.cpp
@@ -36,7 +36,7 @@ bool FALSoundSource::Init( FWaveInstance* InWaveInstance )
 			WaveInstance = InWaveInstance;
 
 			// Enable/disable spatialization of sounds
-			alSourcei( SourceId, AL_SOURCE_RELATIVE, WaveInstance->bUseSpatialization ? AL_FALSE : AL_TRUE );
+			alSourcei( SourceId, AL_SOURCE_RELATIVE, WaveInstance->GetUseSpatialization() ? AL_FALSE : AL_TRUE );
 
 			// Setting looping on a real time decompressed source suppresses the buffers processed message
 			alSourcei( SourceId, AL_LOOPING, ( WaveInstance->LoopingMode == LOOP_Forever ) ? AL_TRUE : AL_FALSE );
@@ -132,7 +132,7 @@ void FALSoundSource::Update( void )
 
 	// We're using a relative coordinate system for un- spatialized sounds.
 	FVector RelativeDirection = FVector::ZeroVector;
-	if (WaveInstance->bUseSpatialization)
+	if (WaveInstance->GetUseSpatialization())
 	{
 		FVector UnnormalizedDirection = ((FALAudioDevice*)AudioDevice)->InverseTransform.TransformPosition(WaveInstance->Location);
 		RelativeDirection = UnnormalizedDirection.GetSafeNormal();
diff --git a/Engine/Source/Runtime/AudioMixer/Private/AudioMixerSource.cpp b/Engine/Source/Runtime/AudioMixer/Private/AudioMixerSource.cpp
--- a/Engine/Source/Runtime/AudioMixer/Private/AudioMixerSource.cpp
+++ b/Engine/Source/Runtime/AudioMixer/Private/AudioMixerSource.cpp
@@ -276,7 +276,7 @@ namespace Audio
 #endif
 
 			// Whether or not we're 3D
-			bIs3D = !UseObjectBasedSpatialization() && WaveInstance->bUseSpatialization && SoundBuffer->NumChannels < 3;
+			bIs3D = !UseObjectBasedSpatialization() && WaveInstance->GetUseSpatialization() && SoundBuffer->NumChannels < 3;
 
 			// Grab the source's reverb plugin settings 
 			InitParams.SpatializationPluginSettings = UseSpatializationPlugin() ? WaveInstance->SpatializationPluginSettings : nullptr;
@@ -823,7 +823,7 @@ namespace Audio
 	void FMixerSource::UpdateSpatialization()
 	{
 		SpatializationParams = GetSpatializationParams();
-		if (WaveInstance->bUseSpatialization)
+		if (WaveInstance->GetUseSpatialization())
 		{
 			MixerSourceVoice->SetSpatializationParams(SpatializationParams);
 		}
@@ -949,7 +949,7 @@ namespace Audio
 			// Treat the source as if it is a 2D stereo source:
 			return ComputeStereoChannelMap(SubmixChannelType, OutChannelMap);
 		}
-		else if (WaveInstance->bUseSpatialization && (!FMath::IsNearlyEqual(WaveInstance->AbsoluteAzimuth, PreviousAzimuth, 0.01f) || MixerSourceVoice->NeedsSpeakerMap()))
+		else if (WaveInstance->GetUseSpatialization() && (!FMath::IsNearlyEqual(WaveInstance->AbsoluteAzimuth, PreviousAzimuth, 0.01f) || MixerSourceVoice->NeedsSpeakerMap()))
 		{
 			// Don't need to compute the source channel map if the absolute azimuth hasn't changed much
 			PreviousAzimuth = WaveInstance->AbsoluteAzimuth;
@@ -971,7 +971,7 @@ namespace Audio
 	bool FMixerSource::ComputeStereoChannelMap(const ESubmixChannelFormat InSubmixChannelType, Audio::AlignedFloatBuffer& OutChannelMap)
 	{
 		// Only recalculate positional data if the source has moved a significant amount:
-		if (WaveInstance->bUseSpatialization && (!FMath::IsNearlyEqual(WaveInstance->AbsoluteAzimuth, PreviousAzimuth, 0.01f) || MixerSourceVoice->NeedsSpeakerMap()))
+		if (WaveInstance->GetUseSpatialization() && (!FMath::IsNearlyEqual(WaveInstance->AbsoluteAzimuth, PreviousAzimuth, 0.01f) || MixerSourceVoice->NeedsSpeakerMap()))
 		{
 			// Make sure our stereo emitter positions are updated relative to the sound emitter position
 			if (Buffer->NumChannels == 2)
diff --git a/Engine/Source/Runtime/Engine/Private/Audio.cpp b/Engine/Source/Runtime/Engine/Private/Audio.cpp
--- a/Engine/Source/Runtime/Engine/Private/Audio.cpp
+++ b/Engine/Source/Runtime/Engine/Private/Audio.cpp
@@ -66,6 +66,14 @@ FAutoConsoleVariableRef CVarDisableStereoSpread(
 	TEXT("0: Not Disabled, 1: Disabled"),
 	ECVF_Default);
 
+static int32 AllowAudioSpatializationCVar = 1;
+FAutoConsoleVariableRef CVarAllowAudioSpatializationCVar(
+	TEXT("au.AllowAudioSpatialization"),
+	AllowAudioSpatializationCVar,
+	TEXT("Controls if we allow spatialization of audio, normally this is enabled.  If disabled all audio won't be spatialized, but will have attenuation.\n")
+	TEXT("0: Disable, >0: Enable"),
+	ECVF_Default);
+
 bool IsAudioPluginEnabled(EAudioPlugin PluginType)
 {
 	switch (PluginType)
@@ -372,7 +380,7 @@ void FSoundSource::SetFilterFrequency()
 void FSoundSource::UpdateStereoEmitterPositions()
 {
 	// Only call this function if we're told to use spatialization
-	check(WaveInstance->bUseSpatialization);
+	check(WaveInstance->GetUseSpatialization());
 	check(Buffer->NumChannels == 2);
 
 	if (!DisableStereoSpreadCvar && WaveInstance->StereoSpread > 0.0f)
@@ -417,7 +425,7 @@ void FSoundSource::DrawDebugInfo()
 			USoundBase* Sound = WaveInstance->ActiveSound->GetSound();
 			const FVector Location = WaveInstance->Location;
 
-			const bool bSpatialized = Buffer->NumChannels == 2 && WaveInstance->bUseSpatialization;
+			const bool bSpatialized = Buffer->NumChannels == 2 && WaveInstance->GetUseSpatialization();
 			const FVector LeftChannelSourceLoc = LeftChannelSourceLocation;
 			const FVector RightChannelSourceLoc = RightChannelSourceLocation;
 
@@ -527,7 +535,7 @@ FSpatializationParams FSoundSource::GetSpatializationParams()
 {
 	FSpatializationParams Params;
 
-	if (WaveInstance->bUseSpatialization)
+	if (WaveInstance->GetUseSpatialization())
 	{
 		FVector EmitterPosition = AudioDevice->GetListenerTransformedDirection(WaveInstance->Location, &Params.Distance);
 
@@ -970,6 +978,11 @@ bool FWaveInstance::IsStreaming() const
 	return FPlatformProperties::SupportsAudioStreaming() && WaveData != nullptr && WaveData->IsStreaming();
 }
 
+bool FWaveInstance::GetUseSpatialization() const
+{
+	return AllowAudioSpatializationCVar && bUseSpatialization;
+}
+
 FString FWaveInstance::GetName() const
 {
 	if (WaveData)
diff --git a/Engine/Source/Runtime/Engine/Private/AudioDevice.cpp b/Engine/Source/Runtime/Engine/Private/AudioDevice.cpp
--- a/Engine/Source/Runtime/Engine/Private/AudioDevice.cpp
+++ b/Engine/Source/Runtime/Engine/Private/AudioDevice.cpp
@@ -119,14 +119,6 @@ FAutoConsoleVariableRef CVarNumPrecacheFrames(
 	TEXT("0: Use default value for precache frames, >0: Number of frames to precache."),
 	ECVF_Default);
 
-static int32 AllowAudioSpatializationCVar = 1;
-FAutoConsoleVariableRef CVarAllowAudioSpatializationCVar(
-	TEXT("au.AllowAudioSpatialization"),
-	AllowAudioSpatializationCVar,
-	TEXT("Controls if we allow spatialization of audio, normally this is enabled.  If disabled all audio won't be spatialized, but will have attenuation.\n")
-	TEXT("0: Disable, >0: Enable"),
-	ECVF_Default);
-
 static int32 DisableLegacyReverb = 0;
 FAutoConsoleVariableRef CVarDisableLegacyReverb(
 	TEXT("au.DisableLegacyReverb"),
@@ -4214,12 +4206,6 @@ void FAudioDevice::AddNewActiveSound(const FActiveSound& NewActiveSound)
 		OneShotCount++;
 	}
 
-	// If we've disabled audio spatialization, then we need to force this active sound to no longer spatialize.
-	if (AllowAudioSpatializationCVar == 0)
-	{
-		ActiveSound->bAllowSpatialization = false;
-	}
-
 	// Set the active sound to be playing audio so it gets parsed at least once.
 	ActiveSound->bIsPlayingAudio = true;
 
diff --git a/Engine/Source/Runtime/Engine/Private/SoundWave.cpp b/Engine/Source/Runtime/Engine/Private/SoundWave.cpp
--- a/Engine/Source/Runtime/Engine/Private/SoundWave.cpp
+++ b/Engine/Source/Runtime/Engine/Private/SoundWave.cpp
@@ -1571,7 +1571,7 @@ void USoundWave::Parse( FAudioDevice* AudioDevice, const UPTRINT NodeWaveInstanc
 		WaveInstance->Location = ParseParams.Transform.GetTranslation();
 		WaveInstance->bIsStarted = true;
 		WaveInstance->bAlreadyNotifiedHook = false;
-		WaveInstance->bUseSpatialization = ParseParams.bUseSpatialization;
+		WaveInstance->SetUseSpatialization(ParseParams.bUseSpatialization);
 		WaveInstance->SpatializationMethod = ParseParams.SpatializationMethod;
 		WaveInstance->WaveData = this;
 		WaveInstance->NotifyBufferFinishedHooks = ParseParams.NotifyBufferFinishedHooks;
@@ -1579,7 +1579,7 @@ void USoundWave::Parse( FAudioDevice* AudioDevice, const UPTRINT NodeWaveInstanc
 		WaveInstance->bIsPaused = ParseParams.bIsPaused;
 
 		// If we're normalizing 3d stereo spatialized sounds, we need to scale by -6 dB
-		if (WaveInstance->bUseSpatialization && ParseParams.bApplyNormalizationToStereoSounds && NumChannels == 2)
+		if (WaveInstance->GetUseSpatialization() && ParseParams.bApplyNormalizationToStereoSounds && NumChannels == 2)
 		{
 			float WaveInstanceVolume = WaveInstance->GetVolume();
 			WaveInstance->SetVolume(WaveInstanceVolume * 0.5f);
@@ -1648,7 +1648,7 @@ void USoundWave::Parse( FAudioDevice* AudioDevice, const UPTRINT NodeWaveInstanc
 		}
 
 		// Sanity check
-		if( NumChannels > 2 && WaveInstance->bUseSpatialization && !WaveInstance->bReportedSpatializationWarning)
+		if( NumChannels > 2 && WaveInstance->GetUseSpatialization() && !WaveInstance->bReportedSpatializationWarning)
 		{
 			static TSet<USoundWave*> ReportedSounds;
 			if (!ReportedSounds.Contains(this))
diff --git a/Engine/Source/Runtime/Engine/Public/Audio.h b/Engine/Source/Runtime/Engine/Public/Audio.h
--- a/Engine/Source/Runtime/Engine/Public/Audio.h
+++ b/Engine/Source/Runtime/Engine/Public/Audio.h
@@ -340,9 +340,11 @@ public:
 	/** Whether the notify finished hook has been called since the last update/parsenodes */
 	uint32 bAlreadyNotifiedHook:1;
 
+private:
 	/** Whether to use spatialization */
 	uint32 bUseSpatialization:1;
 
+public:
 	/** Whether or not to enable the low pass filter */
 	uint32 bEnableLowPassFilter:1;
 
@@ -517,6 +519,11 @@ public:
 	/** Gets the envelope value of the waveinstance. Only returns non-zero values if it's a real voice. Only implemented in the audio mixer. */
 	float GetEnvelopeValue() const { return EnvelopValue; }
 
+	/** Whether to use spatialization, which controls 3D effects like panning */
+	void SetUseSpatialization(const bool InUseSpatialization) { bUseSpatialization = InUseSpatialization; }
+
+	/** Whether this wave will be spatializized, which controls 3D effects like panning */
+	bool GetUseSpatialization() const;
 };
 
 inline uint32 GetTypeHash(FWaveInstance* A) { return A->TypeHash; }
diff --git a/Engine/Source/Runtime/IOS/IOSAudio/Private/IOSAudioSource.cpp b/Engine/Source/Runtime/IOS/IOSAudio/Private/IOSAudioSource.cpp
--- a/Engine/Source/Runtime/IOS/IOSAudio/Private/IOSAudioSource.cpp
+++ b/Engine/Source/Runtime/IOS/IOSAudio/Private/IOSAudioSource.cpp
@@ -168,7 +168,7 @@ bool FIOSAudioSoundSource::Init(FWaveInstance* InWaveInstance)
 			const AudioUnitParameterValue AzimuthRangeScale = 90.f;
 			Pan = (-1.0f + (Channel * 2.0f)) * AzimuthRangeScale;
 		}
-		else if (!WaveInstance->bUseSpatialization)
+		else if (!WaveInstance->GetUseSpatialization())
 		{
 			Pan = 0.0f;
 		}
@@ -229,7 +229,7 @@ void FIOSAudioSoundSource::Update(void)
     SourceLPFFrequency = LPFFrequency / (((float) SampleRate) * 0.5f);
     
 	// Factor in the xaudio2 attenuation that happens to stereo assets.
-	if (WaveInstance->WaveData->NumChannels == 2 && WaveInstance->bUseSpatialization)
+	if (WaveInstance->WaveData->NumChannels == 2 && WaveInstance->GetUseSpatialization())
 	{
 		Volume *= 0.5f;
 	}
@@ -245,7 +245,7 @@ void FIOSAudioSoundSource::Update(void)
 	OSStatus Status = noErr;
 
 	// We only adjust panning on playback for mono sounds that want spatialization
-	if (IOSBuffer->NumChannels == 1 && WaveInstance->bUseSpatialization)
+	if (IOSBuffer->NumChannels == 1 && WaveInstance->GetUseSpatialization())
 	{
 		// Compute the directional offset
 		FVector Offset = GetSpatializationParams().EmitterPosition;
diff --git a/Engine/Source/Runtime/Mac/CoreAudio/Private/CoreAudioSource.cpp b/Engine/Source/Runtime/Mac/CoreAudio/Private/CoreAudioSource.cpp
--- a/Engine/Source/Runtime/Mac/CoreAudio/Private/CoreAudioSource.cpp
+++ b/Engine/Source/Runtime/Mac/CoreAudio/Private/CoreAudioSource.cpp
@@ -360,7 +360,7 @@ void FCoreAudioSoundSource::Update( void )
 		{
 			Volume = WaveInstance->RadioFilterVolume;
 		}
-		else if( WaveInstance->bUseSpatialization )
+		else if( WaveInstance->GetUseSpatialization() )
 		{
 			FVector Direction = AudioDevice->InverseTransform.TransformPosition(WaveInstance->Location).GetSafeNormal();
 
diff --git a/Engine/Source/Runtime/Windows/XAudio2/Private/XAudio2Source.cpp b/Engine/Source/Runtime/Windows/XAudio2/Private/XAudio2Source.cpp
--- a/Engine/Source/Runtime/Windows/XAudio2/Private/XAudio2Source.cpp
+++ b/Engine/Source/Runtime/Windows/XAudio2/Private/XAudio2Source.cpp
@@ -814,7 +814,7 @@ void FXAudio2SoundSource::GetMonoChannelVolumes(float ChannelVolumes[CHANNEL_MAT
 void FXAudio2SoundSource::GetStereoChannelVolumes(float ChannelVolumes[CHANNEL_MATRIX_COUNT], float AttenuatedVolume)
 {
 	// If we're doing 3d spatializaton of stereo sounds
-	if (!IsUsingHrtfSpatializer() && WaveInstance->bUseSpatialization)
+	if (!IsUsingHrtfSpatializer() && WaveInstance->GetUseSpatialization())
 	{
 		check(MAX_INPUT_CHANNELS_SPATIALIZED >= 2);
 
@@ -982,7 +982,7 @@ void FXAudio2SoundSource::RouteStereoToDry(float Chans[CHANNEL_MATRIX_COUNT])
 		// Stereo sounds map 2 channels to 6 speakers
 		AudioDevice->ValidateAPICall(TEXT("SetOutputMatrix (stereo)"), Source->SetOutputMatrix(Destinations[DEST_DRY].pOutputVoice, 2, SPEAKER_COUNT, SpatialisationMatrix));		// Build a non-3d "multi-channel" blend from the stereo channels
 	}
-	else if (WaveInstance->bUseSpatialization)
+	else if (WaveInstance->GetUseSpatialization())
 	{
 		// Build a non-3d "multi-channel" blend from the stereo channels
 		float SpatialisationMatrix[SPEAKER_COUNT * 2] =
@@ -1292,7 +1292,7 @@ FString FXAudio2SoundSource::Describe(bool bUseLongName)
 FString FXAudio2SoundSource::Describe_Internal(bool bUseLongName, bool bIncludeChannelVolumes)
 {
 	FString SpatializedVolumeInfo;
-	if (bIncludeChannelVolumes && WaveInstance->bUseSpatialization)
+	if (bIncludeChannelVolumes && WaveInstance->GetUseSpatialization())
 	{
 		float ChannelVolumes[CHANNEL_MATRIX_COUNT] = { 0.0f };
 		GetChannelVolumes( ChannelVolumes, WaveInstance->GetActualVolume() );
