Fixed surround panning center channel issue

diff --git a/Engine/Source/Runtime/AudioMixer/Private/AudioMixerSource.cpp b/Engine/Source/Runtime/AudioMixer/Private/AudioMixerSource.cpp
--- a/Engine/Source/Runtime/AudioMixer/Private/AudioMixerSource.cpp
+++ b/Engine/Source/Runtime/AudioMixer/Private/AudioMixerSource.cpp
@@ -16,11 +16,20 @@ FAutoConsoleVariableRef CVarDisableHRTF(
 	TEXT("0: Not Disabled, 1: Disabled"),
 	ECVF_Default);
 
+static int32 UseListenerOverrideForSpreadCVar = 1;
+FAutoConsoleVariableRef CVarUseListenerOverrideForSpread(
+	TEXT("au.UseListenerOverrideForSpread"),
+	UseListenerOverrideForSpreadCVar,
+	TEXT("Zero attenuation override distance stereo panning\n")
+	TEXT("0: Use actual distance, 1: use listener override"),
+	ECVF_Default);
+
+
 namespace Audio
 {
 	FMixerSource::FMixerSource(FAudioDevice* InAudioDevice)
 		: FSoundSource(InAudioDevice)
-		, MixerDevice((FMixerDevice*)InAudioDevice)
+		, MixerDevice(static_cast<FMixerDevice*>(InAudioDevice))
 		, MixerBuffer(nullptr)
 		, MixerSourceBuffer(nullptr)
 		, MixerSourceVoice(nullptr)		
@@ -982,22 +991,28 @@ namespace Audio
 			if (!UseObjectBasedSpatialization())
 			{
 				float AzimuthOffset = 0.0f;
-				if (WaveInstance->ListenerToSoundDistance > 0.0f)
+
+				float LeftAzimuth = 90.0f;
+				float RightAzimuth = 270.0f;
+
+				const float DistanceToUse = UseListenerOverrideForSpreadCVar ? WaveInstance->ListenerToSoundDistance : WaveInstance->ListenerToSoundDistanceForPanning;
+
+				if (DistanceToUse > KINDA_SMALL_NUMBER)
 				{
-					AzimuthOffset = FMath::Atan(0.5f * WaveInstance->StereoSpread / WaveInstance->ListenerToSoundDistance);
+					AzimuthOffset = FMath::Atan(0.5f * WaveInstance->StereoSpread / DistanceToUse);
 					AzimuthOffset = FMath::RadiansToDegrees(AzimuthOffset);
-				}
 
-				float LeftAzimuth = WaveInstance->AbsoluteAzimuth - AzimuthOffset;
-				if (LeftAzimuth < 0.0f)
-				{
-					LeftAzimuth += 360.0f;
-				}
+					LeftAzimuth = WaveInstance->AbsoluteAzimuth - AzimuthOffset;
+					if (LeftAzimuth < 0.0f)
+					{
+						LeftAzimuth += 360.0f;
+					}
 
-				float RightAzimuth = WaveInstance->AbsoluteAzimuth + AzimuthOffset;
-				if (RightAzimuth > 360.0f)
-				{
-					RightAzimuth -= 360.0f;
+					RightAzimuth = WaveInstance->AbsoluteAzimuth + AzimuthOffset;
+					if (RightAzimuth > 360.0f)
+					{
+						RightAzimuth -= 360.0f;
+					}
 				}
 
 				// Reset the channel map, the stereo spatialization channel mapping calls below will append their mappings
diff --git a/Engine/Source/Runtime/Engine/Private/ActiveSound.cpp b/Engine/Source/Runtime/Engine/Private/ActiveSound.cpp
--- a/Engine/Source/Runtime/Engine/Private/ActiveSound.cpp
+++ b/Engine/Source/Runtime/Engine/Private/ActiveSound.cpp
@@ -1357,6 +1357,8 @@ void FActiveSound::ApplyAttenuation(FSoundParseParameters& ParseParams, const FL
 
 			ParseParams.ListenerToSoundDistance = ListenerData.ListenerToSoundDistance;
 
+			ParseParams.ListenerToSoundDistanceForPanning = ListenerData.ListenerToSoundDistanceForPanning;
+
 			ParseParams.AbsoluteAzimuth = AbsoluteAzimuth;
 		}
 
diff --git a/Engine/Source/Runtime/Engine/Private/Audio.cpp b/Engine/Source/Runtime/Engine/Private/Audio.cpp
--- a/Engine/Source/Runtime/Engine/Private/Audio.cpp
+++ b/Engine/Source/Runtime/Engine/Private/Audio.cpp
@@ -819,6 +819,7 @@ FWaveInstance::FWaveInstance( FActiveSound* InActiveSound )
 	, StereoSpread(0.0f)
 	, AttenuationDistance(0.0f)
 	, ListenerToSoundDistance(0.0f)
+	, ListenerToSoundDistanceForPanning(0.0f)
 	, AbsoluteAzimuth(0.0f)
 	, PlaybackTime(0.0f)
 	, ReverbSendMethod(EReverbSendMethod::Linear)
diff --git a/Engine/Source/Runtime/Engine/Private/AudioDevice.cpp b/Engine/Source/Runtime/Engine/Private/AudioDevice.cpp
--- a/Engine/Source/Runtime/Engine/Private/AudioDevice.cpp
+++ b/Engine/Source/Runtime/Engine/Private/AudioDevice.cpp
@@ -4669,6 +4669,9 @@ void FAudioDevice::GetAttenuationListenerData(FAttenuationListenerData& OutListe
 		ListenerToSound = SoundTransform.GetTranslation() - ListenerLocation;
 		ListenerToSound.ToDirectionAndLength(OutListenerData.ListenerToSoundDir, OutListenerData.ListenerToSoundDistance);
 
+		// store the actual distance for surround-panning sources with spread (AudioMixer)
+		OutListenerData.ListenerToSoundDistanceForPanning = OutListenerData.ListenerToSoundDistance;
+
 		if (bUseListenerAttenuationOverride)
 		{
 			ListenerToSound = SoundTransform.GetTranslation() - ListenerAttenuationOverride;
diff --git a/Engine/Source/Runtime/Engine/Private/SoundWave.cpp b/Engine/Source/Runtime/Engine/Private/SoundWave.cpp
--- a/Engine/Source/Runtime/Engine/Private/SoundWave.cpp
+++ b/Engine/Source/Runtime/Engine/Private/SoundWave.cpp
@@ -1499,6 +1499,7 @@ void USoundWave::Parse( FAudioDevice* AudioDevice, const UPTRINT NodeWaveInstanc
 		WaveInstance->StereoSpread = ParseParams.StereoSpread;
 		WaveInstance->AttenuationDistance = ParseParams.AttenuationDistance;
 		WaveInstance->ListenerToSoundDistance = ParseParams.ListenerToSoundDistance;
+		WaveInstance->ListenerToSoundDistanceForPanning = ParseParams.ListenerToSoundDistanceForPanning;
 		WaveInstance->AbsoluteAzimuth = ParseParams.AbsoluteAzimuth;
 
 		if (NumChannels <= 2)
diff --git a/Engine/Source/Runtime/Engine/Public/ActiveSound.h b/Engine/Source/Runtime/Engine/Public/ActiveSound.h
--- a/Engine/Source/Runtime/Engine/Public/ActiveSound.h
+++ b/Engine/Source/Runtime/Engine/Public/ActiveSound.h
@@ -79,6 +79,9 @@ struct FSoundParseParameters
 	// The distance from the listener to the sound
 	float ListenerToSoundDistance;
 
+	// The distance from the listener to the sound (ignores attenuation settings)
+	float ListenerToSoundDistanceForPanning;
+
 	// The absolute azimuth angle of the sound relative to the forward listener vector (359 degrees to left, 1 degrees to right)
 	float AbsoluteAzimuth;
 
@@ -167,6 +170,7 @@ struct FSoundParseParameters
 		, OmniRadius(0.0f)
 		, AttenuationDistance(0.0f)
 		, ListenerToSoundDistance(0.0f)
+		, ListenerToSoundDistanceForPanning(0.0f)
 		, AbsoluteAzimuth(0.0f)
 		, SoundSubmix(nullptr)
 		, ReverbSendMethod(EReverbSendMethod::Linear)
diff --git a/Engine/Source/Runtime/Engine/Public/Audio.h b/Engine/Source/Runtime/Engine/Public/Audio.h
--- a/Engine/Source/Runtime/Engine/Public/Audio.h
+++ b/Engine/Source/Runtime/Engine/Public/Audio.h
@@ -426,6 +426,9 @@ public:
 	/** The distance from this wave instance to the closest listener. */
 	float ListenerToSoundDistance;
 
+	/** The distance from this wave instance to the closest listener. (ignoring attenuation override) */
+	float ListenerToSoundDistanceForPanning;
+
 	/** The absolute position of the wave instance relative to forward vector of listener. */
 	float AbsoluteAzimuth; 
 
diff --git a/Engine/Source/Runtime/Engine/Public/AudioDevice.h b/Engine/Source/Runtime/Engine/Public/AudioDevice.h
--- a/Engine/Source/Runtime/Engine/Public/AudioDevice.h
+++ b/Engine/Source/Runtime/Engine/Public/AudioDevice.h
@@ -319,12 +319,18 @@ struct FAttenuationListenerData
 	FTransform ListenerTransform;
 	float AttenuationDistance;
 	float ListenerToSoundDistance;
+	
+	// (AudioMixer only)
+	// Non-attenuation distance for calculating surround sound speaker maps for sources w/ spread
+	float ListenerToSoundDistanceForPanning;
+
 	bool bDataComputed;
 
 	FAttenuationListenerData()
 		: ListenerToSoundDir(FVector::ZeroVector)
 		, AttenuationDistance(0.0f)
 		, ListenerToSoundDistance(0.0f)
+		, ListenerToSoundDistanceForPanning(0.0f)
 		, bDataComputed(false)
 	{}
 };
