Fix Robomerge from CL 5836157:

diff --git a/Engine/Source/Editor/AudioEditor/Private/Factories/SoundFactory.cpp b/Engine/Source/Editor/AudioEditor/Private/Factories/SoundFactory.cpp
--- a/Engine/Source/Editor/AudioEditor/Private/Factories/SoundFactory.cpp
+++ b/Engine/Source/Editor/AudioEditor/Private/Factories/SoundFactory.cpp
@@ -135,11 +135,13 @@ UObject* USoundFactory::FactoryCreateBinary
 
 	UObject* SoundObject = nullptr;
 
+	// First, see if we support this file type in-engine:
 	if (FCString::Stricmp(FileType, TEXT("WAV")) == 0)
 	{
 		SoundObject = CreateObject(Class, InParent, Name, Flags, Context, FileType, Buffer, BufferEnd, Warn);
 	}
 
+	// If we do not, we can use LibSoundFile here to attempt to convert the file to a 16 bit wave file.
 #if WITH_SNDFILE_IO
 	if (!SoundObject)
 	{
@@ -157,7 +159,7 @@ UObject* USoundFactory::FactoryCreateBinary
 
 			// Perpetuate the setting of the suppression flag to avoid
 			// user notification if we attempt to call CreateObject twice
-			SoundObject = CreateObject(Class, InParent, Name, Flags, Context, TEXT("WAV"), Ptr, Ptr + RawWaveData.Num() - 1, Warn);
+			SoundObject = CreateObject(Class, InParent, Name, Flags, Context, TEXT("WAV"), Ptr, Ptr + RawWaveData.Num(), Warn);
 		}
 	}
 #endif // WITH_SNDFILE_IO
@@ -309,14 +311,15 @@ UObject* USoundFactory::CreateObject
 				return nullptr;
 			}
 
+			// If we are not using libSoundFile, we cannot support non-16 bit WAV files.
 			if (*WaveInfo.pBitsPerSample != 16)
 			{
 #if !WITH_SNDFILE_IO
 				WaveInfo.ReportImportFailure();
 				Warn->Logf(ELogVerbosity::Error, TEXT("Only 16 bit WAV source files are supported (%s) on this editor platform."), *Name.ToString());
-#endif // WITH_SNDFILE_IO
-
 				GEditor->GetEditorSubsystem<UImportSubsystem>()->BroadcastAssetPostImport(this, nullptr);
+#endif // WITH_SNDFILE_IO
+				
 				return nullptr;
 			}
 		}
@@ -356,7 +359,9 @@ UObject* USoundFactory::CreateObject
 		int32 ChannelCount = (int32)*WaveInfo.pChannels;
 		check(ChannelCount >0);
 
-		int32 NumSamples = WaveInfo.SampleDataSize / sizeof(int16);
+		int32 SizeOfSample = (*WaveInfo.pBitsPerSample) / 8;
+
+		int32 NumSamples = WaveInfo.SampleDataSize / SizeOfSample;
 		int32 NumFrames = NumSamples / ChannelCount;
 
 		if (ChannelCount > 2)
diff --git a/Engine/Source/Runtime/AudioMixer/Private/SoundFileIO/SoundFileIO.cpp b/Engine/Source/Runtime/AudioMixer/Private/SoundFileIO/SoundFileIO.cpp
--- a/Engine/Source/Runtime/AudioMixer/Private/SoundFileIO/SoundFileIO.cpp
+++ b/Engine/Source/Runtime/AudioMixer/Private/SoundFileIO/SoundFileIO.cpp
@@ -106,12 +106,13 @@ namespace Audio
 		Error = InputSoundDataReader->ReadSamples(ProcessBuffer.GetData(), ProcessBufferSamples, InputSamplesRead);
 		check(Error == ESoundFileError::Type::NONE);
 
-		while (InputSamplesRead != 0)
+		SoundFileCount SamplesWritten = 0;
+
+		while (InputSamplesRead == ProcessBuffer.Num())
 		{
-			SoundFileCount SamplesWritten;
-			Error = SoundFileWriter->WriteSamples((const float*)ProcessBuffer.GetData(), ProcessBuffer.Num(), SamplesWritten);
+			Error = SoundFileWriter->WriteSamples((const float*)ProcessBuffer.GetData(), InputSamplesRead, SamplesWritten);
 			check(Error == ESoundFileError::Type::NONE);
-			check(SamplesWritten == ProcessBuffer.Num());
+			check(SamplesWritten == InputSamplesRead);
 
 			// read more samples
 			Error = InputSoundDataReader->ReadSamples(ProcessBuffer.GetData(), ProcessBufferSamples, InputSamplesRead);
@@ -127,6 +128,10 @@ namespace Audio
 			}
 		}
 
+		// Write final block of samples
+		Error = SoundFileWriter->WriteSamples((const float*)ProcessBuffer.GetData(), InputSamplesRead, SamplesWritten);
+		check(Error == ESoundFileError::Type::NONE);
+
 		// Release the sound file handles as soon as we finished converting the file
 		InputSoundDataReader->Release();
 		SoundFileWriter->Release();
@@ -141,5 +146,3 @@ namespace Audio
 		return true;
 	}
 }
-
-
diff --git a/Engine/Source/Runtime/AudioMixer/Private/SoundFileIO/SoundFileIOManagerImpl.cpp b/Engine/Source/Runtime/AudioMixer/Private/SoundFileIO/SoundFileIOManagerImpl.cpp
--- a/Engine/Source/Runtime/AudioMixer/Private/SoundFileIO/SoundFileIOManagerImpl.cpp
+++ b/Engine/Source/Runtime/AudioMixer/Private/SoundFileIO/SoundFileIOManagerImpl.cpp
@@ -1150,8 +1150,7 @@ namespace Audio
 		{
 			int32 DataSize = BulkData.Num();
 
-			SoundFileCount MaxBytes = DataSize;
-			if (MaxBytes == 0)
+			if (DataSize == 0)
 			{
 				OutOffset = 0;
 				CurrentIndexBytes = 0;
@@ -1169,7 +1168,7 @@ namespace Audio
 				break;
 
 			case ESoundFileSeekMode::FROM_END:
-				CurrentIndexBytes = MaxBytes + Offset;
+				CurrentIndexBytes = DataSize + Offset;
 				break;
 
 			default:
@@ -1177,17 +1176,6 @@ namespace Audio
 				break;
 			}
 
-			// Wrap the byte index to fall between 0 and MaxBytes
-			while (CurrentIndexBytes < 0)
-			{
-				CurrentIndexBytes += MaxBytes;
-			}
-
-			while (CurrentIndexBytes > MaxBytes)
-			{
-				CurrentIndexBytes -= MaxBytes;
-			}
-
 			OutOffset = CurrentIndexBytes;
 			return ESoundFileError::Type::NONE;
 		}
@@ -1203,40 +1191,23 @@ namespace Audio
 		{
 			const uint8* InDataBytes = (const uint8*)DataPtr;
 
-			SoundFileCount MaxBytes = BulkData.Num();
+			SoundFileCount BulkDataLength = BulkData.Num();
 
-			// Append the data to the buffer
-			if (CurrentIndexBytes == MaxBytes)
+			// If we need more room, we add it here.
+			int64 NumExtraBytesNeeded = (CurrentIndexBytes + NumBytes) - BulkDataLength;
+			if (NumExtraBytesNeeded > 0)
 			{
-				BulkData.Append(InDataBytes, NumBytes);
-				CurrentIndexBytes += NumBytes;
+				BulkData.AddUninitialized(NumExtraBytesNeeded);
 			}
-			else if ((CurrentIndexBytes + NumBytes) < MaxBytes)
-			{
-				// Write over the existing data in the buffer
-				for (SoundFileCount i = 0; i < NumBytes; ++i)
-				{
-					BulkData[CurrentIndexBytes++] = InDataBytes[i];
-				}
-			}
-			else
-			{
-				// Overwrite some data until end of current buffer size
-				SoundFileCount RemainingBytes = MaxBytes - CurrentIndexBytes;
-				SoundFileCount InputBufferIndex = 0;
-				for (InputBufferIndex = 0; InputBufferIndex < RemainingBytes; ++InputBufferIndex)
-				{
-					BulkData[CurrentIndexBytes++] = InDataBytes[InputBufferIndex];
-				}
 
-				// Now append the remainder of the input buffer to the internal data buffer
-				const uint8* RemainingData = &InDataBytes[InputBufferIndex];
-				RemainingBytes = NumBytes - RemainingBytes;
-				BulkData.Append(RemainingData, RemainingBytes);
-				CurrentIndexBytes += RemainingBytes;
-				check(CurrentIndexBytes == BulkData.Num());
-			}
+			// Copy the input data into our current place in the BulkData.
+			uint8* BulkDataPtr = BulkData.GetData();
+			FMemory::Memcpy(&BulkDataPtr[CurrentIndexBytes], InDataBytes, NumBytes);
+
+			// Seek our cursor forward accordingly.
+			CurrentIndexBytes += NumBytes;
 			OutNumBytesWritten = NumBytes;
+
 			return ESoundFileError::Type::NONE;
 		}
 
@@ -1376,7 +1347,8 @@ namespace Audio
 
 			if (FileHandle)
 			{
-				SoundFileClose(FileHandle);
+				int32 Result = SoundFileClose(FileHandle);
+				check(Result == 0);
 				FileHandle = nullptr;
 			}
 			return ESoundFileError::Type::NONE;
