[NO NON-PII COMMIT MESSAGE]

diff --git a/Engine/Source/Runtime/AudioMixer/AudioMixer.Build.cs b/Engine/Source/Runtime/AudioMixer/AudioMixer.Build.cs
--- a/Engine/Source/Runtime/AudioMixer/AudioMixer.Build.cs
+++ b/Engine/Source/Runtime/AudioMixer/AudioMixer.Build.cs
@@ -87,16 +87,19 @@ namespace UnrealBuildTool.Rules
 
 			if (Target.IsInPlatformGroup(UnrealPlatformGroup.Windows))
 			{
-				string LibSndFilePath = Target.UEThirdPartyBinariesDirectory + "libsndfile/";
-				LibSndFilePath += Target.Platform == UnrealTargetPlatform.Win32
-					? "Win32"
-					: "Win64";
-					
-				PublicAdditionalLibraries.Add("libsndfile-1.lib");
+				string PlatformName = Target.Platform == UnrealTargetPlatform.Win32 ? "Win32" : "Win64";
+
+                string LibSndFilePath = Target.UEThirdPartyBinariesDirectory + "libsndfile/";
+                LibSndFilePath += PlatformName;
+
+
+                PublicAdditionalLibraries.Add("libsndfile-1.lib");
 				PublicDelayLoadDLLs.Add("libsndfile-1.dll");
 				PublicIncludePathModuleNames.Add("UELibSampleRate");
 				PublicLibraryPaths.Add(LibSndFilePath);
-			}
+
+                RuntimeDependencies.Add("$(EngineDir)/Binaries/ThirdParty/libsndfile/" + PlatformName + "/libsndfile-1.dll");
+            }
 		}
 	}
 }
