Fix for hang in shutdown on android apps

diff --git a/Engine/Source/Runtime/Android/AudioMixerAndroid/Private/AudioMixerPlatformAndroid.cpp b/Engine/Source/Runtime/Android/AudioMixerAndroid/Private/AudioMixerPlatformAndroid.cpp
--- a/Engine/Source/Runtime/Android/AudioMixerAndroid/Private/AudioMixerPlatformAndroid.cpp
+++ b/Engine/Source/Runtime/Android/AudioMixerAndroid/Private/AudioMixerPlatformAndroid.cpp
@@ -337,27 +337,35 @@ namespace Audio
 		return PlatformSettings;
 	}
 
-	static bool bSuspended = false;
-
 	void FMixerPlatformAndroid::SuspendContext()
 	{
+		FScopeLock ScopeLock(&SuspendedCriticalSection);
+
 		if (!bSuspended)
 		{
-			bSuspended = true;
+			UE_LOG(LogAudioMixerAndroid, Display, TEXT("Suspending android audio renderer"));
+
 			// set the player's state to paused
 			SLresult result = (*SL_PlayerPlayInterface)->SetPlayState(SL_PlayerPlayInterface, SL_PLAYSTATE_PAUSED);
 			check(SL_RESULT_SUCCESS == result);
+
+			bSuspended = true;
 		}
 	}
 
 	void FMixerPlatformAndroid::ResumeContext()
 	{
+		FScopeLock ScopeLock(&SuspendedCriticalSection);
+
 		// set the player's state to paused
 		if (bSuspended)
 		{
-			bSuspended = false;
+			UE_LOG(LogAudioMixerAndroid, Display, TEXT("Resuming android audio renderer"));
+
 			SLresult result = (*SL_PlayerPlayInterface)->SetPlayState(SL_PlayerPlayInterface, SL_PLAYSTATE_PLAYING);
 			check(SL_RESULT_SUCCESS == result);
+
+			bSuspended = false;
 		}
 	}
 
diff --git a/Engine/Source/Runtime/Android/AudioMixerAndroid/Private/AudioMixerPlatformAndroid.h b/Engine/Source/Runtime/Android/AudioMixerAndroid/Private/AudioMixerPlatformAndroid.h
--- a/Engine/Source/Runtime/Android/AudioMixerAndroid/Private/AudioMixerPlatformAndroid.h
+++ b/Engine/Source/Runtime/Android/AudioMixerAndroid/Private/AudioMixerPlatformAndroid.h
@@ -53,6 +53,8 @@ namespace Audio
 		SLPlayItf SL_PlayerPlayInterface;
 		SLAndroidSimpleBufferQueueItf SL_PlayerBufferQueue;
 
+		FCriticalSection SuspendedCriticalSection;
+
 		bool bSuspended;
 		bool bInitialized;
 		bool bInCallback;
diff --git a/Engine/Source/Runtime/AudioMixer/Private/AudioMixerSourceManager.cpp b/Engine/Source/Runtime/AudioMixer/Private/AudioMixerSourceManager.cpp
--- a/Engine/Source/Runtime/AudioMixer/Private/AudioMixerSourceManager.cpp
+++ b/Engine/Source/Runtime/AudioMixer/Private/AudioMixerSourceManager.cpp
@@ -111,6 +111,9 @@ namespace Audio
 		const bool bIsManualReset = true;
 		CommandsProcessedEvent = FPlatformProcess::GetSynchEventFromPool(bIsManualReset);
 		check(CommandsProcessedEvent != nullptr);
+
+		// Immediately trigger the command processed in case a flush happens before the audio thread swaps command buffers
+		CommandsProcessedEvent->Trigger();
 	}
 
 	FMixerSourceManager::~FMixerSourceManager()
@@ -2412,6 +2415,7 @@ namespace Audio
 		// Add the function to the command queue
 		int32 AudioThreadCommandIndex = AudioThreadCommandBufferIndex.GetValue();
 		CommandBuffers[AudioThreadCommandIndex].SourceCommandQueue.Add(MoveTemp(InFunction));
+		NumCommands.Increment();
 	}
 
 	void FMixerSourceManager::PumpCommandQueue()
@@ -2434,6 +2438,7 @@ namespace Audio
 		{
 			TFunction<void()>& CommandFunction = Commands.SourceCommandQueue[Id];
 			CommandFunction();
+			NumCommands.Decrement();
 		}
 
 		Commands.SourceCommandQueue.Reset();
@@ -2454,8 +2459,23 @@ namespace Audio
 	{
 		check(CommandsProcessedEvent != nullptr);
 
+		// If we have no commands enqueued, exit
+		if (NumCommands.GetValue() == 0)
+		{
+			UE_LOG(LogAudioMixer, Display, TEXT("No commands were queued while flushing the source manager."));
+			return;
+		}
+
 		// Make sure current current executing 
-		CommandsProcessedEvent->Wait();
+		if (!CommandsProcessedEvent->Wait(1000))
+		{
+			CommandsProcessedEvent->Trigger();
+			UE_LOG(LogAudioMixer, Warning, TEXT("Timed out waiting to flush the source manager command queue (1)."));
+		}
+		else
+		{
+			UE_LOG(LogAudioMixer, Display, TEXT("Flush succeeded in the source manager command queue (1)."));
+		}
 
 		// Call update to trigger a final pump of commands
 		Update();
@@ -2466,7 +2486,15 @@ namespace Audio
 		}
 
 		// Wait one more time for the double pump
-		CommandsProcessedEvent->Wait();
+		if (!CommandsProcessedEvent->Wait(1000))
+		{
+			CommandsProcessedEvent->Trigger();
+			UE_LOG(LogAudioMixer, Warning, TEXT("Timed out waiting to flush the source manager command queue (2)."));
+		}
+		else
+		{
+			UE_LOG(LogAudioMixer, Display, TEXT("Flush succeeded the source manager command queue (2)."));
+		}
 	}
 
 	void FMixerSourceManager::UpdatePendingReleaseData(bool bForceWait)
diff --git a/Engine/Source/Runtime/AudioMixer/Private/AudioMixerSourceManager.h b/Engine/Source/Runtime/AudioMixer/Private/AudioMixerSourceManager.h
--- a/Engine/Source/Runtime/AudioMixer/Private/AudioMixerSourceManager.h
+++ b/Engine/Source/Runtime/AudioMixer/Private/AudioMixerSourceManager.h
@@ -548,6 +548,9 @@ namespace Audio
 		int32 NumOutputSamples;
 		int32 NumSourceWorkers;
 
+		// Commands queued up to execute
+		FThreadSafeCounter NumCommands;
+
 		uint8 bInitialized : 1;
 		uint8 bUsingSpatializationPlugin : 1;
 
