Add static flag to make sure we don't double trigger hardware allocations on audio device swap on consoles.

diff --git a/Engine/Source/Runtime/AudioMixer/Private/AudioMixerDevice.cpp b/Engine/Source/Runtime/AudioMixer/Private/AudioMixerDevice.cpp
--- a/Engine/Source/Runtime/AudioMixer/Private/AudioMixerDevice.cpp
+++ b/Engine/Source/Runtime/AudioMixer/Private/AudioMixerDevice.cpp
@@ -267,6 +267,8 @@ namespace Audio
 
 	void FMixerDevice::TeardownHardware()
 	{
+		AUDIO_MIXER_CHECK_GAME_THREAD(this);
+
 		// Make sure all submixes are registered but not initialized
 		for (TObjectIterator<USoundSubmix> It; It; ++It)
 		{
