Merging fix for crash in SoundWaveDecoder from Dev-Audio.

diff --git a/Engine/Source/Runtime/AudioMixer/Private/AudioMixerBuffer.cpp b/Engine/Source/Runtime/AudioMixer/Private/AudioMixerBuffer.cpp
--- a/Engine/Source/Runtime/AudioMixer/Private/AudioMixerBuffer.cpp
+++ b/Engine/Source/Runtime/AudioMixer/Private/AudioMixerBuffer.cpp
@@ -25,7 +25,7 @@ namespace Audio
 		// Set the base-class NumChannels to wave's NumChannels
 		NumChannels = InWave->NumChannels;
 
-		if (InWave->DecompressionType == EDecompressionType::DTYPE_Native || InWave->DecompressionType == EDecompressionType::DTYPE_Preview)
+		if (InBufferType != EBufferType::PCMRealTime && (InWave->DecompressionType == EDecompressionType::DTYPE_Native || InWave->DecompressionType == EDecompressionType::DTYPE_Preview))
 		{
 			check(!InWave->RawPCMData || InWave->RawPCMDataSize);
 			Data = InWave->RawPCMData;
diff --git a/Engine/Source/Runtime/AudioMixer/Private/AudioMixerSourceBuffer.cpp b/Engine/Source/Runtime/AudioMixer/Private/AudioMixerSourceBuffer.cpp
--- a/Engine/Source/Runtime/AudioMixer/Private/AudioMixerSourceBuffer.cpp
+++ b/Engine/Source/Runtime/AudioMixer/Private/AudioMixerSourceBuffer.cpp
@@ -313,9 +313,6 @@ namespace Audio
 		}
 		else if (!bIsBus)
 		{
-			// We should have already kicked off and finished a task. 
-			check(AsyncRealtimeAudioTask != nullptr);
-
 			ProcessRealTimeSource();
 		}
 	}
diff --git a/Engine/Source/Runtime/AudioMixer/Private/DSP/SoundWaveDecoder.cpp b/Engine/Source/Runtime/AudioMixer/Private/DSP/SoundWaveDecoder.cpp
--- a/Engine/Source/Runtime/AudioMixer/Private/DSP/SoundWaveDecoder.cpp
+++ b/Engine/Source/Runtime/AudioMixer/Private/DSP/SoundWaveDecoder.cpp
@@ -18,7 +18,7 @@ namespace Audio
 
 		SourceInfo.PitchScale = InitData.PitchScale;
 
-		MixerBuffer = FMixerBuffer::Init(AudioDevice, InitData.SoundWave, InitData.SeekTime > 0.0f);
+		MixerBuffer = FMixerBuffer::Init(AudioDevice, InitData.SoundWave, true);
 	}
 
 	FDecodingSoundSource::~FDecodingSoundSource()
@@ -84,6 +84,15 @@ namespace Audio
 	{
 		if (MixerBuffer->GetNumChannels() > 0 && MixerBuffer->GetNumChannels() <= 2)
 		{
+			// Pass the decompression state off to the mixer source buffer if it hasn't already done so
+			ICompressedAudioInfo* Decoder = MixerBuffer->GetDecompressionState(false);
+			MixerSourceBuffer.SetDecoder(Decoder);
+			
+			if (!MixerSourceBuffer.IsAsyncTaskInProgress())
+			{
+				MixerSourceBuffer.ReadMoreRealtimeData(Decoder, 0, EBufferReadMode::Asynchronous);
+			}
+
 			SourceInfo.NumSourceChannels = MixerBuffer->GetNumChannels();
 			SourceInfo.TotalNumFrames = MixerBuffer->GetNumFrames();
 
@@ -96,6 +105,7 @@ namespace Audio
 			SourceInfo.PitchParam.SetValue(SourceInfo.BasePitchScale * SourceInfo.PitchScale);
 
 			MixerSourceBuffer.Init();
+			MixerSourceBuffer.OnBeginGenerate();
 
 			bInitialized = true;
 		}
@@ -162,7 +172,6 @@ namespace Audio
 			}
 			else
 			{
-				SourceInfo.bIsLastBuffer = true;
 				return;
 			}
 
@@ -521,7 +530,7 @@ namespace Audio
 	bool FSoundSourceDecoder::IsFinished(const FDecodingSoundSourceHandle& InHandle) const
 	{
 		const FDecodingSoundSourcePtr* DecodingSoundWaveDataPtr = DecodingSources.Find(InHandle.Id);
-		if (!DecodingSoundWaveDataPtr)
+		if (!DecodingSoundWaveDataPtr || !DecodingSoundWaveDataPtr->IsValid())
 		{
 			return true;
 		}
