[NO NON-PII COMMIT MESSAGE]

diff --git a/Engine/Source/Runtime/AudioMixer/Classes/SubmixEffects/AudioMixerSubmixEffectReverbFast.h b/Engine/Source/Runtime/AudioMixer/Classes/SubmixEffects/AudioMixerSubmixEffectReverbFast.h
new file mode 100644
--- /dev/null
+++ b/Engine/Source/Runtime/AudioMixer/Classes/SubmixEffects/AudioMixerSubmixEffectReverbFast.h
@@ -0,0 +1,135 @@
+// Copyright 1998-2019 Epic Games, Inc. All Rights Reserved.
+
+#pragma once
+
+#include "Sound/SoundEffectSubmix.h"
+#include "DSP/ReverbFast.h"
+#include "AudioEffect.h"
+#include "AudioMixerSubmixEffectReverbFast.generated.h"
+
+USTRUCT(BlueprintType)
+struct AUDIOMIXER_API FSubmixEffectReverbFastSettings
+{
+	GENERATED_USTRUCT_BODY()
+
+	/** Density - 0.0 < 0.85 < 1.0 - Coloration of the late reverb - lower value is more grainy */
+	UPROPERTY(EditAnywhere, BlueprintReadWrite, Category = ReverbParameters, meta = (ClampMin = "0.0", ClampMax = "1.0"))
+	float Density;
+
+	/** Diffusion - 0.0 < 0.85 < 1.0 - Echo density in the reverberation decay - lower is more grainy */
+	UPROPERTY(EditAnywhere, BlueprintReadWrite, Category = ReverbParameters, meta = (ClampMin = "0.0", ClampMax = "1.0"))
+	float Diffusion;
+
+	/** Reverb Gain - 0.0 < 0.32 < 1.0 - overall reverb gain - master volume control */
+	UPROPERTY(EditAnywhere, BlueprintReadWrite, Category = ReverbParameters, meta = (ClampMin = "0.0", ClampMax = "1.0"))
+	float Gain;
+
+	/** Reverb Gain High Frequency - 0.0 < 0.89 < 1.0 - attenuates the high frequency reflected sound */
+	UPROPERTY(EditAnywhere, BlueprintReadWrite, Category = ReverbParameters, meta = (ClampMin = "0.0", ClampMax = "1.0"))
+	float GainHF;
+
+	/** Decay Time - 0.1 < 1.49 < 20.0 Seconds - larger is more reverb */
+	UPROPERTY(EditAnywhere, BlueprintReadWrite, Category = ReverbParameters, meta = (ClampMin = "0.1", ClampMax = "20.0"))
+	float DecayTime;
+
+	/** Decay High Frequency Ratio - 0.1 < 0.83 < 2.0 - how much the quicker or slower the high frequencies decay relative to the lower frequencies. */
+	UPROPERTY(EditAnywhere, BlueprintReadWrite, Category = ReverbParameters, meta = (ClampMin = "0.1", ClampMax = "2.0"))
+	float DecayHFRatio;
+
+	/** Reflections Gain - 0.0 < 0.05 < 3.16 - controls the amount of initial reflections */
+	UPROPERTY(EditAnywhere, BlueprintReadWrite, Category = ReverbParameters, meta = (ClampMin = "0.0", ClampMax = "3.16"))
+	float ReflectionsGain;
+
+	/** Reflections Delay - 0.0 < 0.007 < 0.3 Seconds - the time between the listener receiving the direct path sound and the first reflection */
+	UPROPERTY(EditAnywhere, BlueprintReadWrite, Category = ReverbParameters, meta = (ClampMin = "0.0", ClampMax = "0.3"))
+	float ReflectionsDelay;
+
+	/** Late Reverb Gain - 0.0 < 1.26 < 10.0 - gain of the late reverb */
+	UPROPERTY(EditAnywhere, BlueprintReadWrite, Category = ReverbParameters, meta = (ClampMin = "0.0", ClampMax = "10.0"))
+	float LateGain;
+
+	/** Late Reverb Delay - 0.0 < 0.011 < 0.1 Seconds - time difference between late reverb and first reflections */
+	UPROPERTY(EditAnywhere, BlueprintReadWrite, Category = ReverbParameters, meta = (ClampMin = "0.0", ClampMax = "0.1"))
+	float LateDelay;
+
+	/** Air Absorption - 0.0 < 0.994 < 1.0 - lower value means more absorption */
+	UPROPERTY(EditAnywhere, BlueprintReadWrite, Category = ReverbParameters, meta = (ClampMin = "0.0", ClampMax = "1.0"))
+	float AirAbsorptionGainHF;
+
+	// Overall wetlevel of the reverb effect
+	UPROPERTY(EditAnywhere, BlueprintReadWrite, Category = Reverb)
+	float WetLevel;
+
+	// Overall drylevel of the reverb effect
+	UPROPERTY(EditAnywhere, BlueprintReadWrite, Category = Reverb)
+	float DryLevel;
+
+	FSubmixEffectReverbFastSettings()
+	: Density(0.85f)
+	, Diffusion(0.85f)
+	, Gain(0.0f)
+	, GainHF(0.89f)
+	, DecayTime(1.49f)
+	, DecayHFRatio(0.83f)
+	, ReflectionsGain(0.05f)
+	, ReflectionsDelay(0.007f)
+	, LateGain(1.26f)
+	, LateDelay(0.1f)
+	, AirAbsorptionGainHF(0.994f)
+	, WetLevel(0.3f)
+	, DryLevel(0.0f)
+	{
+	}
+};
+
+class AUDIOMIXER_API FSubmixEffectReverbFast : public FSoundEffectSubmix
+{
+public:
+	FSubmixEffectReverbFast();
+
+	// Called on an audio effect at initialization on main thread before audio processing begins.
+	virtual void Init(const FSoundEffectSubmixInitData& InSampleRate) override;
+	
+	// Called when an audio effect preset is changed
+	virtual void OnPresetChanged() override;
+
+	// We want to receive downmixed submix audio to stereo input for the reverb effect
+	virtual uint32 GetDesiredInputChannelCountOverride() const override { return 2; }
+
+	// Process the input block of audio. Called on audio thread.
+	virtual void OnProcessAudio(const FSoundEffectSubmixInputData& InData, FSoundEffectSubmixOutputData& OutData) override;
+
+	// Sets the reverb effect parameters based from audio thread code
+	void SetEffectParameters(const FAudioReverbEffect& InReverbEffectParameters);
+
+private:
+	void UpdateParameters();
+
+	// The fast reverb effect
+	TUniquePtr<Audio::FPlateReverbFast> PlateReverb;
+
+	// The reverb effect params
+	Audio::TParams<Audio::FPlateReverbFastSettings> Params;
+
+	// Curve which maps old reverb times to new decay value
+	FRichCurve DecayCurve;
+
+};
+
+UCLASS()
+class AUDIOMIXER_API USubmixEffectReverbFastPreset : public USoundEffectSubmixPreset
+{
+	GENERATED_BODY()
+
+public:
+	EFFECT_PRESET_METHODS(SubmixEffectReverbFast)
+
+	UFUNCTION(BlueprintCallable, Category = "Audio|Effects")
+	void SetSettings(const FSubmixEffectReverbFastSettings& InSettings);
+
+	UFUNCTION(BlueprintCallable, Category = "Audio|Effects")
+	void SetSettingsWithReverbEffect(const UReverbEffect* InReverbEffect, const float WetLevel, const float DryLevel = 0.0f);
+
+	UPROPERTY(EditAnywhere, BlueprintReadWrite, Category = SubmixEffectPreset)
+	FSubmixEffectReverbFastSettings Settings;
+};
diff --git a/Engine/Source/Runtime/AudioMixer/Private/DSP/AlignedBlockBuffer.cpp b/Engine/Source/Runtime/AudioMixer/Private/DSP/AlignedBlockBuffer.cpp
new file mode 100644
--- /dev/null
+++ b/Engine/Source/Runtime/AudioMixer/Private/DSP/AlignedBlockBuffer.cpp
@@ -0,0 +1,324 @@
+#include "DSP/AlignedBlockBuffer.h"
+#include "CoreMinimal.h"
+#include "AudioMixer.h"
+
+using namespace Audio;
+
+FAlignedBlockBuffer::FAlignedBlockBuffer(int32 InSampleCapacity, int32 InMaxNumInspectSamples, uint32 InByteAlignment)
+:	ByteAlignment(InByteAlignment),
+	FloatAlignment(0),
+	InternalBuffer(nullptr),
+	RolloverBuffer(nullptr),
+	ReadPointer(nullptr),
+	WritePointer(nullptr),
+	SampleCapacity(InSampleCapacity),
+	MaxNumInspectSamples(InMaxNumInspectSamples),
+	WriteCount(0),
+	WritePos(0),
+	ReadPos(0)
+{
+	FloatAlignment = ByteAlignment / sizeof(float);
+
+	int32 MaxStorage = TNumericLimits<int32>::Max() / 4;
+
+	// The SampleCapacity needs to also be divisible by the byte alignment as it is
+	// used in a bunch of the pointer math.  This ensures that the ReadPointer
+	// is always pointing to aligned memory
+	while (0 != (SampleCapacity % FloatAlignment))
+	{
+		SampleCapacity++;
+	}
+
+	// Check sanity of variables.
+	if (MaxNumInspectSamples > MaxStorage)
+	{
+		UE_LOG(LogAudioMixer, Warning, TEXT("MaxNumInspectSamples [%d] reduced to maximum storage [%d]"), MaxNumInspectSamples, MaxStorage);
+		MaxNumInspectSamples = MaxStorage;
+	}
+
+	if (SampleCapacity > MaxStorage)
+	{
+		UE_LOG(LogAudioMixer, Warning, TEXT("SampleCapacity [%d] reduced to maximum storage [%d]"), SampleCapacity, MaxStorage);
+		SampleCapacity = MaxStorage;
+	}
+
+	if (MaxNumInspectSamples > SampleCapacity)
+	{
+		UE_LOG(LogAudioMixer, Warning, TEXT("SampleCapacity [%d] increased to Maxmimum Inspection [%d]"), SampleCapacity, MaxNumInspectSamples);
+		SampleCapacity = MaxNumInspectSamples;
+	}
+
+	// Allocate space
+	RolloverBuffer = AllocateAlignedFloatArray(MaxNumInspectSamples);
+	InternalBuffer = AllocateAlignedFloatArray(SampleCapacity);
+	
+	// Set everything to zero
+	ZeroFloatArray(RolloverBuffer, MaxNumInspectSamples);
+	ZeroFloatArray(InternalBuffer, SampleCapacity);
+	
+	// Initialize pointers to internal buffers
+	WritePointer = InternalBuffer;
+	ReadPointer = InternalBuffer;
+}
+
+		
+FAlignedBlockBuffer::~FAlignedBlockBuffer() throw()
+{
+	FreeFloatArray(InternalBuffer);
+	FreeFloatArray(RolloverBuffer);
+}
+
+int32 FAlignedBlockBuffer::GetSampleCapacity( ) const
+{
+	return SampleCapacity;
+}
+		
+int32 FAlignedBlockBuffer::GetNumAvailable() const
+{
+	return WriteCount;
+}
+
+int32 FAlignedBlockBuffer::CheckNumAndIncrementWriteCount(int32 InNum)
+{
+	// Clamp InNum to capacity
+	checkf(InNum <= SampleCapacity, TEXT("Input buffer [%d] larger than capacity [%d]"), InNum, SampleCapacity);
+	InNum = InNum < SampleCapacity ? InNum : SampleCapacity;
+
+	// Increment write count keeping within capacity
+	WriteCount += InNum;
+
+	// Handle buffer overflow
+	checkf(WriteCount <= SampleCapacity, TEXT("Buffer overflow"));
+	if(WriteCount > SampleCapacity)
+	{
+		// Reduce number of input samples on buffer overflow.
+		InNum -= WriteCount - SampleCapacity;
+		WriteCount = SampleCapacity;
+	}
+
+	return InNum;
+}
+
+void FAlignedBlockBuffer::AddZeros(int32 InNum)
+{
+	InNum = CheckNumAndIncrementWriteCount(InNum);
+	
+	// Check if we expect to be wrapping around end of input buffer
+	int32 RolloverState = SampleCapacity - (WritePos + InNum);
+	
+	if(RolloverState > 0)
+	{
+		// No rollover if rollover state is above zero.
+		ZeroFloatArray(WritePointer, InNum);
+		WritePos += InNum;
+		WritePointer += InNum;
+	}
+	else
+	{
+		// Zeroing buffer goes over internal buffer border, so perform zeroing
+		// in two parts.
+		
+		// Zero first part of buffer
+		int32 RemainingBufferSamples = InNum + RolloverState;
+		ZeroFloatArray(WritePointer, RemainingBufferSamples);
+		WritePointer = InternalBuffer;
+		WritePos = -RolloverState;
+		
+		// Zero the second part
+		ZeroFloatArray(WritePointer, -RolloverState);
+		WritePointer -= RolloverState;
+	}
+}
+		
+void FAlignedBlockBuffer::AddSamples(const float* InSamples, int32 InNum)
+{
+	checkf(nullptr != InSamples, TEXT("Null pointer"));
+	if (nullptr == InSamples)
+	{
+		// Quit early with null pointers.
+		return;
+	}
+
+	InNum = CheckNumAndIncrementWriteCount(InNum);
+	
+	// If the write operation is expected to go past the internal buffer boundary
+	// the RolloverState will be negative. If RolloverState is negative, the copy
+	// needs to happen in two parts. 
+	int32 RolloverState = SampleCapacity - (WritePos + InNum);
+	
+	if(RolloverState > 0)
+	{
+		// No rollover if rollover state above zero.
+		CopyFloatArray(WritePointer, InSamples, InNum);
+		
+		WritePos += InNum;
+		WritePointer += InNum;
+	}
+	else
+	{
+		//Copy first part of buffer
+		int32 RemainingBufferSamples = InNum + RolloverState;
+	
+		CopyFloatArray(WritePointer, InSamples, RemainingBufferSamples);
+		
+		WritePointer = InternalBuffer;
+		WritePos = -RolloverState;
+		InSamples += RemainingBufferSamples;
+		
+		//Copy the second part
+		CopyFloatArray(WritePointer, InSamples, -RolloverState);
+
+		WritePointer -= RolloverState;
+	}
+}
+
+void FAlignedBlockBuffer::RemoveSamples(int32 InNum)
+{
+	// In order to keep the read pointer aligned, any operations that move the read pointer
+	// must also be aligned. 
+	checkf(0 == (InNum % FloatAlignment), TEXT("InNum must be divisible by alignement."));
+
+	//decrement the available samples
+	WriteCount -= InNum;
+
+	checkf(WriteCount >= 0, TEXT("Buffer underflow"));
+	if (WriteCount < 0)
+	{
+		// Gracefully handle underflow.
+		WriteCount = 0;
+		ReadPos = 0;
+		ReadPointer = InternalBuffer;
+		return;
+	}
+
+	int32 RolloverState = SampleCapacity - (ReadPos + InNum);
+	
+	if( RolloverState > 0 )
+	{
+		ReadPos += InNum;
+		ReadPointer += InNum;
+	}
+	else
+	{
+		// Wrap the get pointer around in case where buffer boundary is crossed.
+		ReadPos = -RolloverState;
+		ReadPointer = InternalBuffer - RolloverState;
+	}
+}
+		
+const float* FAlignedBlockBuffer::InspectSamples(int32 InNum, int32 InOffset)
+{
+	checkf(0 == (InOffset % FloatAlignment), TEXT("InOffset must divisible by alignment"));
+	checkf(InNum <= MaxNumInspectSamples, TEXT("Requested samples [%d] are more than maximum [%d]"), InNum, MaxNumInspectSamples);
+	checkf((InNum + InOffset) <= WriteCount, TEXT("Buffer underflow"));
+	checkf(InOffset >= 0, TEXT("InOffset must be greater than 0"));
+
+	if (InNum < 1)
+	{
+		return nullptr;
+	}
+
+	if (InNum > MaxNumInspectSamples)
+	{
+		// This is an unfortunate situation because it allocates memory.
+		SetMaxNumInspectSamples(InNum);
+	}
+
+	if ((InNum + InOffset) > WriteCount)
+	{
+		// This is a pretty bad situation. Anything done here is a bad option.
+		// Hopefully the lesser of all the evils is returning a nullptr so that the issue
+		// can be identified quickly. 
+		return nullptr;
+	}
+
+	if (InOffset < 0)
+	{
+		// Invalid offset. Despite the profound implications for the advancement of humankind, we cannot see into the future. 
+		InOffset = 0;
+	}
+	
+	int32 StartPos = ReadPos + InOffset;
+	int32 EndPos = StartPos + InNum;
+
+	while (StartPos > SampleCapacity)
+	{
+		StartPos -= SampleCapacity;
+		EndPos -= SampleCapacity;
+	}
+
+	if (EndPos < SampleCapacity)
+	{
+		return &InternalBuffer[StartPos];
+	}
+	else
+	{
+		int32 NumFirstHalf = SampleCapacity - StartPos;
+		// Copy first half into rolloverbuffer
+		CopyFloatArray(RolloverBuffer, &InternalBuffer[StartPos], NumFirstHalf);
+		// Copy second half into rolloverbuffer
+		CopyFloatArray(&RolloverBuffer[NumFirstHalf], InternalBuffer, InNum - NumFirstHalf);
+
+		return RolloverBuffer;
+	}
+}
+		
+void FAlignedBlockBuffer::ClearSamples()
+{
+	ZeroFloatArray(InternalBuffer, SampleCapacity);
+	ZeroFloatArray(RolloverBuffer, MaxNumInspectSamples);
+
+	ReadPointer = InternalBuffer;
+	WritePointer = InternalBuffer;
+	ReadPos = 0;
+	WritePos = 0;
+	
+	WriteCount = 0;
+}
+				
+int32 FAlignedBlockBuffer::GetMaxNumInspectSamples() const
+{
+	return MaxNumInspectSamples;
+}
+
+void FAlignedBlockBuffer::SetMaxNumInspectSamples(int32 InMax)
+{
+	if(InMax > MaxNumInspectSamples)
+	{
+		FreeFloatArray(RolloverBuffer);
+		RolloverBuffer = AllocateAlignedFloatArray(InMax);
+	}
+	
+	MaxNumInspectSamples = InMax;
+}
+
+float* FAlignedBlockBuffer::AllocateAlignedFloatArray(int32 InNum) const 
+{
+	int32 Size = InNum * sizeof(float);
+	return (float*)FMemory::Malloc(Size, ByteAlignment);
+}
+
+void FAlignedBlockBuffer::FreeFloatArray(float*& Array) const
+{
+	FMemory::Free(Array);
+	Array = nullptr;
+}
+
+void FAlignedBlockBuffer::ZeroFloatArray(float* InArray, int32 InNum) const
+{
+	if (InNum > 0) 
+	{
+		int32 Size = InNum * sizeof(float);
+		FMemory::Memset(InArray, 0, Size);
+	}
+}
+
+void FAlignedBlockBuffer::CopyFloatArray(float* ToArray, const float* FromArray, int32 InNum) const 
+{
+	if (InNum > 0)
+	{
+		int32 Size = InNum * sizeof(float);
+		FMemory::Memcpy(ToArray, FromArray, Size);
+	}
+}
+
diff --git a/Engine/Source/Runtime/AudioMixer/Private/DSP/AllPassFractionalDelay.cpp b/Engine/Source/Runtime/AudioMixer/Private/DSP/AllPassFractionalDelay.cpp
new file mode 100644
--- /dev/null
+++ b/Engine/Source/Runtime/AudioMixer/Private/DSP/AllPassFractionalDelay.cpp
@@ -0,0 +1,165 @@
+
+// Copyright 1998-2019 Epic Games, Inc. All Rights Reserved.
+
+#include "DSP/AllPassFractionalDelay.h"
+#include "DSP/Dsp.h"
+#include "DSP/BufferVectorOperations.h"
+
+using namespace Audio;
+
+FAllPassFractionalDelay::FAllPassFractionalDelay(int32 InMaxDelay, int32 InNumInternalBufferSamples)
+: MaxDelay(InMaxDelay)
+, NumInternalBufferSamples(InNumInternalBufferSamples)
+, Z1(0.0f)
+{
+	checkf(MaxDelay > 0, TEXT("Maximum delay must be greater than zero"));
+	checkf(InNumInternalBufferSamples > 0, TEXT("Internal buffer length must be greater than zero"));
+
+	// Clamp settings to avoid hard crashes during runtime.
+	if (MaxDelay < 1)
+	{
+		MaxDelay = 1;
+	}
+	if (InNumInternalBufferSamples < 1)
+	{
+		InNumInternalBufferSamples = 32;
+	}
+
+	// Allocate and prepare delay line for maximum delay.
+	DelayLine = MakeUnique<FAlignedBlockBuffer>((2 * MaxDelay) + NumInternalBufferSamples, MaxDelay + NumInternalBufferSamples);
+	DelayLine->AddZeros(MaxDelay);
+
+	Coefficients.Reset(NumInternalBufferSamples);
+	FractionalDelays.Reset(NumInternalBufferSamples);
+	IntegerDelays.Reset(NumInternalBufferSamples);
+	IntegerDelayOffsets.Reset(NumInternalBufferSamples);
+
+	Coefficients.AddUninitialized(NumInternalBufferSamples);
+	FractionalDelays.AddUninitialized(NumInternalBufferSamples);
+	IntegerDelays.AddUninitialized(NumInternalBufferSamples);
+	IntegerDelayOffsets.AddUninitialized(NumInternalBufferSamples);
+
+	// Integer delay offsets account for buffer position when doing block processing of data. 
+	for (int32 i = 0; i < InNumInternalBufferSamples; i++)
+	{
+		IntegerDelayOffsets[i] = i + MaxDelay;
+	}
+}
+
+
+// Destructor
+FAllPassFractionalDelay::~FAllPassFractionalDelay()
+{}
+
+// Resets the delay line state, flushes buffer and resets read/write pointers.
+void FAllPassFractionalDelay::Reset() 
+{
+	DelayLine->ClearSamples();
+	DelayLine->AddZeros(MaxDelay);
+	Z1 = 0.0f;
+}
+
+
+void FAllPassFractionalDelay::ProcessAudio(const AlignedFloatBuffer& InSamples, const AlignedFloatBuffer& InDelays, AlignedFloatBuffer& OutSamples)
+{
+	const int32 InNum = InSamples.Num();
+	checkf(InNum == InDelays.Num(), TEXT("Input buffers must be equal length."));
+
+	// Prepare output buffer
+	OutSamples.Reset(InNum);
+	OutSamples.AddUninitialized(InNum);
+
+	float* OutSampleData = OutSamples.GetData();
+	const float* InSampleData = InSamples.GetData();
+	const float* InDelayData = InDelays.GetData();
+
+	// Process audio one block at a time.
+	int32 LeftOver = InNum;
+	int32 BufferPos = 0;
+	while (LeftOver > 0)
+	{
+		int32 NumToProcess = FMath::Min(LeftOver, NumInternalBufferSamples);
+		ProcessAudioBlock(&InSampleData[BufferPos], &InDelayData[BufferPos], NumToProcess, &OutSampleData[BufferPos]);
+		BufferPos += NumToProcess;
+		LeftOver -= NumToProcess;
+	}
+}
+
+void FAllPassFractionalDelay::ProcessAudioBlock(const float* InSamples, const float* InDelays, const int32 InNum, float* OutSamples)
+{
+	checkf(0 == (InNum % 4), TEXT("Array length must be multiple of 4"));
+
+	// All these assume that InNum <= NumInternalBufferSamples
+	const int32* IntegerDelayOffsetData = IntegerDelayOffsets.GetData();
+	int32* IntegerDelayData = IntegerDelays.GetData();
+	float* FractionalDelayData = FractionalDelays.GetData();
+	float* CoefficientsData = Coefficients.GetData();
+
+	FMemory::Memcpy(FractionalDelayData, InDelays, InNum * sizeof(float));
+	
+	// Determine integer delays and filter coefficients per a sample
+	const VectorRegister VTwoAndAHalf = MakeVectorRegister(2.5f, 2.5f, 2.5f, 2.5f);
+	const VectorRegister VMaxDelay = MakeVectorRegister((float)MaxDelay, (float)MaxDelay, (float)MaxDelay, (float)MaxDelay);
+	for (int32 i = 0; i < InNum; i += 4)
+	{
+		checkf(FractionalDelayData[i] <= MaxDelay, TEXT("Delay exceeds maximum"));
+		checkf(FractionalDelayData[i + 1] <= MaxDelay, TEXT("Delay exceeds maximum"));
+		checkf(FractionalDelayData[i + 2] <= MaxDelay, TEXT("Delay exceeds maximum"));
+		checkf(FractionalDelayData[i + 3] <= MaxDelay, TEXT("Delay exceeds maximum"));
+
+		// Ensure that delays are in [0.5, MaxDelay]
+		// Delay= Max(Delay, 0.5)
+		VectorRegister VFractionalDelays = VectorLoadAligned(&FractionalDelayData[i]);
+		VFractionalDelays = VectorMax(VFractionalDelays, GlobalVectorConstants::FloatOneHalf);
+		// Delay = Min(Delay, MaxDelay)
+		VFractionalDelays = VectorMin(VFractionalDelays, VMaxDelay);
+
+		// To ensure fractional delays between -0.5 - 0.5, subtract limit then add back in.
+		// Delay = Delay + 0.5
+		VFractionalDelays = VectorAdd(VFractionalDelays, GlobalVectorConstants::FloatOneHalf);
+		// Delay.floor = Floor(Delay)
+		VectorRegister VFloorDelays = VectorFloor(VFractionalDelays);
+		// Delay.frac = Delay - Delay.floor
+		VFractionalDelays = VectorSubtract(VFractionalDelays, VFloorDelays);
+		
+		// Reintroduing 0.5 previously removed.
+		// alpha = Delay.frac - 0.5
+		VectorRegister VCoefficients = VectorSubtract(VFractionalDelays, GlobalVectorConstants::FloatOneHalf);
+		// denom = 2.5 - alpha
+		VectorRegister VDenominator = VectorSubtract(VTwoAndAHalf, VFractionalDelays);
+		// coef = alpha / (2.5 - alpha)
+		VCoefficients = VectorDivide(VCoefficients, VDenominator);
+		VectorStoreAligned(VCoefficients, &Coefficients[i]);
+
+		// Delay.int = int(Delay.floor)
+		VectorRegisterInt VIntegerDelays = VectorFloatToInt(VFloorDelays);
+		VectorRegisterInt VIntegerDelayOffset = VectorIntLoadAligned(&IntegerDelayOffsetData[i]);
+		// Delay.int += BufferIdx
+		VectorIntStoreAligned(VectorIntSubtract(VIntegerDelayOffset, VIntegerDelays), &IntegerDelayData[i]);
+	}
+
+	// Update delay line.
+	DelayLine->AddSamples(InSamples, InNum);
+
+	// All pass filter delay line to get fractional delay
+	const float* DelayData = DelayLine->InspectSamples(InNum + MaxDelay);
+	int32 DelayPos;
+	if (InNum > 0)
+	{ 
+		DelayPos = IntegerDelayData[0];
+		OutSamples[0] = DelayData[DelayPos - 1] + CoefficientsData[0] * DelayData[DelayPos] - CoefficientsData[0] * Z1;
+		for (int32 i = 1; i < InNum; i++)
+		{
+			DelayPos = IntegerDelayData[i];
+			OutSamples[i] = DelayData[DelayPos - 1] + CoefficientsData[i] * DelayData[DelayPos] - CoefficientsData[i] * OutSamples[i - 1];
+		}
+
+		BufferUnderflowClampFast(OutSamples, InNum);
+
+		Z1 = OutSamples[InNum - 1];
+	}
+	
+	// Remove unneeded delay line.
+	DelayLine->RemoveSamples(InNum);
+}
+
diff --git a/Engine/Source/Runtime/AudioMixer/Private/DSP/BufferOnePoleLPF.cpp b/Engine/Source/Runtime/AudioMixer/Private/DSP/BufferOnePoleLPF.cpp
new file mode 100644
--- /dev/null
+++ b/Engine/Source/Runtime/AudioMixer/Private/DSP/BufferOnePoleLPF.cpp
@@ -0,0 +1,64 @@
+#include "DSP/BufferOnePoleLPF.h"
+#include "DSP/BufferVectorOperations.h"
+
+using namespace Audio;
+
+// Constructor 
+FBufferOnePoleLPF::FBufferOnePoleLPF(float InG)
+	: CutoffFrequency(0.0f)
+	, B1(0.0f)
+	, A0(1.0f)
+	, Z1(0.0f)
+{
+	SetG(InG);
+}
+
+// Set the LPF gain coefficient
+void FBufferOnePoleLPF::SetG(float InG)
+{ 
+	B1 = InG;
+	A0 = 1.0f - B1;
+}
+
+// Resets the sample delay to 0
+void FBufferOnePoleLPF::Reset()
+{
+	SetG(0.0f);
+	Z1 = 0.0f;
+}
+
+/** Sets the filter frequency using normalized frequency (between 0.0 and 1.0f or 0.0 hz and Nyquist Frequency in Hz) */
+void FBufferOnePoleLPF::SetFrequency(const float InFrequency)
+{
+	if (!FMath::IsNearlyEqual(InFrequency, CutoffFrequency))
+	{
+		CutoffFrequency = InFrequency;
+		SetG(FMath::Exp(-PI * CutoffFrequency));
+	}
+}
+
+void FBufferOnePoleLPF::ProcessAudio(const Audio::AlignedFloatBuffer& InSamples, Audio::AlignedFloatBuffer& OutSamples)
+{
+	int32 Index;
+	const int32 InNum = InSamples.Num();
+	const float* InSampleData = InSamples.GetData();
+	
+	// Prepare output samples
+	OutSamples.Reset(InNum);
+	OutSamples.AddUninitialized(InNum);
+	float* OutSampleData = OutSamples.GetData();
+
+	// Filter audio
+	int32 DelayIndex;
+	OutSampleData[0] = InSampleData[0] * A0 + B1 * Z1;
+	for(Index = 1, DelayIndex = 0; Index < InNum; Index++, DelayIndex++)
+	{
+		//OutSampleData[Index] = UnderflowClamp(InSampleData[Index] * A0 + B1 * OutSampleData[DelayIndex]);
+		OutSampleData[Index] = InSampleData[Index] * A0 + B1 * OutSampleData[DelayIndex];
+	}
+	BufferUnderflowClampFast(OutSampleData, InNum);
+	// Store delay value
+	Z1 = OutSampleData[InNum - 1];
+}
+
+
diff --git a/Engine/Source/Runtime/AudioMixer/Private/DSP/BufferVectorOperations.cpp b/Engine/Source/Runtime/AudioMixer/Private/DSP/BufferVectorOperations.cpp
--- a/Engine/Source/Runtime/AudioMixer/Private/DSP/BufferVectorOperations.cpp
+++ b/Engine/Source/Runtime/AudioMixer/Private/DSP/BufferVectorOperations.cpp
@@ -7,6 +7,35 @@
 
 namespace Audio
 {
+	/* Sets a values to zero if value is denormal. Denormal numbers significantly slow down floating point operations. */
+	void BufferUnderflowClampFast(AlignedFloatBuffer& InOutBuffer)
+	{
+		BufferUnderflowClampFast(InOutBuffer.GetData(), InOutBuffer.Num());
+	}
+	
+	/* Sets a values to zero if value is denormal. Denormal numbers significantly slow down floating point operations. */
+	void BufferUnderflowClampFast(float* RESTRICT InOutBuffer, const int32 InNum)
+	{
+		checkf(InNum >= 4, TEXT("Buffer must have atleast 4 elements."));
+		checkf(0 == (InNum % 4), TEXT("Buffer length be a multiple of 4."));
+		checkf(IsAligned<float*>(InOutBuffer, AUDIO_SIMD_FLOAT_ALIGNMENT), TEXT("Memory must be aligned to use vector operations."));
+
+		const VectorRegister VFMIN = MakeVectorRegister(FLT_MIN, FLT_MIN, FLT_MIN, FLT_MIN);
+		const VectorRegister VNFMIN = MakeVectorRegister(-FLT_MIN, -FLT_MIN, -FLT_MIN, -FLT_MIN);
+
+		for (int32 i = 0; i < InNum; i += 4)
+		{
+			VectorRegister VInOut = VectorLoadAligned(&InOutBuffer[i]);
+
+			// Create mask of denormal numbers.
+			VectorRegister Mask = VectorBitwiseAnd(VectorCompareGT(VInOut, VNFMIN), VectorCompareLT(VInOut, VFMIN));
+
+			// Choose between zero or original number based upon mask.
+			VInOut = VectorSelect(Mask, GlobalVectorConstants::FloatZero, VInOut);
+			VectorStoreAligned(VInOut, &InOutBuffer[i]);
+		}
+	}
+
 	void BufferMultiplyByConstant(const AlignedFloatBuffer& InFloatBuffer, float InValue, AlignedFloatBuffer& OutFloatBuffer)
 	{
 		check(InFloatBuffer.Num() >= 4);
@@ -23,35 +52,41 @@ namespace Audio
 		const float* InBufferPtr = InFloatBuffer.GetData();
 		float* OutBufferPtr = OutFloatBuffer.GetData();
 
+		BufferMultiplyByConstant(InBufferPtr, InValue, OutBufferPtr, NumSamples);
+	}
+
+	void BufferMultiplyByConstant(const float* RESTRICT InFloatBuffer, float InValue, float* RESTRICT OutFloatBuffer, const int32 InNumSamples)
+	{
+		check(InNumSamples >= 4);
 #if !AUDIO_USE_SIMD
-		for (int32 i = 0; i < NumSamples; ++i)
+		for (int32 i = 0; i < InNumSamples; ++i)
 		{
-			OutBufferPtr[i] = InValue * InBufferPtr[i];
+			OutFloatBuffer[i] = InValue * InFloatBuffer[i];
 		}
 #else
 
 		// Can only SIMD on multiple of 4 buffers, we'll do normal multiples on last bit
-		const int32 NumSamplesRemaining = NumSamples % 4;
-		const int32 NumSamplesToSimd = NumSamples - NumSamplesRemaining;
+		const int32 NumSamplesRemaining = InNumSamples % 4;
+		const int32 NumSamplesToSimd = InNumSamples - NumSamplesRemaining;
 
 		// Load the single value we want to multiply all values by into a vector register
 		const VectorRegister MultiplyValue = VectorLoadFloat1(&InValue);
 		for (int32 i = 0; i < NumSamplesToSimd; i += 4)
 		{
 			// Load the next 4 samples of the input buffer into a register
-			VectorRegister InputBufferRegister = VectorLoadAligned(&InBufferPtr[i]);
+			VectorRegister InputBufferRegister = VectorLoadAligned(&InFloatBuffer[i]);
 
 			// Perform the multiply
 			VectorRegister Temp = VectorMultiply(InputBufferRegister, MultiplyValue);
 
 			// Store results into the output buffer
-			VectorStoreAligned(Temp, &OutBufferPtr[i]);
+			VectorStoreAligned(Temp, &OutFloatBuffer[i]);
 		}
 
 		// Perform remaining non-simd values left over
 		for (int32 i = 0; i < NumSamplesRemaining; ++i)
 		{
-			OutBufferPtr[NumSamplesToSimd + i] = InValue * InBufferPtr[NumSamplesToSimd + i];
+			OutFloatBuffer[NumSamplesToSimd + i] = InValue * InFloatBuffer[NumSamplesToSimd + i];
 		}
 #endif
 	}
@@ -73,6 +108,77 @@ namespace Audio
 		}
 	}
 
+	/* Performs an element-wise weighted sum OutputBuffer = (InBuffer1 x InGain1) + (InBuffer2 x InGain2) */
+	void BufferWeightedSumFast(const AlignedFloatBuffer& InBuffer1, float InGain1, const AlignedFloatBuffer& InBuffer2, float InGain2, AlignedFloatBuffer& OutBuffer)
+	{
+		checkf(InBuffer1.Num() == InBuffer2.Num(), TEXT("Buffers must be equal length"));
+		OutBuffer.Reset();
+		OutBuffer.AddUninitialized(InBuffer1.Num());
+
+		BufferWeightedSumFast(InBuffer1.GetData(), InGain1, InBuffer2.GetData(), InGain2, OutBuffer.GetData(), InBuffer1.Num());
+	}
+
+	/* Performs an element-wise weighted sum OutputBuffer = (InBuffer1 x InGain1) + InBuffer2 */
+	void BufferWeightedSumFast(const AlignedFloatBuffer& InBuffer1, float InGain1, const AlignedFloatBuffer& InBuffer2, AlignedFloatBuffer& OutBuffer)
+	{
+		checkf(InBuffer1.Num() == InBuffer2.Num(), TEXT("Buffers must be equal length"));
+		OutBuffer.Reset();
+		OutBuffer.AddUninitialized(InBuffer1.Num());
+
+		BufferWeightedSumFast(InBuffer1.GetData(), InGain1, InBuffer2.GetData(), OutBuffer.GetData(), InBuffer1.Num());
+	}
+
+	/* Performs an element-wise weighted sum OutputBuffer = (InBuffer1 x InGain1) + (InBuffer2 x InGain2) */
+	void BufferWeightedSumFast(const float* RESTRICT InBuffer1, float InGain1, const float* RESTRICT InBuffer2, float InGain2, float* RESTRICT OutBuffer, int32 InNum)
+	{
+		checkf(InNum >= 4, TEXT("Buffer must have atleast 4 elements."));
+		checkf(0 == (InNum % 4), TEXT("Buffer length be a multiple of 4."));
+		checkf(IsAligned<const float*>(InBuffer1, AUDIO_SIMD_FLOAT_ALIGNMENT), TEXT("Memory must be aligned to use vector operations."));
+		checkf(IsAligned<const float*>(InBuffer2, AUDIO_SIMD_FLOAT_ALIGNMENT), TEXT("Memory must be aligned to use vector operations."));
+		checkf(IsAligned<float*>(OutBuffer, AUDIO_SIMD_FLOAT_ALIGNMENT), TEXT("Memory must be aligned to use vector operations."));
+
+		VectorRegister Gain1Vector = VectorLoadFloat1(&InGain1);
+		VectorRegister Gain2Vector = VectorLoadFloat1(&InGain2);
+
+		for (int32 i = 0; i < InNum; i += 4)
+		{
+			// InBuffer1 x InGain1
+			VectorRegister Input1  = VectorLoadAligned(&InBuffer1[i]);
+			VectorRegister Weighted1 = VectorMultiply(Input1, Gain1Vector);
+
+			// InBuffer2 x InGain2
+			VectorRegister Input2  = VectorLoadAligned(&InBuffer2[i]);
+			VectorRegister Weighted2 = VectorMultiply(Input2, Gain2Vector);
+
+			VectorRegister Output = VectorAdd(Weighted1, Weighted2);
+			VectorStoreAligned(Output, &OutBuffer[i]);
+		}
+	}
+
+	/* Performs an element-wise weighted sum OutputBuffer = (InBuffer1 x InGain1) + InBuffer2 */
+	void BufferWeightedSumFast(const float* RESTRICT InBuffer1, float InGain1, const float* RESTRICT InBuffer2, float* RESTRICT OutBuffer, int32 InNum)
+	{
+		checkf(InNum >= 4, TEXT("Buffer must have atleast 4 elements."));
+		checkf(0 == (InNum % 4), TEXT("Buffer length be a multiple of 4."));
+		checkf(IsAligned<const float*>(InBuffer1, AUDIO_SIMD_FLOAT_ALIGNMENT), TEXT("Memory must be aligned to use vector operations."));
+		checkf(IsAligned<const float*>(InBuffer2, AUDIO_SIMD_FLOAT_ALIGNMENT), TEXT("Memory must be aligned to use vector operations."));
+		checkf(IsAligned<float*>(OutBuffer, AUDIO_SIMD_FLOAT_ALIGNMENT), TEXT("Memory must be aligned to use vector operations."));
+
+		VectorRegister Gain1Vector = VectorLoadFloat1(&InGain1);
+
+		for (int32 i = 0; i < InNum; i += 4)
+		{
+			// InBuffer1 x InGain1
+			VectorRegister Input1  = VectorLoadAligned(&InBuffer1[i]);
+			VectorRegister Weighted1 = VectorMultiply(Input1, Gain1Vector);
+
+			VectorRegister Input2  = VectorLoadAligned(&InBuffer2[i]);
+
+			VectorRegister Output = VectorAdd(Weighted1, Input2);
+			VectorStoreAligned(Output, &OutBuffer[i]);
+		}
+	}
+
 	void FadeBufferFast(AlignedFloatBuffer& OutFloatBuffer, const float StartValue, const float EndValue)
 	{
 		FadeBufferFast(OutFloatBuffer.GetData(), OutFloatBuffer.Num(), StartValue, EndValue);
@@ -180,6 +286,12 @@ namespace Audio
 #endif
 	}
 
+	void MixInBufferFast(const AlignedFloatBuffer& InFloatBuffer, AlignedFloatBuffer& BufferToSumTo)
+	{
+		checkf(InFloatBuffer.Num() == BufferToSumTo.Num(), TEXT("Buffers must be equal size"));
+		MixInBufferFast(InFloatBuffer.GetData(), BufferToSumTo.GetData(), InFloatBuffer.Num());
+	}
+
 	void MixInBufferFast(const float* RESTRICT InFloatBuffer, float* RESTRICT BufferToSumTo, int32 NumSamples)
 	{
 		checkf(IsAligned<const float*>(InFloatBuffer, 4), TEXT("Memory must be aligned to use vector operations."));
@@ -202,9 +314,94 @@ namespace Audio
 #endif
 	}
 
+	/* Subtracts two buffers together element-wise. */
+	void BufferSubtractFast(const AlignedFloatBuffer& InMinuend, const AlignedFloatBuffer& InSubtrahend, AlignedFloatBuffer& OutputBuffer)
+	{
+		const int32 InNum = InMinuend.Num();
+		OutputBuffer.Reset(InNum);
+		OutputBuffer.AddUninitialized(InNum);
+
+		checkf(InMinuend.Num() == InSubtrahend.Num(), TEXT("Input buffers must be equal length"));
+
+		BufferSubtractFast(InMinuend.GetData(), InSubtrahend.GetData(), OutputBuffer.GetData(), OutputBuffer.Num());
+	}
+
+	/* Subtracts two buffers together element-wise. */
+	void BufferSubtractFast(const float* RESTRICT InMinuend, const float* RESTRICT InSubtrahend, float* RESTRICT OutBuffer, int32 InNum)
+	{
+		checkf(InNum >= 4, TEXT("Buffer must have atleast 4 elements."));
+		checkf(0 == (InNum % 4), TEXT("Buffer length be a multiple of 4."));
+		checkf(IsAligned<const float*>(InMinuend, AUDIO_SIMD_FLOAT_ALIGNMENT), TEXT("Memory must be aligned to use vector operations."));
+		checkf(IsAligned<const float*>(InSubtrahend, AUDIO_SIMD_FLOAT_ALIGNMENT), TEXT("Memory must be aligned to use vector operations."));
+		checkf(IsAligned<float*>(OutBuffer, AUDIO_SIMD_FLOAT_ALIGNMENT), TEXT("Memory must be aligned to use vector operations."));
+
+		for (int32 i = 0; i < InNum; i += 4)
+		{
+			VectorRegister Input1 = VectorLoadAligned(&InMinuend[i]);
+			VectorRegister Input2 = VectorLoadAligned(&InSubtrahend[i]);
+			VectorRegister Output = VectorSubtract(Input1, Input2);
+			VectorStoreAligned(Output, &OutBuffer[i]);
+		}
+	}
+
+	/* Performs element-wise in-place subtraction placing the result in the subtrahend. InOutSubtrahend = InMinuend - InOutSubtrahend */
+	void BufferSubtractInPlace1Fast(const AlignedFloatBuffer& InMinuend, AlignedFloatBuffer& InOutSubtrahend)
+	{
+		checkf(InMinuend.Num() == InOutSubtrahend.Num(), TEXT("Input buffers must be equal length"));
+		BufferSubtractInPlace1Fast(InMinuend.GetData(), InOutSubtrahend.GetData(), InMinuend.Num());
+	}
+
+	/* Performs element-wise in-place subtraction placing the result in the subtrahend. InOutSubtrahend = InMinuend - InOutSubtrahend */
+	void BufferSubtractInPlace1Fast(const float* RESTRICT InMinuend, float* RESTRICT InOutSubtrahend, int32 InNum)
+	{
+		checkf(InNum >= 4, TEXT("Buffer must have atleast 4 elements."));
+		checkf(0 == (InNum % 4), TEXT("Buffer length be a multiple of 4."));
+		checkf(IsAligned<const float*>(InMinuend, AUDIO_SIMD_FLOAT_ALIGNMENT), TEXT("Memory must be aligned to use vector operations."));
+		checkf(IsAligned<const float*>(InOutSubtrahend, AUDIO_SIMD_FLOAT_ALIGNMENT), TEXT("Memory must be aligned to use vector operations."));
+		
+		for (int32 i = 0; i < InNum; i += 4)
+		{
+			VectorRegister Input1 = VectorLoadAligned(&InMinuend[i]);
+			VectorRegister Input2 = VectorLoadAligned(&InOutSubtrahend[i]);
+
+			VectorRegister Output = VectorSubtract(Input1, Input2);
+			VectorStoreAligned(Output, &InOutSubtrahend[i]);
+		}
+	}
+	
+	/* Performs element-wise in-place subtraction placing the result in the minuend. InOutMinuend = InOutMinuend - InSubtrahend */
+	void BufferSubtractInPlace2Fast(AlignedFloatBuffer& InOutMinuend, const AlignedFloatBuffer& InSubtrahend)
+	{
+		checkf(InOutMinuend.Num() == InSubtrahend.Num(), TEXT("Input buffers must be equal length"));
+		BufferSubtractInPlace2Fast(InOutMinuend.GetData(), InSubtrahend.GetData(), InOutMinuend.Num());
+	}
+
+	/* Performs element-wise in-place subtraction placing the result in the minuend. InOutMinuend = InOutMinuend - InSubtrahend */
+	void BufferSubtractInPlace2Fast(float* RESTRICT InOutMinuend, const float* RESTRICT InSubtrahend, int32 InNum)
+	{
+		checkf(InNum >= 4, TEXT("Buffer must have atleast 4 elements."));
+		checkf(0 == (InNum % 4), TEXT("Buffer length be a multiple of 4."));
+		checkf(IsAligned<const float*>(InOutMinuend, AUDIO_SIMD_FLOAT_ALIGNMENT), TEXT("Memory must be aligned to use vector operations."));
+		checkf(IsAligned<const float*>(InSubtrahend, AUDIO_SIMD_FLOAT_ALIGNMENT), TEXT("Memory must be aligned to use vector operations."));
+
+		for (int32 i = 0; i < InNum; i += 4)
+		{
+			VectorRegister Input1 = VectorLoadAligned(&InOutMinuend[i]);
+			VectorRegister Input2 = VectorLoadAligned(&InSubtrahend[i]);
+
+			VectorRegister Output = VectorSubtract(Input1, Input2);
+			VectorStoreAligned(Output, &InOutMinuend[i]);
+		}
+	}
+
 	void SumBuffers(const AlignedFloatBuffer& InFloatBuffer1, const AlignedFloatBuffer& InFloatBuffer2, AlignedFloatBuffer& OutputBuffer)
 	{
-		SumBuffers(InFloatBuffer1.GetData(), InFloatBuffer2.GetData(), OutputBuffer.GetData(), OutputBuffer.Num());
+		checkf(InFloatBuffer1.Num() == InFloatBuffer2.Num(), TEXT("Input buffers must be equal length"));
+		const int32 InNum = InFloatBuffer1.Num();
+		OutputBuffer.Reset(InNum);
+		OutputBuffer.AddUninitialized(InNum);
+
+		SumBuffers(InFloatBuffer1.GetData(), InFloatBuffer2.GetData(), OutputBuffer.GetData(), InNum);
 	}
 
 	void SumBuffers(const float* RESTRICT InFloatBuffer1, const float* RESTRICT InFloatBuffer2, float* RESTRICT OutputBuffer, int32 NumSamples)
@@ -490,6 +687,42 @@ namespace Audio
 		}
 	}
 
+	void MixMonoTo2ChannelsFast(const AlignedFloatBuffer& MonoBuffer, AlignedFloatBuffer& DestinationBuffer)
+	{
+		MixMonoTo2ChannelsFast(MonoBuffer.GetData(), DestinationBuffer.GetData(), MonoBuffer.Num());
+	}
+
+	/**
+	* See CHANNEL MIXING OPERATIONS above for more info.
+	* 2 frames per iteration:
+	* +------------+---------+---------+---------+---------+
+	* | VectorName | Index 0 | Index 1 | Index 2 | Index 3 |
+	* +------------+---------+---------+---------+---------+
+	* | Input      | i0      | i0      | i1      | i1      |
+	* |            | =       | =       | =       | =       |
+	* | Output     | o0      | o1      | o2      | o3      |
+	* +------------+---------+---------+---------+---------+
+	*/
+	void MixMonoTo2ChannelsFast(const float* RESTRICT MonoBuffer, float* RESTRICT DestinationBuffer, int32 InNumFrames)
+	{
+		checkf(InNumFrames >= 4, TEXT("Buffer must have atleast 4 elements."));
+		checkf(0 == (InNumFrames % 4), TEXT("Buffer length be a multiple of 4."));
+		checkf(IsAligned<const float*>(MonoBuffer, AUDIO_SIMD_FLOAT_ALIGNMENT), TEXT("Memory must be aligned to use vector operations."));
+		checkf(IsAligned<float*>(DestinationBuffer, AUDIO_SIMD_FLOAT_ALIGNMENT), TEXT("Memory must be aligned to use vector operations."));
+
+		int32 OutPos = 0;
+		for (int32 i = 0; i < InNumFrames; i += 4)
+		{
+			VectorRegister Input = VectorLoadAligned(&MonoBuffer[i]);
+			VectorRegister Output = VectorSwizzle(Input, 0, 0, 1, 1);
+			VectorStoreAligned(Output, &DestinationBuffer[OutPos]);
+			OutPos += 4;
+			Output = VectorSwizzle(Input, 2, 2, 3, 3);
+			VectorStoreAligned(Output, &DestinationBuffer[OutPos]);
+			OutPos += 4;
+		}
+	}
+
 	void Mix2ChannelsTo2ChannelsFast(const AlignedFloatBuffer& SourceBuffer, AlignedFloatBuffer& DestinationBuffer, const float* RESTRICT Gains)
 	{
 		Mix2ChannelsTo2ChannelsFast(SourceBuffer.GetData(), DestinationBuffer.GetData(), DestinationBuffer.Num() / 2, Gains);
@@ -1432,4 +1665,142 @@ namespace Audio
 			}
 		}
 	}
+
+
+	/** Interleaves samples from two input buffers */
+	void BufferInterleave2ChannelFast(const AlignedFloatBuffer& InBuffer1, const AlignedFloatBuffer& InBuffer2, AlignedFloatBuffer& OutBuffer)
+	{
+		checkf(InBuffer1.Num() == InBuffer2.Num(), TEXT("InBuffer1 Num not equal to InBuffer2 Num"));
+
+		const int32 InNum = InBuffer1.Num();
+
+		OutBuffer.Reset(2 * InNum);
+		OutBuffer.AddUninitialized(2 * InNum);
+		
+		BufferInterleave2ChannelFast(InBuffer1.GetData(), InBuffer2.GetData(), OutBuffer.GetData(), InNum);
+	}
+
+	/** Interleaves samples from two input buffers */
+	void BufferInterleave2ChannelFast(const float* RESTRICT InBuffer1, const float* RESTRICT InBuffer2, float* RESTRICT OutBuffer, const int32 InNum)
+	{
+		checkf(InNum >= 4, TEXT("Buffer must have atleast 4 elements."));
+		checkf(0 == (InNum % 4), TEXT("Buffer length be a multiple of 4."));
+		checkf(IsAligned<const float*>(InBuffer1, AUDIO_SIMD_FLOAT_ALIGNMENT), TEXT("Memory must be aligned to use vector operations."));
+		checkf(IsAligned<const float*>(InBuffer2, AUDIO_SIMD_FLOAT_ALIGNMENT), TEXT("Memory must be aligned to use vector operations."));
+		checkf(IsAligned<float*>(OutBuffer, AUDIO_SIMD_FLOAT_ALIGNMENT), TEXT("Memory must be aligned to use vector operations."));
+
+		const int32 OutNum = 2 * InNum;
+
+		int32 OutPos = 0;
+		for (int32 i = 0; i < InNum; i += 4)
+		{
+			// Vector1[L0, L1, L2, L3]
+			VectorRegister Vector1 = VectorLoadAligned(&InBuffer1[i]);
+			// Vector2[R0, R1, R2, R3]
+			VectorRegister Vector2 = VectorLoadAligned(&InBuffer2[i]);
+
+			// HalfInterleaved[L0, L1, R0, R1]
+			VectorRegister HalfInterleaved = VectorShuffle(Vector1, Vector2, 0, 1, 0, 1);
+			// Interleaved[L0, R0, L1, R1]
+			VectorRegister Interleaved = VectorSwizzle(HalfInterleaved, 0, 2, 1, 3);
+			VectorStoreAligned(Interleaved, &OutBuffer[OutPos]);
+			OutPos += 4;
+
+			// HalfInterleaved[L2, L3, R2, R3]
+			HalfInterleaved = VectorShuffle(Vector1, Vector2, 2, 3, 2, 3);
+			// Interleaved[L2, R2, L3, R3]
+			Interleaved = VectorSwizzle(HalfInterleaved, 0, 2, 1, 3);
+			VectorStoreAligned(Interleaved, &OutBuffer[OutPos]);
+			OutPos += 4;
+		}
+	}
+
+	/** Deinterleaves samples from a 2 channel input buffer */
+	void BufferDeinterleave2ChannelFast(const AlignedFloatBuffer& InBuffer, AlignedFloatBuffer& OutBuffer1, AlignedFloatBuffer& OutBuffer2)
+	{
+		const int32 InNum = InBuffer.Num();
+		const int32 InNumFrames = InNum / 2;
+		const int32 OutNum = InNumFrames;
+
+		OutBuffer1.Reset(OutNum);
+		OutBuffer2.Reset(OutNum);
+		OutBuffer1.AddUninitialized(OutNum);
+		OutBuffer2.AddUninitialized(OutNum);
+		
+		BufferDeinterleave2ChannelFast(InBuffer.GetData(), OutBuffer1.GetData(), OutBuffer2.GetData(), InNumFrames);
+	}
+
+	/** Deinterleaves samples from a 2 channel input buffer */
+	void BufferDeinterleave2ChannelFast(const float* RESTRICT InBuffer, float* RESTRICT OutBuffer1, float* RESTRICT OutBuffer2, const int32 InNumFrames)
+	{
+		checkf(InNumFrames >= 4, TEXT("Buffer must have atleast 4 elements."));
+		checkf(0 == (InNumFrames % 4), TEXT("Buffer length be a multiple of 4."));
+		checkf(IsAligned<const float*>(InBuffer, AUDIO_SIMD_FLOAT_ALIGNMENT), TEXT("Memory must be aligned to use vector operations."));
+		checkf(IsAligned<const float*>(OutBuffer1, 4), TEXT("Memory must be aligned to use vector operations."));
+		checkf(IsAligned<float*>(OutBuffer2, 4), TEXT("Memory must be aligned to use vector operations."));
+
+		int32 InNum = InNumFrames * 2;
+		int32 OutPos = 0;
+		for (int32 InPos = 0; InPos < InNum; InPos += 8)
+		{
+			// load 4 frames (2 frames per vector)
+			VectorRegister InVector1 = VectorLoadAligned(&InBuffer[InPos]);
+			VectorRegister InVector2 = VectorLoadAligned(&InBuffer[InPos + 4]);
+
+			// Write channel 0
+			VectorRegister OutVector = VectorShuffle(InVector1, InVector2, 0, 2, 0, 2);
+			VectorStoreAligned(OutVector, &OutBuffer1[OutPos]);
+
+			// Write channel 1
+			OutVector = VectorShuffle(InVector1, InVector2, 1, 3, 1, 3);
+			VectorStoreAligned(OutVector, &OutBuffer2[OutPos]);
+
+			OutPos += 4;
+		}
+	}
+
+	/** Sums 2 channel interleaved input samples. OutSamples[n] = InSamples[2n] + InSamples[2n + 1] */
+	void BufferSum2ChannelToMonoFast(const AlignedFloatBuffer& InSamples, AlignedFloatBuffer& OutSamples)
+	{
+		const int32 InNum = InSamples.Num();
+		const int32 Frames = InNum / 2;
+
+		OutSamples.Reset(Frames);
+		OutSamples.AddUninitialized(Frames);
+		
+		BufferSum2ChannelToMonoFast(InSamples.GetData(), OutSamples.GetData(), Frames);
+	}
+
+	/** Sums 2 channel interleaved input samples. OutSamples[n] = InSamples[2n] + InSamples[2n + 1] */
+	void BufferSum2ChannelToMonoFast(const float* RESTRICT InSamples, float* RESTRICT OutSamples, const int32 InNumFrames)
+	{
+		checkf(InNumFrames >= 4, TEXT("Buffer must have atleast 4 elements."));
+		checkf(0 == (InNumFrames % 4), TEXT("Buffer length be a multiple of 4."));
+		checkf(IsAligned<const float*>(InSamples, AUDIO_SIMD_FLOAT_ALIGNMENT), TEXT("Memory must be aligned to use vector operations."));
+		checkf(IsAligned<float*>(OutSamples, 4), TEXT("Memory must be aligned to use vector operations."));
+
+		const int32 InNum = InNumFrames * 2;
+		int32 OutPos = 0;
+		for (int32 i = 0; i < InNum; i += 8)
+		{
+			// Load 4 frames (2 frames per vector)
+			// Buffer1[L0, R0, L1, R1]
+			VectorRegister Buffer1 = VectorLoadAligned(&InSamples[i]);
+			// Buffer2[L2, R2, L3, R3]
+			VectorRegister Buffer2 = VectorLoadAligned(&InSamples[i + 4]);
+
+			// Shuffle samples into order
+			// Channel0[L0, L1, L2, L3]
+			VectorRegister Channel0 = VectorShuffle(Buffer1, Buffer2, 0, 2, 0, 2);
+			// Channel1[R0, R1, R2, R3]
+			VectorRegister Channel1 = VectorShuffle(Buffer1, Buffer2, 1, 3, 1, 3);
+
+			// Sum left and right.
+			// Out[L0 + R0, L1 + R1, L2 + R2, L3 + R3]
+			VectorRegister Out = VectorAdd(Channel0, Channel1);
+
+			VectorStoreAligned(Out, &OutSamples[OutPos]);
+			OutPos += 4;
+		}
+	}
 }
diff --git a/Engine/Source/Runtime/AudioMixer/Private/DSP/DynamicDelayAPF.cpp b/Engine/Source/Runtime/AudioMixer/Private/DSP/DynamicDelayAPF.cpp
new file mode 100644
--- /dev/null
+++ b/Engine/Source/Runtime/AudioMixer/Private/DSP/DynamicDelayAPF.cpp
@@ -0,0 +1,123 @@
+// Copyright 1998-2019 Epic Games, Inc. All Rights Reserved.
+
+#include "DSP/DynamicDelayAPF.h"
+#include "DSP/BufferVectorOperations.h"
+
+using namespace Audio;
+
+FDynamicDelayAPF::FDynamicDelayAPF(float InG, int32 InMinDelay, int32 InMaxDelay, int32 InMaxNumInternalBufferSamples)
+:	G(InG)
+,	MinDelay(InMinDelay)
+,	MaxDelay(InMaxDelay)
+,	NumDelaySamples(InMinDelay - 1)
+,	NumInternalBufferSamples(InMaxNumInternalBufferSamples)
+{
+	checkf(NumDelaySamples >= 0, TEXT("Minimum delay must be atleast 1"));
+	// NumInternalBufferSamples must be less than the length of the delay to support buffer indexing logic.
+	int32 MaxBufferSamples = FMath::Min(NumDelaySamples, InMaxNumInternalBufferSamples);
+	if (NumInternalBufferSamples > MaxBufferSamples)
+	{
+		NumInternalBufferSamples = MaxBufferSamples;
+		// Block length must be divisible by simd alignment to support simd operations.
+		NumInternalBufferSamples -= (NumInternalBufferSamples % AUDIO_SIMD_FLOAT_ALIGNMENT);
+	}
+
+	checkf(NumInternalBufferSamples > 0, TEXT("Invalid internal buffer length"));
+
+	// Check that delays are sane.
+	checkf(MinDelay <= MaxDelay, TEXT("invalid delay range"));
+
+
+	// Allocate delay line.
+	IntegerDelayLine = MakeUnique<FAlignedBlockBuffer>((2 * MinDelay) + NumInternalBufferSamples, MinDelay + NumInternalBufferSamples);
+	IntegerDelayLine->AddZeros(NumDelaySamples);
+	//FractionalDelayLine = MakeUnique<FAllPassFractionalDelay>(MaxDelay - MinDelay + 1, NumInternalBufferSamples);
+	FractionalDelayLine = MakeUnique<FLinearInterpFractionalDelay>(MaxDelay - MinDelay + 1, NumInternalBufferSamples);
+
+	// Allocate internal buffers.
+	FractionalDelays.Reset(NumInternalBufferSamples);
+	FractionalDelays.AddUninitialized(NumInternalBufferSamples);
+	WorkBufferA.Reset(NumInternalBufferSamples);
+	WorkBufferA.AddUninitialized(NumInternalBufferSamples);
+	WorkBufferB.Reset(NumInternalBufferSamples);
+	WorkBufferB.AddUninitialized(NumInternalBufferSamples);
+}
+
+FDynamicDelayAPF::~FDynamicDelayAPF()
+{}
+
+void FDynamicDelayAPF::ProcessAudio(const AlignedFloatBuffer& InSamples, const AlignedFloatBuffer& InSampleDelays, AlignedFloatBuffer& OutSamples)
+{
+	const int32 InNum = InSamples.Num();
+	checkf(InNum == InSampleDelays.Num(), TEXT("InSamples [%d], InSampleDelays [%d] length mismatch"), InNum, InSampleDelays.Num());
+
+	OutSamples.Reset(InNum);
+	OutSamples.AddUninitialized(InNum);
+
+	if (InNum != InSampleDelays.Num())
+	{
+		// Output silence in the case of an error.
+		FMemory::Memset(OutSamples.GetData(), 0, sizeof(float) * InNum);
+		return;
+	}
+
+	const float* InSampleData = InSamples.GetData();
+	float* OutSampleData = OutSamples.GetData();
+
+	int32 LeftOver = InNum;
+	int32 BufferPos = 0;
+	while (LeftOver != 0)
+	{
+		const int32 NumToProcess = FMath::Min(LeftOver, NumInternalBufferSamples);
+
+		FractionalDelays.Reset(NumToProcess);
+		FractionalDelays.AddUninitialized(NumToProcess);
+
+		float* FractionalDelayData = FractionalDelays.GetData();
+		FMemory::Memcpy(FractionalDelayData, &InSampleDelays.GetData()[BufferPos], NumToProcess * sizeof(float));
+		
+		for (int32 i = 0; i < NumToProcess; i++)
+		{
+			FractionalDelayData[i] -= NumDelaySamples;
+		}
+		
+		ProcessAudioBlock(&InSampleData[BufferPos], FractionalDelays, NumToProcess, &OutSampleData[BufferPos]);
+
+		LeftOver -= NumToProcess;
+		BufferPos += NumToProcess;
+	}
+}
+
+void FDynamicDelayAPF::ProcessAudioBlock(const float* InSamples, const AlignedFloatBuffer& InFractionalDelays, const int32 InNum, float* OutSamples)
+{
+	// Make a copy of the delay line w[n - d_int]
+	WorkBufferA.Reset(InNum);
+	WorkBufferA.AddUninitialized(InNum);
+	FMemory::Memcpy(WorkBufferA.GetData(), IntegerDelayLine->InspectSamples(InNum), InNum * sizeof(float));
+
+	// Apply fractional delay
+	// w[n - d]
+	FractionalDelayLine->ProcessAudio(WorkBufferA, InFractionalDelays, WorkBufferB);
+
+	// w[n] = x[n] + G * w[n - d]
+	DelayLineInput.Reset(InNum);
+	DelayLineInput.AddUninitialized(InNum);
+	BufferWeightedSumFast(WorkBufferB.GetData(), G, InSamples, DelayLineInput.GetData(), InNum);
+	
+	BufferUnderflowClampFast(DelayLineInput);
+
+	// y[n] = -G w[n] + w[n - d]
+	BufferWeightedSumFast(DelayLineInput.GetData(), -G, WorkBufferB.GetData(), OutSamples, InNum);
+	
+	// Update delay line
+	IntegerDelayLine->RemoveSamples(InNum);
+	IntegerDelayLine->AddSamples(DelayLineInput.GetData(), InNum);
+}
+
+void FDynamicDelayAPF::Reset() 
+{
+	IntegerDelayLine->ClearSamples();
+	IntegerDelayLine->AddZeros(MinDelay - 1);
+	FractionalDelayLine->Reset();
+}
+
diff --git a/Engine/Source/Runtime/AudioMixer/Private/DSP/EarlyReflectionsFast.cpp b/Engine/Source/Runtime/AudioMixer/Private/DSP/EarlyReflectionsFast.cpp
new file mode 100644
--- /dev/null
+++ b/Engine/Source/Runtime/AudioMixer/Private/DSP/EarlyReflectionsFast.cpp
@@ -0,0 +1,211 @@
+#include "DSP/EarlyReflectionsFast.h"
+
+using namespace Audio;
+
+// The presets for the late refelections reverb were based on a strange sample rate. This 
+// function helps translate delay values in the paper to delay values for this reverb.
+// https://ccrma.stanford.edu/~dattorro/EffectDesignPart1.pdf
+static int32 GetDelaySamples(float InSampleRate, float InPresetValue) {
+	
+	return (int32)((float)InSampleRate / 29761.f * InPresetValue);
+}
+
+FEarlyReflectionsFastSettings::FEarlyReflectionsFastSettings()
+	: Gain(1.0f)
+	, PreDelayMsec(0.0f)
+	, Bandwidth(0.8)
+	, Decay(0.5)
+	, Absorption(0.7)
+{}
+
+bool FEarlyReflectionsFastSettings::operator==(const FEarlyReflectionsFastSettings& Other) const
+{
+	bool bIsEqual = (
+		(Other.Gain == Gain) &&
+		(Other.PreDelayMsec == PreDelayMsec) &&
+		(Other.Bandwidth == Bandwidth) &&
+		(Other.Decay == Decay) &&
+		(Other.Absorption == Absorption));
+	return bIsEqual;
+}
+
+bool FEarlyReflectionsFastSettings::operator!=(const FEarlyReflectionsFastSettings& Other) const
+{
+	return !(*this == Other);
+}
+
+const float FEarlyReflectionsFast::MinGain       = 0.0f;
+const float FEarlyReflectionsFast::MaxGain       = 0.9999f;
+const float FEarlyReflectionsFast::MaxPreDelay   = 1000.0f;
+const float FEarlyReflectionsFast::MinPreDelay   = 0.0f;
+const float FEarlyReflectionsFast::MaxBandwidth  = 0.99999f;
+const float FEarlyReflectionsFast::MinBandwidth  = 0.0f;
+const float FEarlyReflectionsFast::MaxDecay      = 1.0f;
+const float FEarlyReflectionsFast::MinDecay      = 0.0001f;
+const float FEarlyReflectionsFast::MaxAbsorption = 0.99999f;
+const float FEarlyReflectionsFast::MinAbsorption = 0.0f;
+
+const FEarlyReflectionsFastSettings FEarlyReflectionsFast::DefaultSettings;
+
+FEarlyReflectionsFast::FEarlyReflectionsFast(float InSampleRate, int32 InMaxNumInternalBufferSamples, const FEarlyReflectionsFastSettings& InSettings)
+:	Settings(InSettings),
+	SampleRate(InSampleRate),
+	LeftFDN(InMaxNumInternalBufferSamples, FFDNDelaySettings::DefaultLeftDelays(InSampleRate)),
+	RightFDN(InMaxNumInternalBufferSamples, FFDNDelaySettings::DefaultRightDelays(InSampleRate)),
+	LeftPreDelay((int32)InSampleRate * 2.0f, 0, InMaxNumInternalBufferSamples),
+	RightPreDelay((int32)InSampleRate * 2.0f, 0, InMaxNumInternalBufferSamples)
+
+{
+	LeftCoefficients.InputScale = 0.25f;
+	RightCoefficients.InputScale = 0.25f;
+
+	ClampSettings(Settings);
+	ApplySettings();
+}
+
+FEarlyReflectionsFast::~FEarlyReflectionsFast()
+{
+}
+
+void FEarlyReflectionsFast::ClampSettings(FEarlyReflectionsFastSettings& InOutSettings)
+{
+	InOutSettings.Gain = FMath::Clamp(
+		InOutSettings.Gain,
+		FEarlyReflectionsFast::MinGain,
+		FEarlyReflectionsFast::MaxGain);
+	InOutSettings.PreDelayMsec = FMath::Clamp(
+		InOutSettings.PreDelayMsec,
+		FEarlyReflectionsFast::MinPreDelay,
+		FEarlyReflectionsFast::MaxPreDelay);
+	InOutSettings.Bandwidth = FMath::Clamp(
+		InOutSettings.Bandwidth,
+		FEarlyReflectionsFast::MinBandwidth,
+		FEarlyReflectionsFast::MaxBandwidth);
+	InOutSettings.Decay = FMath::Clamp(
+		InOutSettings.Decay,
+		FEarlyReflectionsFast::MinDecay,
+		FEarlyReflectionsFast::MaxDecay);
+	InOutSettings.Absorption = FMath::Clamp(
+		InOutSettings.Absorption,
+		FEarlyReflectionsFast::MinAbsorption,
+		FEarlyReflectionsFast::MaxAbsorption);
+}
+
+void FEarlyReflectionsFast::SetSettings(const FEarlyReflectionsFastSettings& InSettings)
+{
+	Settings = InSettings;
+	ClampSettings(Settings);
+	ApplySettings();
+}
+
+void FEarlyReflectionsFast::ApplySettings()
+{
+	int32 DelaySamples = (int32)SampleRate * Settings.PreDelayMsec / 1000.0f;
+	LeftPreDelay.SetDelayLengthSamples(DelaySamples);
+	RightPreDelay.SetDelayLengthSamples(DelaySamples);
+
+	LeftInputLPF.SetG(Settings.Bandwidth);
+	RightInputLPF.SetG(Settings.Bandwidth);
+
+	LeftCoefficients.LPFB[0] = FMath::Min(Settings.Absorption + 0.1f, 0.9999f);
+	LeftCoefficients.LPFB[1] = FMath::Min(Settings.Absorption - 0.12f, 0.9999f);
+	LeftCoefficients.LPFB[2] = FMath::Min(Settings.Absorption + 0.08f, 0.9999f);
+	LeftCoefficients.LPFB[3] = FMath::Min(Settings.Absorption - 0.07f, 0.9999f);
+	LeftCoefficients.LPFA[0] = 1.0f - LeftCoefficients.LPFB[0];
+	LeftCoefficients.LPFA[1] = 1.0f - LeftCoefficients.LPFB[1];
+	LeftCoefficients.LPFA[2] = 1.0f - LeftCoefficients.LPFB[2];
+	LeftCoefficients.LPFA[3] = 1.0f - LeftCoefficients.LPFB[3];
+	LeftCoefficients.APFG[0] = 0.1f;
+	LeftCoefficients.APFG[1] = 0.2f;
+	LeftCoefficients.APFG[2] = 0.3f;
+	LeftCoefficients.APFG[3] = 0.4f;
+
+	RightCoefficients.LPFB[0] = FMath::Min(Settings.Absorption + 0.17f, 0.999f);
+	RightCoefficients.LPFB[1] = FMath::Min(Settings.Absorption - 0.07f, 0.999f);
+	RightCoefficients.LPFB[2] = FMath::Min(Settings.Absorption + 0.05f, 0.999f);
+	RightCoefficients.LPFB[3] = FMath::Min(Settings.Absorption - 0.11f, 0.999f);
+	RightCoefficients.LPFA[0] = 1.0f - RightCoefficients.LPFB[0];
+	RightCoefficients.LPFA[1] = 1.0f - RightCoefficients.LPFB[1];
+	RightCoefficients.LPFA[2] = 1.0f - RightCoefficients.LPFB[2];
+	RightCoefficients.LPFA[3] = 1.0f - RightCoefficients.LPFB[3];
+	RightCoefficients.APFG[0] = 0.1f;
+	RightCoefficients.APFG[1] = 0.2f;
+	RightCoefficients.APFG[2] = 0.3f;
+	RightCoefficients.APFG[3] = 0.4f;
+
+	// 1/sqrt(2) * Decay
+	//float Feedback = (1.0f - Settings.Decay) * 0.707f;
+	float Feedback = (1.0f - Settings.Decay) * 0.5f;
+	LeftCoefficients.Feedback = Feedback;
+	RightCoefficients.Feedback = Feedback;
+
+	LeftFDN.SetCoefficients(LeftCoefficients);
+	RightFDN.SetCoefficients(RightCoefficients);
+}
+
+void FEarlyReflectionsFast::ProcessAudio(const AlignedFloatBuffer& InSamples, const int32 InNumChannels, AlignedFloatBuffer& OutLeftSamples, AlignedFloatBuffer& OutRightSamples)
+{
+	checkf((InNumChannels == 1) || (InNumChannels == 2), TEXT("EarlyReflections only supports one or two channel input samples."));
+
+	const int32 InNum = InSamples.Num();
+	const int32 InNumFrames = InNum / InNumChannels;
+
+	// Resize internal buffers
+	LeftInputBuffer.Reset(InNumFrames);
+	LeftWorkBufferA.Reset(InNumFrames);
+	LeftWorkBufferB.Reset(InNumFrames);
+	RightInputBuffer.Reset(InNumFrames);
+	RightWorkBufferA.Reset(InNumFrames);
+	RightWorkBufferB.Reset(InNumFrames);
+
+	LeftInputBuffer.AddUninitialized(InNumFrames);
+	LeftWorkBufferA.AddUninitialized(InNumFrames);
+	LeftWorkBufferB.AddUninitialized(InNumFrames);
+	RightInputBuffer.AddUninitialized(InNumFrames);
+	RightWorkBufferA.AddUninitialized(InNumFrames);
+	RightWorkBufferB.AddUninitialized(InNumFrames);
+
+	// Resize output buffers
+	OutLeftSamples.Reset(InNumFrames);
+	OutRightSamples.Reset(InNumFrames);
+
+	OutLeftSamples.AddUninitialized(InNumFrames);
+	OutRightSamples.AddUninitialized(InNumFrames);
+
+	if (1 == InNumChannels)
+	{
+		// Copy mono to both left and right
+		FMemory::Memcpy(LeftInputBuffer.GetData(), InSamples.GetData(), sizeof(float) * InNumFrames);
+		FMemory::Memcpy(RightInputBuffer.GetData(), InSamples.GetData(), sizeof(float) * InNumFrames);
+	}
+	else if (2 == InNumChannels)
+	{
+		BufferDeinterleave2ChannelFast(InSamples, LeftInputBuffer, RightInputBuffer);
+	}
+	else 
+	{
+		// This class only supports 1 and 2 channel input audio.
+		FMemory::Memset(OutLeftSamples.GetData(), 0, sizeof(float) * InNumFrames);
+		FMemory::Memset(OutRightSamples.GetData(), 0, sizeof(float) * InNumFrames);
+		return;
+	}
+
+	// predelay
+	LeftPreDelay.ProcessAudio(LeftInputBuffer, LeftWorkBufferB);
+	RightPreDelay.ProcessAudio(RightInputBuffer, RightWorkBufferB);
+
+	// lpf
+	LeftInputLPF.ProcessAudio(LeftWorkBufferB, LeftWorkBufferA);
+	RightInputLPF.ProcessAudio(RightWorkBufferB, RightWorkBufferA);
+
+	// feedback delay network
+	LeftFDN.ProcessAudio(LeftWorkBufferA, OutLeftSamples);
+	RightFDN.ProcessAudio(RightWorkBufferA, OutRightSamples);
+	
+	const float OneMinusGain = 1.0 - Settings.Gain;
+
+	// Apply Gain
+	MultiplyBufferByConstantInPlace(OutLeftSamples, Settings.Gain);
+	MultiplyBufferByConstantInPlace(OutRightSamples, Settings.Gain);
+}
+
diff --git a/Engine/Source/Runtime/AudioMixer/Private/DSP/FeedbackDelayNetwork.cpp b/Engine/Source/Runtime/AudioMixer/Private/DSP/FeedbackDelayNetwork.cpp
new file mode 100644
--- /dev/null
+++ b/Engine/Source/Runtime/AudioMixer/Private/DSP/FeedbackDelayNetwork.cpp
@@ -0,0 +1,178 @@
+#include "DSP/FeedbackDelayNetwork.h"
+
+
+using namespace Audio;
+
+// The presets for the late refelections reverb were based on a strange sample rate. This 
+// function helps translate delay values in the paper to delay values for this reverb.
+// https://ccrma.stanford.edu/~dattorro/EffectDesignPart1.pdf
+static int32 GetDelaySamples(float InSampleRate, float InPresetValue) {
+	
+	return (int32)((float)InSampleRate / 29761.f * InPresetValue);
+}
+
+FFDNDelaySettings FFDNDelaySettings::DefaultLeftDelays(float InSampleRate)
+{
+	FFDNDelaySettings Settings;
+	Settings.APF0DelayNumSamples = GetDelaySamples(InSampleRate, 5 * 142);
+	Settings.APF1DelayNumSamples = GetDelaySamples(InSampleRate, 5 * 107);
+	Settings.APF2DelayNumSamples = GetDelaySamples(InSampleRate, 5 * 379);
+	Settings.APF3DelayNumSamples = GetDelaySamples(InSampleRate, 5 * 277);
+	return Settings;
+}
+
+FFDNDelaySettings FFDNDelaySettings::DefaultRightDelays(float InSampleRate)
+{
+	FFDNDelaySettings Settings;
+	Settings.APF0DelayNumSamples = GetDelaySamples(InSampleRate, 5 * 279);
+	Settings.APF1DelayNumSamples = GetDelaySamples(InSampleRate, 5 * 137);
+	Settings.APF2DelayNumSamples = GetDelaySamples(InSampleRate, 5 * 213);
+	Settings.APF3DelayNumSamples = GetDelaySamples(InSampleRate, 5 * 327);
+	return Settings;
+}
+
+FFeedbackDelayNetwork::FFeedbackDelayNetwork(int32 InMaxNumInternalBufferSamples, const FFDNDelaySettings& InSettings)
+{
+	// num internal buffer samples must be less than the minimum delay line and abide by memory alignment.
+	NumInternalBufferSamples = FMath::Min(
+		FMath::Min(InSettings.APF0DelayNumSamples, InSettings.APF1DelayNumSamples),
+		FMath::Min(InSettings.APF2DelayNumSamples, InSettings.APF3DelayNumSamples));
+
+	NumInternalBufferSamples = FMath::Min(NumInternalBufferSamples, InMaxNumInternalBufferSamples);
+
+	NumInternalBufferSamples -= (NumInternalBufferSamples % AUDIO_SIMD_FLOAT_ALIGNMENT);
+
+	if (NumInternalBufferSamples <= 0)
+	{
+		NumInternalBufferSamples = 1;
+	}
+
+	// Allocate delay lines.
+	DelayLine0 = MakeUnique<FAlignedBlockBuffer>((2 * InSettings.APF0DelayNumSamples) + NumInternalBufferSamples, NumInternalBufferSamples);
+	DelayLine1 = MakeUnique<FAlignedBlockBuffer>((2 * InSettings.APF1DelayNumSamples) + NumInternalBufferSamples, NumInternalBufferSamples);
+	DelayLine2 = MakeUnique<FAlignedBlockBuffer>((2 * InSettings.APF2DelayNumSamples) + NumInternalBufferSamples, NumInternalBufferSamples);
+	DelayLine3 = MakeUnique<FAlignedBlockBuffer>((2 * InSettings.APF3DelayNumSamples) + NumInternalBufferSamples, NumInternalBufferSamples);
+
+	// Initialize delay lines
+	DelayLine0->AddZeros(InSettings.APF0DelayNumSamples);
+	DelayLine1->AddZeros(InSettings.APF1DelayNumSamples);
+	DelayLine2->AddZeros(InSettings.APF2DelayNumSamples);
+	DelayLine3->AddZeros(InSettings.APF3DelayNumSamples);
+
+	// Initialize internal buffers
+	WorkBuffer0.Reset(NumInternalBufferSamples);
+	WorkBuffer1.Reset(NumInternalBufferSamples);
+	WorkBuffer2.Reset(NumInternalBufferSamples);
+	WorkBuffer3.Reset(NumInternalBufferSamples);
+	WorkBuffer0.AddUninitialized(NumInternalBufferSamples);
+	WorkBuffer1.AddUninitialized(NumInternalBufferSamples);
+	WorkBuffer2.AddUninitialized(NumInternalBufferSamples);
+	WorkBuffer3.AddUninitialized(NumInternalBufferSamples);
+
+	// Initialize filter memory
+	for (int32 i = 0; i < 4; i++)
+	{
+		FMO[i] = 0.0f;
+		LPFZ[i] = 0.0f;
+	}
+}
+
+FFeedbackDelayNetwork::~FFeedbackDelayNetwork()
+{}
+
+void FFeedbackDelayNetwork::SetCoefficients(const FFDNCoefficients& InCoefficients)
+{
+	Coefficients = InCoefficients;
+}
+
+void FFeedbackDelayNetwork::ProcessAudio(const AlignedFloatBuffer& InSamples, AlignedFloatBuffer& OutSamples)
+{
+	const int32 InNum = InSamples.Num();
+	const float* InSampleData = InSamples.GetData();
+
+	// Prepare output buffer
+	OutSamples.Reset(InNum);
+	OutSamples.AddUninitialized(InNum);
+
+	// Process audio one block at a time
+	float* OutSampleData = OutSamples.GetData();
+	int32 LeftToProcess = InNum;
+	int32 BufferPos = 0;
+
+	while (LeftToProcess > 0)
+	{
+		int32 NumToProcess = FMath::Min(NumInternalBufferSamples, LeftToProcess);
+		ProcessAudioBuffer(&InSampleData[BufferPos], NumToProcess, &OutSampleData[BufferPos]);
+		LeftToProcess -= NumToProcess;
+		BufferPos += NumToProcess;
+	}
+
+}
+
+int32 FFeedbackDelayNetwork::GetNumInternalBufferSamples() const
+{
+	return NumInternalBufferSamples;
+}
+
+void FFeedbackDelayNetwork::ProcessAudioBuffer(const float* InSampleData, const int32 InNum, float* OutSampleData)
+{
+	const float* DelayLineData0 = DelayLine0->InspectSamples(InNum);
+	const float* DelayLineData1 = DelayLine1->InspectSamples(InNum);
+	const float* DelayLineData2 = DelayLine2->InspectSamples(InNum);
+	const float* DelayLineData3 = DelayLine3->InspectSamples(InNum);
+	float* WorkBuffer0Data = WorkBuffer0.GetData();
+	float* WorkBuffer1Data = WorkBuffer1.GetData();
+	float* WorkBuffer2Data = WorkBuffer2.GetData();
+	float* WorkBuffer3Data = WorkBuffer3.GetData();
+	
+	float AllPassOut[4] = {0.0f};
+	float DelayInput[4] = {0.0f};
+	for (int32 i = 0; i < InNum; i++)
+	{
+		
+		float DelayOut[4] = {DelayLineData0[i], DelayLineData1[i], DelayLineData2[i], DelayLineData3[i]};
+
+		// Scale input
+		float ScaledInput = Coefficients.InputScale * InSampleData[i];
+		
+		// Calculate input to delay line. 
+		for (int32 j = 0; j < 4; j++)
+		{
+			DelayInput[j] = UnderflowClamp((ScaledInput + FMO[j]) + (Coefficients.APFG[j] * DelayOut[j]));
+		}
+
+		// Store delay line values
+		WorkBuffer0Data[i] = DelayInput[0];
+		WorkBuffer1Data[i] = DelayInput[1];
+		WorkBuffer2Data[i] = DelayInput[2];
+		WorkBuffer3Data[i] = DelayInput[3];
+
+		// complete all pass filters and process low pass filters
+		OutSampleData[i] = 0.0f;
+		for (int32 j = 0; j < 4; j++)
+		{
+			AllPassOut[j] = (DelayInput[j] * -Coefficients.APFG[j]) + DelayOut[j];
+			// Note the reasignment of lpz. dont fuck it up.
+			LPFZ[j] = UnderflowClamp((AllPassOut[j] * Coefficients.LPFA[j]) + (LPFZ[j] * Coefficients.LPFB[j]));
+			OutSampleData[i] += LPFZ[j];
+		}
+
+		// Mix channels and apply feedback.
+		FMO[0] = Coefficients.Feedback * (-LPFZ[1] + LPFZ[2]);
+		FMO[1] = Coefficients.Feedback * (LPFZ[0] + LPFZ[3]);
+		FMO[2] = Coefficients.Feedback * (LPFZ[0] - LPFZ[3]);
+		FMO[3] = Coefficients.Feedback * (-LPFZ[1] - LPFZ[2]);
+
+	}
+
+	// Update delay lines.
+	DelayLine0->RemoveSamples(InNum);
+	DelayLine1->RemoveSamples(InNum);
+	DelayLine2->RemoveSamples(InNum);
+	DelayLine3->RemoveSamples(InNum);
+	DelayLine0->AddSamples(WorkBuffer0Data, InNum);
+	DelayLine1->AddSamples(WorkBuffer1Data, InNum);
+	DelayLine2->AddSamples(WorkBuffer2Data, InNum);
+	DelayLine3->AddSamples(WorkBuffer3Data, InNum);
+}
+
diff --git a/Engine/Source/Runtime/AudioMixer/Private/DSP/IntegerDelay.cpp b/Engine/Source/Runtime/AudioMixer/Private/DSP/IntegerDelay.cpp
new file mode 100644
--- /dev/null
+++ b/Engine/Source/Runtime/AudioMixer/Private/DSP/IntegerDelay.cpp
@@ -0,0 +1,121 @@
+// Copyright 1998-2019 Epic Games, Inc. All Rights Reserved.
+
+#include "DSP/IntegerDelay.h"
+#include "DSP/BufferVectorOperations.h"
+
+using namespace Audio;
+
+FIntegerDelay::FIntegerDelay(int32 InMaxNumDelaySamplesSamples, int32 InNumDelaySamples, int32 InNumInternalBufferSamples)
+:	MaxNumDelaySamples(InMaxNumDelaySamplesSamples)
+	, NumDelaySamples(0)
+	, NumDelayLineOffsetSamples(0)
+	, NumBufferOffsetSamples(0)
+	, NumInternalBufferSamples(InNumInternalBufferSamples)
+{
+	// Allocate and prepare delay line for maximum delay.
+	DelayLine = MakeUnique<FAlignedBlockBuffer>((2 * MaxNumDelaySamples) + NumInternalBufferSamples, MaxNumDelaySamples + NumInternalBufferSamples);
+	DelayLine->AddZeros(MaxNumDelaySamples);
+
+	// Set current delay.
+	SetDelayLengthSamples(InNumDelaySamples);
+}
+
+
+// Destructor
+FIntegerDelay::~FIntegerDelay()
+{}
+
+void FIntegerDelay::SetDelayLengthSamples(int32 InNumDelaySamples)
+{
+	checkf(InNumDelaySamples <= MaxNumDelaySamples, TEXT("InNumDelaySamples must be less than or equal to MaxNumDelaySamples"));
+	checkf(InNumDelaySamples >= 0, TEXT("InNumDelaySamples must be greater than or equal to 0"));
+
+	if (InNumDelaySamples > MaxNumDelaySamples)
+	{
+		InNumDelaySamples = MaxNumDelaySamples;
+	}
+	else if (InNumDelaySamples < 0)
+	{
+		InNumDelaySamples = 0;
+	}
+
+	NumDelaySamples = InNumDelaySamples;
+	// Store offset location for reading samples out from delay line.
+	// We keep a couple offsets around to ensure we do go over buffer bounds
+	// and abide by alignment rules of FAlignedBlockBuffer
+	NumDelayLineOffsetSamples = MaxNumDelaySamples - NumDelaySamples;
+	NumBufferOffsetSamples = 0;
+	while (0 != (NumDelayLineOffsetSamples % AUDIO_SIMD_FLOAT_ALIGNMENT))
+	{
+		NumDelayLineOffsetSamples--;
+		NumBufferOffsetSamples++;
+	}
+}
+
+// Resets the delay line state, flushes buffer and resets read/write pointers.
+void FIntegerDelay::Reset() 
+{
+	DelayLine->ClearSamples();
+	DelayLine->AddZeros(MaxNumDelaySamples);
+}
+
+// Returns the current delay line length (in samples).
+int32 FIntegerDelay::GetNumDelaySamples() const 
+{
+	return NumDelaySamples;
+}
+
+void FIntegerDelay::ProcessAudio(const Audio::AlignedFloatBuffer& InSamples, Audio::AlignedFloatBuffer& OutSamples)
+{
+	const float* InSampleData = InSamples.GetData();
+	const int32 InNum = InSamples.Num();
+
+	// Prepare output buffer
+	OutSamples.Reset(InNum);
+	OutSamples.AddUninitialized(InNum);
+	float* OutSampleData = OutSamples.GetData();
+
+	// Process audio one block at a time.
+	int32 LeftOver = InNum;
+	int32 BufferPos = 0;
+	while (LeftOver > 0)
+	{
+		int32 NumToProcess = FMath::Min(LeftOver, NumInternalBufferSamples);
+		ProcessAudioBlock(&InSampleData[BufferPos], NumToProcess, &OutSampleData[BufferPos]);
+		BufferPos += NumToProcess;
+		LeftOver -= NumToProcess;
+	}
+}
+
+void FIntegerDelay::ProcessAudioBlock(const float* InSamples, const int32 InNum, float* OutSamples)
+{
+	
+	// Update delay line.	
+	DelayLine->AddSamples(InSamples, InNum);
+
+	// Copy delayed version to output
+	const float* DelayData = DelayLine->InspectSamples(InNum + NumBufferOffsetSamples, NumDelayLineOffsetSamples);
+
+	FMemory::Memcpy(OutSamples, &DelayData[NumBufferOffsetSamples], InNum * sizeof(float));
+	
+	// Remove unneeded delay line.
+	DelayLine->RemoveSamples(InNum);
+}
+
+void FIntegerDelay::PeekDelayLine(int32 InNum, Audio::AlignedFloatBuffer& OutSamples)
+{
+	int32 NumToInspect = FMath::Min(DelayLine->GetNumAvailable(), NumDelaySamples);
+	NumToInspect = FMath::Min(InNum, NumToInspect);
+
+	const float* DelaySamples = DelayLine->InspectSamples(NumToInspect);
+	if (nullptr == DelaySamples)
+	{
+		OutSamples.Reset(0);
+		return;
+	}
+
+	OutSamples.Reset(NumToInspect);
+	OutSamples.AddUninitialized(NumToInspect);
+	FMemory::Memcpy(OutSamples.GetData(), DelaySamples, NumToInspect * sizeof(float));
+}
+
diff --git a/Engine/Source/Runtime/AudioMixer/Private/DSP/LateReflectionsFast.cpp b/Engine/Source/Runtime/AudioMixer/Private/DSP/LateReflectionsFast.cpp
new file mode 100644
--- /dev/null
+++ b/Engine/Source/Runtime/AudioMixer/Private/DSP/LateReflectionsFast.cpp
@@ -0,0 +1,531 @@
+#include "DSP/LateReflectionsFast.h"
+#include "CoreMinimal.h"
+#include "DSP/BufferVectorOperations.h"
+
+using namespace Audio;
+
+// The presets for the late refelections reverb were based on a strange sample rate. This 
+// function helps translate delay values in the paper to delay values for this reverb.
+// https://ccrma.stanford.edu/~dattorro/EffectDesignPart1.pdf
+static int32 GetDelaySamples(float InSampleRate, float InPresetValue) {
+	
+	return (int32)((float)InSampleRate / 29761.f * InPresetValue);
+}
+
+FLateReflectionsFastSettings::FLateReflectionsFastSettings()
+	: LateDelayMsec(0.0f)
+	, LateGainDB(0.0f)
+	, Bandwidth(0.5f)
+	, Diffusion(0.5f)
+	, Dampening(0.5f)
+	, Decay(0.5f)
+	, Density(0.5f)
+{}
+
+bool FLateReflectionsFastSettings::operator==(const FLateReflectionsFastSettings& Other) const
+{
+	bool bIsEqual = (
+		(Other.LateDelayMsec == LateDelayMsec) &&
+		(Other.LateGainDB == LateGainDB) &&
+		(Other.Bandwidth == Bandwidth) &&
+		(Other.Diffusion == Diffusion) &&
+		(Other.Dampening == Dampening) &&
+		(Other.Decay == Decay) &&
+		(Other.Density == Density));
+
+	return bIsEqual;
+}
+
+bool FLateReflectionsFastSettings::operator!=(const FLateReflectionsFastSettings& Other) const
+{
+	return !(*this == Other);
+}
+
+void FLateReflectionsPlateOutputs::ResizeAndZero(int32 InNumSamples)
+{
+	for (int32 i = 0; i < NumTaps; i++)
+	{
+		Taps[i].Reset(InNumSamples);
+		Taps[i].AddUninitialized(InNumSamples);
+		FMemory::Memset(Taps[i].GetData(), 0, sizeof(float) * InNumSamples);
+	}
+
+	Output.Reset(InNumSamples);
+	Output.AddUninitialized(InNumSamples);
+	FMemory::Memset(Output.GetData(), 0, sizeof(float) * InNumSamples);
+}
+
+
+FLateReflectionsPlateDelays FLateReflectionsPlateDelays::DefaultLeftDelays(float InSampleRate)
+{
+	FLateReflectionsPlateDelays Delays;
+
+	Delays.NumSamplesModulatedBase = GetDelaySamples(InSampleRate, 908);
+	Delays.NumSamplesModulatedDelta = GetDelaySamples(InSampleRate, 16);
+	// Left Delay 1 should add up to 4217 and have tap outs at 353 (R0), 1190 (L4), and 3627 (R1)
+	Delays.NumSamplesDelayA = GetDelaySamples(InSampleRate, 353);
+	Delays.NumSamplesDelayB = GetDelaySamples(InSampleRate, 837);
+	Delays.NumSamplesDelayC = GetDelaySamples(InSampleRate, 2437);
+	Delays.NumSamplesDelayD = GetDelaySamples(InSampleRate, 590);
+	Delays.NumSamplesAPF = GetDelaySamples(InSampleRate, 2656);
+	// APF Delay should have tapouts at 187 (L5) and 1228 (R2)
+	Delays.NumSamplesDelayE = GetDelaySamples(InSampleRate, 187);
+	Delays.NumSamplesDelayF = GetDelaySamples(InSampleRate, 1041);
+	// Left Delay 2 should add up to 3136 and have tap outs at 1066 (L6) and 2673 (R3)
+	Delays.NumSamplesDelayG = GetDelaySamples(InSampleRate, 1066);
+	Delays.NumSamplesDelayH = GetDelaySamples(InSampleRate, 1607);
+	Delays.NumSamplesDelayI = GetDelaySamples(InSampleRate, 463);
+
+	return Delays;
+}
+
+FLateReflectionsPlateDelays FLateReflectionsPlateDelays::DefaultRightDelays(float InSampleRate)
+{
+	FLateReflectionsPlateDelays Delays;
+
+	Delays.NumSamplesModulatedBase = GetDelaySamples(InSampleRate, 672);
+	Delays.NumSamplesModulatedDelta = GetDelaySamples(InSampleRate, 16);
+	// Right Delay 1 should add up to 4453 and have tap outs at 266 (L0), 2111 (R4), 2974 (L1)
+	Delays.NumSamplesDelayA = GetDelaySamples(InSampleRate, 266);
+	Delays.NumSamplesDelayB = GetDelaySamples(InSampleRate, 1845);
+	Delays.NumSamplesDelayC = GetDelaySamples(InSampleRate, 863);
+	Delays.NumSamplesDelayD = GetDelaySamples(InSampleRate, 1479);
+	Delays.NumSamplesAPF = GetDelaySamples(InSampleRate, 1800);
+	// APF Delay should have tapouts at 335 (R5) and 1913 (L2)
+	Delays.NumSamplesDelayE = GetDelaySamples(InSampleRate, 335);
+	Delays.NumSamplesDelayF = GetDelaySamples(InSampleRate, 1578);
+	// Right Delay 2 should add up to 3720 and have tap outs at 121 (R6) and 1996 (L3)
+	Delays.NumSamplesDelayG = GetDelaySamples(InSampleRate, 121);
+	Delays.NumSamplesDelayH = GetDelaySamples(InSampleRate, 1875);
+	Delays.NumSamplesDelayI = GetDelaySamples(InSampleRate, 1724);
+
+	return Delays;
+}
+
+
+FLateReflectionsPlate::FLateReflectionsPlate(
+			float InSampleRate, 
+			int32 InMaxNumInternalBufferSamples, 
+			const FLateReflectionsPlateDelays& InDelays)
+: 	SampleRate(InSampleRate)
+	, NumInternalBufferSamples(InMaxNumInternalBufferSamples)
+	, Dampening(0.0005f)
+	, Decay(0.50f)
+	, Density(0.0f)
+	, PlateDelays(InDelays)
+{
+	// NumInternalBufferSamples must be less than last delay and should be aligned to SIMD boundaries
+	if (NumInternalBufferSamples > PlateDelays.NumSamplesDelayI) {
+		NumInternalBufferSamples = PlateDelays.NumSamplesDelayI;
+	}
+	NumInternalBufferSamples -= (NumInternalBufferSamples % AUDIO_SIMD_FLOAT_ALIGNMENT);
+
+	checkf(NumInternalBufferSamples > 0, TEXT("InDelays.DelaySample2C too small"));
+
+	// Create the delay line for a plate
+	// ModulatedAPF -> Delay1[A-D] -> LPF -> APF -> Delay2[A-C] -> Output
+	//                                        |-> APFTapDelay[1-2]
+	ModulatedAPF = MakeUnique<FDynamicDelayAPF>(
+			-Density,
+			PlateDelays.NumSamplesModulatedBase - PlateDelays.NumSamplesModulatedDelta,
+			PlateDelays.NumSamplesModulatedBase + PlateDelays.NumSamplesModulatedDelta,
+			NumInternalBufferSamples);
+	DelayA = MakeUnique<FIntegerDelay>(PlateDelays.NumSamplesDelayA, PlateDelays.NumSamplesDelayA);
+	DelayB = MakeUnique<FIntegerDelay>(PlateDelays.NumSamplesDelayB, PlateDelays.NumSamplesDelayB);
+	DelayC = MakeUnique<FIntegerDelay>(PlateDelays.NumSamplesDelayC, PlateDelays.NumSamplesDelayC);
+	DelayD = MakeUnique<FIntegerDelay>(PlateDelays.NumSamplesDelayD, PlateDelays.NumSamplesDelayD);
+	LPF = MakeUnique<FBufferOnePoleLPF>(Dampening);
+	APF = MakeUnique<FLongDelayAPF>(
+			Density - 0.15f,
+			PlateDelays.NumSamplesAPF,
+			NumInternalBufferSamples);
+	DelayE = MakeUnique<FIntegerDelay>(PlateDelays.NumSamplesDelayE, PlateDelays.NumSamplesDelayE);
+	DelayF = MakeUnique<FIntegerDelay>(PlateDelays.NumSamplesDelayF, PlateDelays.NumSamplesDelayF);
+	DelayG = MakeUnique<FIntegerDelay>(PlateDelays.NumSamplesDelayG, PlateDelays.NumSamplesDelayG);
+	DelayH = MakeUnique<FIntegerDelay>(PlateDelays.NumSamplesDelayH, PlateDelays.NumSamplesDelayH);
+	DelayI = MakeUnique<FIntegerDelay>(PlateDelays.NumSamplesDelayI, PlateDelays.NumSamplesDelayI);
+}
+
+FLateReflectionsPlate::~FLateReflectionsPlate()
+{
+}
+
+
+void FLateReflectionsPlate::ProcessAudioFrames(
+		const AlignedFloatBuffer& InSamples,
+		const AlignedFloatBuffer& InFeedbackSamples,
+		const AlignedFloatBuffer& InDelayModulations,
+		FLateReflectionsPlateOutputs& OutPlateSamples)
+{
+	
+	// Prepare output
+	const int32 InNum = InSamples.Num();
+	OutPlateSamples.ResizeAndZero(InNum);
+
+	checkf(InNum == InFeedbackSamples.Num(), TEXT("InSamples and InFeedbackSamples must have equal length"));
+	checkf(InNum == InDelayModulations.Num(), TEXT("InSamples and InDelayModulations must have equal length"));
+
+	// Check buffer sizes and quit early if there is an error.
+	if (InNum != InFeedbackSamples.Num())
+	{
+		return;	
+	}
+	if (InNum != InDelayModulations.Num())
+	{
+		return;
+	}
+
+	float* DataPtr = nullptr;
+
+	// Prepare internal buffers to handle input data.
+	WorkBufferA.Reset(InNum);
+	WorkBufferA.AddUninitialized(InNum);
+	WorkBufferB.Reset(InNum);
+	WorkBufferB.AddUninitialized(InNum);
+	WorkBufferC.Reset(InNum);
+	WorkBufferC.AddUninitialized(InNum);
+
+	// Copy Feedback Samples to aligned internal buffer.
+	BufferWeightedSumFast(InFeedbackSamples, 1.0 - Decay, InSamples, WorkBufferA);
+
+	// Input -> ModulatedAPF
+	ModulatedAPF->ProcessAudio(WorkBufferA, InDelayModulations, WorkBufferB);
+
+	// ModulatedAPF -> Delay1
+	DelayA->ProcessAudio(WorkBufferB, OutPlateSamples.Taps[0]);
+	DelayB->ProcessAudio(OutPlateSamples.Taps[0], OutPlateSamples.Taps[1]);
+	DelayC->ProcessAudio(OutPlateSamples.Taps[1], OutPlateSamples.Taps[2]);
+	DelayD->ProcessAudio(OutPlateSamples.Taps[2], WorkBufferA);
+
+	// Apply dampening
+	MultiplyBufferByConstantInPlace(WorkBufferA, 1.0f - Dampening);
+	
+	// Delay1 -> LPF
+	LPF->ProcessAudio(WorkBufferA, WorkBufferB);
+
+	// Apply decay
+	MultiplyBufferByConstantInPlace(WorkBufferB, 1.0f - Decay);
+
+	// LPF -> APF
+	APF->ProcessAudio(WorkBufferB, WorkBufferA, WorkBufferC);
+	DelayE->ProcessAudio(WorkBufferC, OutPlateSamples.Taps[3]);
+	DelayF->ProcessAudio(OutPlateSamples.Taps[3], OutPlateSamples.Taps[4]);
+
+	// APF-> Delay2
+	DelayG->ProcessAudio(WorkBufferA, OutPlateSamples.Taps[5]);
+	DelayH->ProcessAudio(OutPlateSamples.Taps[5], OutPlateSamples.Taps[6]);
+	DelayI->ProcessAudio(OutPlateSamples.Taps[6], OutPlateSamples.Output);
+}
+
+void FLateReflectionsPlate::SetDampening(float InDampening)
+{
+	Dampening = InDampening;
+	LPF->SetG(Dampening);
+}
+
+void FLateReflectionsPlate::SetDecay(float InDecay)
+{
+	Decay = InDecay;
+}
+
+void FLateReflectionsPlate::SetDensity(float InDensity)
+{
+	Density = InDensity;
+	ModulatedAPF->SetG(-Density);
+	APF->SetG(Density - 0.15f);
+}
+
+int32 FLateReflectionsPlate::GetNumInternalBufferSamples() const
+{
+	return NumInternalBufferSamples;
+}
+
+void FLateReflectionsPlate::PeekDelayLine(int32 InNum, AlignedFloatBuffer& OutSamples)
+{
+	DelayI->PeekDelayLine(InNum, OutSamples);
+}
+
+// Limits on late reflections settings.
+const float FLateReflectionsFast::MaxLateDelayMsec = 2000.0f;
+const float FLateReflectionsFast::MinLateDelayMsec = 0.0f;
+const float FLateReflectionsFast::MaxLateGainDB = 0.0f;
+const float FLateReflectionsFast::MaxBandwidth = 0.99999f;
+const float FLateReflectionsFast::MinBandwidth = 0.0f;
+const float FLateReflectionsFast::MaxDampening = 0.99999f;
+const float FLateReflectionsFast::MinDampening = 0.0f;
+const float FLateReflectionsFast::MaxDiffusion = 0.99999f;
+const float FLateReflectionsFast::MinDiffusion = 0.0f;
+const float FLateReflectionsFast::MaxDecay = 1.0f;
+const float FLateReflectionsFast::MinDecay = 0.0001f;
+const float FLateReflectionsFast::MaxDensity = 1.0f;
+const float FLateReflectionsFast::MinDensity = 0.0f;
+
+const FLateReflectionsFastSettings FLateReflectionsFast::DefaultSettings;
+
+FLateReflectionsFast::FLateReflectionsFast(float InSampleRate, int32 InMaxNumInternalBufferSamples, const FLateReflectionsFastSettings& InSettings)
+:	SampleRate(InSampleRate)
+	, Gain(1.0f)
+	, ModulationPhase(0.0f)
+	, ModulationQuadPhase(PI / 2.0f)
+	, ModulationPhaseIncrement(0.0f)
+	, NumInternalBufferSamples(0)
+	, Settings(InSettings)
+	, LeftPlateDelays(FLateReflectionsPlateDelays::DefaultLeftDelays(InSampleRate))
+	, RightPlateDelays(FLateReflectionsPlateDelays::DefaultRightDelays(InSampleRate))
+
+{
+	// Copy and clamp settings.
+	Settings = InSettings;
+	ClampSettings(Settings);
+	
+	// Create plates. 
+	LeftPlate = MakeUnique<FLateReflectionsPlate>(SampleRate, InMaxNumInternalBufferSamples, LeftPlateDelays);
+	RightPlate = MakeUnique<FLateReflectionsPlate>(SampleRate, InMaxNumInternalBufferSamples, RightPlateDelays);
+
+	// The block size is limited by the delay lengths in the left & right plates.
+	// Internal block size will be set to minimum of allowable sizes from plates.
+	NumInternalBufferSamples = FMath::Min(LeftPlate->GetNumInternalBufferSamples(), RightPlate->GetNumInternalBufferSamples());
+
+	// Set maximum predelay. Add 8 samples to give a little extra buffer for floating point rounding
+	// differences. 
+	int32 MaxDelay = (int32)(SampleRate * FLateReflectionsFast::MaxLateDelayMsec / 1000.0f) + 8;
+	PreDelay = MakeUnique<FIntegerDelay>(MaxDelay, 0, NumInternalBufferSamples);
+
+	InputLPF = MakeUnique<FBufferOnePoleLPF>(1.0f - Settings.Bandwidth);
+
+	// make the signal decorrelators
+	APF1 = MakeUnique<FLongDelayAPF>(Settings.Diffusion, GetDelaySamples(SampleRate, 142), NumInternalBufferSamples);
+	APF2 = MakeUnique<FLongDelayAPF>(Settings.Diffusion, GetDelaySamples(SampleRate, 107), NumInternalBufferSamples);
+	APF3 = MakeUnique<FLongDelayAPF>(Settings.Diffusion - 0.125f, GetDelaySamples(SampleRate, 379), NumInternalBufferSamples);
+	APF4 = MakeUnique<FLongDelayAPF>(Settings.Diffusion - 0.125f, GetDelaySamples(SampleRate, 277), NumInternalBufferSamples);
+
+	// Set modulation rate to 1 Hz. 
+	ModulationPhaseIncrement = 2.0 * PI / SampleRate;
+	
+	// Resize internal buffers
+	WorkBufferA.Reset(NumInternalBufferSamples);
+	WorkBufferB.Reset(NumInternalBufferSamples);
+	WorkBufferC.Reset(NumInternalBufferSamples);
+	LeftDelayModSamples.Reset(NumInternalBufferSamples);
+	RightDelayModSamples.Reset(NumInternalBufferSamples);
+
+	WorkBufferA.AddUninitialized(NumInternalBufferSamples);
+	WorkBufferB.AddUninitialized(NumInternalBufferSamples);
+	WorkBufferC.AddUninitialized(NumInternalBufferSamples);
+	LeftDelayModSamples.AddUninitialized(NumInternalBufferSamples);
+	RightDelayModSamples.AddUninitialized(NumInternalBufferSamples);
+
+	LeftPlateOutputs.ResizeAndZero(NumInternalBufferSamples);
+	RightPlateOutputs.ResizeAndZero(NumInternalBufferSamples);
+
+	LeftPlateOutputBuffer = MakeUnique<FAlignedBlockBuffer>(4 * NumInternalBufferSamples, NumInternalBufferSamples);
+	RightPlateOutputBuffer = MakeUnique<FAlignedBlockBuffer>(4 * NumInternalBufferSamples, NumInternalBufferSamples);
+
+	ApplySettings();
+}
+
+FLateReflectionsFast::~FLateReflectionsFast()
+{}
+
+void FLateReflectionsFast::ClampSettings(FLateReflectionsFastSettings& InOutSettings)
+{
+	// Enforce settings to be within max/min
+	InOutSettings.LateDelayMsec = FMath::Clamp(InOutSettings.LateDelayMsec, MinLateDelayMsec, MaxLateDelayMsec);
+	InOutSettings.LateGainDB = FMath::Min(InOutSettings.LateGainDB, FLateReflectionsFast::MaxLateGainDB);
+	InOutSettings.Bandwidth = FMath::Clamp(InOutSettings.Bandwidth, MinBandwidth, MaxBandwidth);
+	InOutSettings.Dampening = FMath::Clamp(InOutSettings.Dampening, MinDampening, MaxDampening);
+	InOutSettings.Diffusion = FMath::Clamp(InOutSettings.Diffusion, MinDiffusion, MaxDiffusion);
+	InOutSettings.Decay = FMath::Clamp(InOutSettings.Decay, MinDecay, MaxDecay);
+	InOutSettings.Density = FMath::Clamp(InOutSettings.Density, MinDensity, MaxDensity);
+}
+
+// Sets the reverb settings, applies, and updates
+void FLateReflectionsFast::SetSettings(const FLateReflectionsFastSettings& InSettings)
+{
+	Settings = InSettings;
+	ClampSettings(Settings);
+	ApplySettings();
+}
+
+void FLateReflectionsFast::ProcessAudio(const AlignedFloatBuffer& InSamples, const int32 InNumChannels, AlignedFloatBuffer& OutLeftSamples, AlignedFloatBuffer& OutRightSamples)
+{
+	checkf((InNumChannels == 1) || (InNumChannels == 2), TEXT("FLateReflections only supports 1 or 2 channel input audio"));
+
+	const int32 InNum = InSamples.Num();
+	const int32 InNumFrames = InNum / InNumChannels;
+	
+	OutLeftSamples.Reset(InNumFrames);
+	OutRightSamples.Reset(InNumFrames);
+
+	OutLeftSamples.AddUninitialized(InNumFrames);
+	OutRightSamples.AddUninitialized(InNumFrames);
+
+	const float* InSampleData = InSamples.GetData();
+	float* OutLeftSampleData = OutLeftSamples.GetData();
+	float* OutRightSampleData = OutRightSamples.GetData();
+
+	int32 FramesLeftOver = InNumFrames;
+	int32 InPos = 0;
+	int32 OutPos = 0;
+	while (FramesLeftOver > 0)
+	{
+		int32 FramesToProcess = FMath::Min(NumInternalBufferSamples, FramesLeftOver);
+
+		ProcessAudioBuffer(&InSampleData[InPos], FramesToProcess, InNumChannels, &OutLeftSampleData[OutPos], &OutRightSampleData[OutPos]);
+
+		FramesLeftOver -= FramesToProcess;
+		InPos += FramesToProcess * InNumChannels;
+		OutPos += FramesToProcess;
+	}
+}
+
+
+void FLateReflectionsFast::ProcessAudioBuffer(const float* InSampleData, const int32 InNumFrames, const int32 InNumChannels, float* OutLeftSampleData, float* OutRightSampleData)
+{
+	checkf((1 == InNumChannels) || (2 == InNumChannels), TEXT("FLateReflections only supports 1 and 2 channel input audio"));
+
+	// Resize internal buffers
+	WorkBufferA.Reset(InNumFrames);
+	WorkBufferB.Reset(InNumFrames);
+	LeftDelayModSamples.Reset(InNumFrames);
+	RightDelayModSamples.Reset(InNumFrames);
+
+	WorkBufferA.AddUninitialized(InNumFrames);
+	WorkBufferB.AddUninitialized(InNumFrames);
+	LeftDelayModSamples.AddUninitialized(InNumFrames);
+	RightDelayModSamples.AddUninitialized(InNumFrames);
+
+	// Downmix to mono
+	if (InNumChannels == 1)
+	{
+		FMemory::Memcpy(WorkBufferA.GetData(), InSampleData, sizeof(float) * InNumFrames);
+	}
+	else if (InNumChannels == 2)
+	{
+		BufferSum2ChannelToMonoFast(InSampleData, WorkBufferA.GetData(), InNumFrames);
+	}
+	else
+	{
+		// Shouldn't reach this line if checks are respected. 
+		FMemory::Memset(OutLeftSampleData, 0, sizeof(float) * InNumFrames);
+		FMemory::Memset(OutRightSampleData, 0, sizeof(float) * InNumFrames);
+		return;
+	}
+
+	// Apply bandwidth, gain and channel averaging multipliers
+	MultiplyBufferByConstantInPlace(WorkBufferA, Settings.Bandwidth * Gain / InNumChannels);
+	
+	// Predelay
+	PreDelay->ProcessAudio(WorkBufferA, WorkBufferB);
+
+	// Input Diffusion 
+	InputLPF->ProcessAudio(WorkBufferB, WorkBufferA);
+	APF1->ProcessAudio(WorkBufferA, WorkBufferB);
+	APF2->ProcessAudio(WorkBufferB, WorkBufferA);
+	APF3->ProcessAudio(WorkBufferA, WorkBufferB);
+	APF4->ProcessAudio(WorkBufferB, WorkBufferA);
+	
+	// The plates are connected to eachother so that the output of the left plate is fed into the input of the
+	// right plate and vice versa. The delay lines of the plates are stored before processing new data through
+	// the plate.
+	LeftPlate->PeekDelayLine(InNumFrames, WorkBufferB);
+	RightPlate->PeekDelayLine(InNumFrames, WorkBufferC);
+	
+	// Create the delay line modulations.
+	GeneraterPlateModulations(InNumFrames, LeftDelayModSamples, RightDelayModSamples);
+
+	// Run audio through plates
+	LeftPlate->ProcessAudioFrames(WorkBufferA, WorkBufferC, LeftDelayModSamples, LeftPlateOutputs);
+	RightPlate->ProcessAudioFrames(WorkBufferA, WorkBufferB, RightDelayModSamples, RightPlateOutputs);
+
+	// Left Output
+	// Channel -> [Plate and Tap#]
+	// L -> +R0
+	// L -> +R2
+	// L -> -R4
+	// L -> +R6
+	// L -> -L1
+	// L -> -L3
+	// L -> -L5
+	FMemory::Memcpy(OutLeftSampleData, RightPlateOutputs.Taps[0].GetData(), InNumFrames * sizeof(float));
+	MixInBufferFast(RightPlateOutputs.Taps[2].GetData(), OutLeftSampleData, InNumFrames);
+	BufferSubtractInPlace2Fast(OutLeftSampleData, RightPlateOutputs.Taps[4].GetData(), InNumFrames);
+	MixInBufferFast(RightPlateOutputs.Taps[6].GetData(), OutLeftSampleData, InNumFrames);
+	BufferSubtractInPlace2Fast(OutLeftSampleData, LeftPlateOutputs.Taps[1].GetData(), InNumFrames);
+	BufferSubtractInPlace2Fast(OutLeftSampleData, LeftPlateOutputs.Taps[3].GetData(), InNumFrames);
+	BufferSubtractInPlace2Fast(OutLeftSampleData, LeftPlateOutputs.Taps[5].GetData(), InNumFrames);
+	
+	// Right Output
+	// Channel -> [Plate and Tap#]
+	// R -> +L0
+	// R -> +L2
+	// R -> -L4
+	// R -> +L6
+	// R -> -R1
+	// R -> -R3
+	// R -> -R5
+	FMemory::Memcpy(OutRightSampleData, LeftPlateOutputs.Taps[0].GetData(), InNumFrames * sizeof(float));
+	MixInBufferFast(LeftPlateOutputs.Taps[2].GetData(), OutRightSampleData, InNumFrames);
+	BufferSubtractInPlace2Fast(OutRightSampleData, LeftPlateOutputs.Taps[4].GetData(), InNumFrames);
+	MixInBufferFast(LeftPlateOutputs.Taps[6].GetData(), OutRightSampleData, InNumFrames);
+	BufferSubtractInPlace2Fast(OutRightSampleData, RightPlateOutputs.Taps[1].GetData(), InNumFrames);
+	BufferSubtractInPlace2Fast(OutRightSampleData, RightPlateOutputs.Taps[3].GetData(), InNumFrames);
+	BufferSubtractInPlace2Fast(OutRightSampleData, RightPlateOutputs.Taps[5].GetData(), InNumFrames);
+}
+
+void FLateReflectionsFast::ApplySettings()
+{
+	// Convert gain to linear. 
+	Gain = FMath::Pow(10.0f, Settings.LateGainDB/ 20.0f);
+
+	// Convert delay to samples
+	PreDelay->SetDelayLengthSamples((int32)(Settings.LateDelayMsec * SampleRate / 1000.0f));
+
+	InputLPF->SetG(1.0f - Settings.Bandwidth);
+
+	APF1->SetG(Settings.Diffusion);
+	APF2->SetG(Settings.Diffusion);
+	APF3->SetG(Settings.Diffusion - 0.125f);
+	APF4->SetG(Settings.Diffusion - 0.125f);
+
+	LeftPlate->SetDensity(Settings.Density);
+	LeftPlate->SetDampening(Settings.Dampening);
+	LeftPlate->SetDecay(Settings.Decay);
+
+	RightPlate->SetDensity(Settings.Density);
+	RightPlate->SetDampening(Settings.Dampening);
+	RightPlate->SetDecay(Settings.Decay);
+}
+
+
+void FLateReflectionsFast::GeneraterPlateModulations(const int32 InNum, AlignedFloatBuffer& OutLeftDelays, AlignedFloatBuffer& OutRightDelays)
+{
+	OutLeftDelays.Reset(InNum);
+	OutRightDelays.Reset(InNum);
+	OutLeftDelays.AddUninitialized(InNum);
+	OutRightDelays.AddUninitialized(InNum);
+
+	// Generate quadrature plate modulations to help break up hi-freq modes in plates. 
+	float* LeftModData = OutLeftDelays.GetData();
+	float* RightModData = OutRightDelays.GetData();
+	for (int32 i = 0; i < InNum; i++)
+	{		
+		float NormalPhaseOut = 0.5f * FastSin(ModulationPhase) + 0.5f;
+		float QuadPhaseOut = 0.5f * FastSin(ModulationQuadPhase) + 0.5f;
+		ModulationPhase += ModulationPhaseIncrement;
+
+		if (ModulationPhase > PI)
+		{
+			ModulationPhase -= 2 * PI;
+		}
+
+		ModulationQuadPhase += ModulationPhaseIncrement;
+		if (ModulationQuadPhase > PI)
+		{
+			ModulationQuadPhase -= 2 * PI;
+		}
+		
+		LeftModData[i] = LeftPlateDelays.NumSamplesModulatedBase + (NormalPhaseOut * LeftPlateDelays.NumSamplesModulatedDelta);
+		RightModData[i] = RightPlateDelays.NumSamplesModulatedBase + (QuadPhaseOut * RightPlateDelays.NumSamplesModulatedDelta);
+	}
+}
diff --git a/Engine/Source/Runtime/AudioMixer/Private/DSP/LinearInterpFractionalDelay.cpp b/Engine/Source/Runtime/AudioMixer/Private/DSP/LinearInterpFractionalDelay.cpp
new file mode 100644
--- /dev/null
+++ b/Engine/Source/Runtime/AudioMixer/Private/DSP/LinearInterpFractionalDelay.cpp
@@ -0,0 +1,161 @@
+// Copyright 1998-2019 Epic Games, Inc. All Rights Reserved.
+
+#include "DSP/LinearInterpFractionalDelay.h"
+#include "DSP/Dsp.h"
+#include "DSP/BufferVectorOperations.h"
+
+using namespace Audio;
+
+FLinearInterpFractionalDelay::FLinearInterpFractionalDelay(int32 InMaxDelay, int32 InMaxNumInternalBufferSamples)
+: MaxDelay(InMaxDelay)
+, NumInternalBufferSamples(InMaxNumInternalBufferSamples)
+, UpperDelayPos(nullptr)
+, LowerDelayPos(nullptr)
+{
+	checkf(MaxDelay > 0, TEXT("InMaxDelay must be greater than zero"));
+	if (MaxDelay < 1)
+	{
+		MaxDelay = 1;
+	}
+
+	while (0 != (NumInternalBufferSamples % AUDIO_SIMD_FLOAT_ALIGNMENT))
+	{
+		NumInternalBufferSamples--;
+	}
+	if (NumInternalBufferSamples < 1)
+	{
+		NumInternalBufferSamples = AUDIO_SIMD_FLOAT_ALIGNMENT;
+	}
+
+	// Allocate and prepare delay line for maximum delay.
+	DelayLine = MakeUnique<FAlignedBlockBuffer>((2 * MaxDelay) + NumInternalBufferSamples, MaxDelay + NumInternalBufferSamples);
+	DelayLine->AddZeros(MaxDelay);
+	
+	IntegerDelayOffsets.Reset(NumInternalBufferSamples);
+	IntegerDelayOffsets.AddUninitialized(NumInternalBufferSamples);
+
+	for (int32 i = 0; i < NumInternalBufferSamples; i++)
+	{
+		IntegerDelayOffsets[i] = i + MaxDelay;
+	}
+
+	UpperDelayPos = (int*)FMemory::Malloc(4 * sizeof(int), AUDIO_BUFFER_ALIGNMENT);
+	LowerDelayPos = (int*)FMemory::Malloc(4 * sizeof(int), AUDIO_BUFFER_ALIGNMENT);
+}
+
+
+// Destructor
+FLinearInterpFractionalDelay::~FLinearInterpFractionalDelay()
+{
+	FMemory::Free(UpperDelayPos);
+	FMemory::Free(LowerDelayPos);
+}
+
+
+// Resets the delay line state, flushes buffer and resets read/write pointers.
+void FLinearInterpFractionalDelay::Reset() 
+{
+	DelayLine->ClearSamples();
+	DelayLine->AddZeros(MaxDelay);
+}
+
+
+void FLinearInterpFractionalDelay::ProcessAudio(const AlignedFloatBuffer& InSamples, const AlignedFloatBuffer& InDelays, AlignedFloatBuffer& OutSamples)
+{
+	const int32 InNum = InSamples.Num();
+	checkf(InNum == InDelays.Num(), TEXT("Input buffers must be equal length"));
+
+	// Prepare output buffer
+	OutSamples.Reset(InNum);
+	OutSamples.AddUninitialized(InNum);
+
+	if (InNum != InDelays.Num())
+	{
+		// Return empty buffer on invalid input.
+		if (InNum > 0)
+		{
+			FMemory::Memset(OutSamples.GetData(), 0, sizeof(float) * InNum);
+		}
+		return;
+	}
+
+
+	float* OutSampleData = OutSamples.GetData();
+	const float* InSampleData = InSamples.GetData();
+	const float* InDelayData = InDelays.GetData();
+
+	// Process audio one block at a time.
+	int32 LeftOver = InNum;
+	int32 BufferPos = 0;
+	while (LeftOver > 0)
+	{
+		int32 NumToProcess = FMath::Min(LeftOver, NumInternalBufferSamples);
+		ProcessAudioBlock(&InSampleData[BufferPos], &InDelayData[BufferPos], NumToProcess, &OutSampleData[BufferPos]);
+		BufferPos += NumToProcess;
+		LeftOver -= NumToProcess;
+	}
+}
+
+void FLinearInterpFractionalDelay::ProcessAudioBlock(const float* InSamples, const float* InDelays, const int32 InNum, float* OutSamples)
+{
+	checkf(0 == (InNum % 4), TEXT("Array length must be multiple of 4"));
+	checkf(IsAligned<const float*>(InSamples, AUDIO_SIMD_FLOAT_ALIGNMENT), TEXT("Memory must be aligned to use vector operations."));
+	checkf(IsAligned<const float*>(InDelays, AUDIO_SIMD_FLOAT_ALIGNMENT), TEXT("Memory must be aligned to use vector operations."));
+	checkf(IsAligned<float*>(OutSamples, AUDIO_SIMD_FLOAT_ALIGNMENT), TEXT("Memory must be aligned to use vector operations."));
+
+
+	// Update delay line.
+	DelayLine->AddSamples(InSamples, InNum);
+
+	const float* DelayData = DelayLine->InspectSamples(InNum + MaxDelay);
+	const int32* IntegerDelayOffsetData = IntegerDelayOffsets.GetData();
+
+	const VectorRegister VMaxDelay = MakeVectorRegister((float)MaxDelay, (float)MaxDelay, (float)MaxDelay, (float)MaxDelay);
+	for (int32 i = 0; i < InNum; i += 4)
+	{
+		VectorRegister VFractionalDelays = VectorLoadAligned(&InDelays[i]);
+		// Ensure fractional delays are positive
+		VFractionalDelays = VectorMax(VFractionalDelays, GlobalVectorConstants::FloatZero);
+		VFractionalDelays = VectorMin(VFractionalDelays, VMaxDelay);
+
+		// Separate integer from fraction
+		VectorRegister VFloorDelays = VectorFloor(VFractionalDelays);
+
+		// Determine linear weights
+		VectorRegister VUpperCoefficients = VectorSubtract(VFractionalDelays, VFloorDelays);
+		VectorRegister VLowerCoefficients = VectorSubtract(GlobalVectorConstants::FloatOne, VUpperCoefficients);
+
+
+		// Make integer locations relative to bblock
+		VectorRegisterInt VIntegerDelays = VectorFloatToInt(VFloorDelays);
+		VectorRegisterInt VIntegerDelayOffset = VectorIntLoadAligned(&IntegerDelayOffsetData[i]);
+		VIntegerDelays = VectorIntSubtract(VIntegerDelayOffset, VIntegerDelays);
+
+		// Lookup samples for interpolation
+		VectorIntStoreAligned(VIntegerDelays, LowerDelayPos);
+		VectorIntStoreAligned(VectorIntSubtract(VIntegerDelays, GlobalVectorConstants::IntOne), UpperDelayPos);
+		VectorRegister VLowerSamples = MakeVectorRegister(
+			DelayData[LowerDelayPos[0]],
+			DelayData[LowerDelayPos[1]],
+			DelayData[LowerDelayPos[2]],
+			DelayData[LowerDelayPos[3]]
+		);
+		VectorRegister VUpperSamples = MakeVectorRegister(
+			DelayData[UpperDelayPos[0]],
+			DelayData[UpperDelayPos[1]],
+			DelayData[UpperDelayPos[2]],
+			DelayData[UpperDelayPos[3]]
+		);
+
+		// Interpolate samples
+		VectorRegister VOut = VectorMultiplyAdd(
+			VLowerSamples,
+			VLowerCoefficients,
+			VectorMultiply(VUpperSamples, VUpperCoefficients));
+		VectorStoreAligned(VOut, &OutSamples[i]);
+	}
+
+	// Remove unneeded delay line.
+	DelayLine->RemoveSamples(InNum);
+}
+
diff --git a/Engine/Source/Runtime/AudioMixer/Private/DSP/LongDelayAPF.cpp b/Engine/Source/Runtime/AudioMixer/Private/DSP/LongDelayAPF.cpp
new file mode 100644
--- /dev/null
+++ b/Engine/Source/Runtime/AudioMixer/Private/DSP/LongDelayAPF.cpp
@@ -0,0 +1,152 @@
+// Copyright 1998-2019 Epic Games, Inc. All Rights Reserved.
+
+#include "DSP/LongDelayAPF.h"
+#include "DSP/BufferVectorOperations.h"
+#include "DSP/Dsp.h"
+
+using namespace Audio;
+
+FLongDelayAPF::FLongDelayAPF(float InG, int32 InNumDelaySamples, int32 InMaxNumInternalBufferSamples)
+:	G(InG),
+	NumDelaySamples(InNumDelaySamples),
+	NumInternalBufferSamples(InMaxNumInternalBufferSamples)
+{
+	// NumInternalBufferSamples must be less than or equal to the length of the delay
+	if (NumInternalBufferSamples > NumDelaySamples)
+	{
+		NumInternalBufferSamples = NumDelaySamples;
+
+		// Block size must be divisible by simd alignment to support simd operations.
+		NumInternalBufferSamples -= (NumInternalBufferSamples % AUDIO_SIMD_FLOAT_ALIGNMENT);
+	}
+
+	checkf(NumInternalBufferSamples > 1, TEXT("Invalid internal buffer length"));
+
+	// Allocate delay line
+	DelayLine = MakeUnique<FAlignedBlockBuffer>(2 * NumInternalBufferSamples + NumDelaySamples, NumInternalBufferSamples);
+	DelayLine->AddZeros(NumDelaySamples);
+
+	// Allocate internal buffer
+	WorkBuffer.Reset(NumInternalBufferSamples);
+	WorkBuffer.AddUninitialized(NumInternalBufferSamples);
+}
+
+FLongDelayAPF::~FLongDelayAPF()
+{}
+
+void FLongDelayAPF::ProcessAudio(const AlignedFloatBuffer& InSamples, AlignedFloatBuffer& OutSamples)
+{
+	const float* InData = InSamples.GetData();
+	const int32 InNum = InSamples.Num();
+	float* OutDelayData = WorkBuffer.GetData();
+
+	// Resize output
+	OutSamples.Reset(InNum);
+	OutSamples.AddUninitialized(InNum);
+	float* OutData = OutSamples.GetData();
+
+
+	// Process audio one block at a time.
+	int32 LeftOver = InNum;
+	int32 BufferIndex = 0;
+	while (LeftOver != 0) 
+	{
+		// Determine block size for this loop.
+		int32 NumToProcess = FMath::Min<int32>(NumInternalBufferSamples, LeftOver);
+		const float* InDelayData = DelayLine->InspectSamples(NumToProcess);
+
+		ProcessAudioBlock(&InData[BufferIndex], InDelayData, NumToProcess, &OutData[BufferIndex], OutDelayData);
+
+		// Update delay line with new data.
+		DelayLine->RemoveSamples(NumToProcess);
+		DelayLine->AddSamples(OutDelayData, NumToProcess);
+
+		LeftOver -= NumToProcess;
+		BufferIndex += NumToProcess;
+	}
+}
+
+void FLongDelayAPF::ProcessAudio(const AlignedFloatBuffer& InSamples, AlignedFloatBuffer& OutSamples, AlignedFloatBuffer& OutDelaySamples)
+{
+	const float* InData = InSamples.GetData();
+	const int32 InNum = InSamples.Num();
+
+	// Resize output buffers
+	OutSamples.Reset(InNum);
+	OutSamples.AddUninitialized(InNum);
+	OutDelaySamples.Reset(InNum);
+	OutDelaySamples.AddUninitialized(InNum);
+
+	float* OutData = OutSamples.GetData();
+	float* OutDelayData = OutDelaySamples.GetData();
+
+	// Process audio one block at a time.
+	int32 LeftOver = InNum;
+	int32 BufferIndex = 0;
+	while (LeftOver != 0) 
+	{
+		int32 NumToProcess = FMath::Min<int32>(NumInternalBufferSamples, LeftOver);
+		const float* InDelayData = DelayLine->InspectSamples(NumToProcess);
+
+		ProcessAudioBlock(&InData[BufferIndex], InDelayData, NumToProcess, &OutData[BufferIndex], &OutDelayData[BufferIndex]);
+		
+		// Update internal delay line.
+		DelayLine->RemoveSamples(NumToProcess);
+		DelayLine->AddSamples(&OutDelayData[BufferIndex], NumToProcess);
+
+		LeftOver -= NumToProcess;
+		BufferIndex += NumToProcess;
+	}
+}
+
+void FLongDelayAPF::ProcessAudioBlock(const float* InSamples, const float* InDelaySamples, const int32 InNum, float* OutSamples, float* OutDelaySamples)
+{
+	// Calculate new delay line samples. "w[n] = x[n] + gw[n - d]"
+	//WeightedSumBuffer(InDelaySamples, G, InSamples, OutDelaySamples, InNum);
+	
+	
+	int32 NumToSIMD = InNum - (InNum % AUDIO_SIMD_FLOAT_ALIGNMENT);
+
+	VectorRegister VG = MakeVectorRegister(G, G, G, G);
+	VectorRegister VNG = MakeVectorRegister(-G, -G, -G, -G);
+	VectorRegister VFMIN = MakeVectorRegister(FLT_MIN, FLT_MIN, FLT_MIN, FLT_MIN);
+	VectorRegister VNFMIN = MakeVectorRegister(-FLT_MIN, -FLT_MIN, -FLT_MIN, -FLT_MIN);
+
+	for (int32 i = 0; i < InNum; i += 4)
+	{
+		VectorRegister VInDelay = VectorLoadAligned(&InDelaySamples[i]);
+		VectorRegister VInSamples = VectorLoadAligned(&InSamples[i]);
+		// w[n] = x[n] + G * w[n - D]
+		VectorRegister VOutDelay = VectorMultiplyAdd(VInDelay, VG, VInSamples);
+		
+		// Underflow clamp
+		VectorRegister Mask = VectorBitwiseAnd(VectorCompareGT(VOutDelay, VNFMIN), VectorCompareLT(VOutDelay, VFMIN));
+		VOutDelay = VectorSelect(Mask, GlobalVectorConstants::FloatZero, VOutDelay);
+		VectorStoreAligned(VOutDelay, &OutDelaySamples[i]);
+
+		// y[n] = -G * w[n] + w[n - D]
+		VectorRegister VOut = VectorMultiplyAdd(VOutDelay, VNG, VInDelay);
+		VectorStoreAligned(VOut, &OutSamples[i]);
+
+
+	}
+
+	// Calculate allpass for remaining samples that we couldn't SIMD
+	for (int32 i = NumToSIMD; i < InNum; i++)
+	{
+		OutDelaySamples[i] = Audio::UnderflowClamp(InDelaySamples[i] * G + InSamples[i]);
+		OutSamples[i] = OutDelaySamples[i] * -G + InDelaySamples[i];
+	}
+}
+
+void FLongDelayAPF::Reset() 
+{
+	DelayLine->ClearSamples();
+	DelayLine->AddZeros(NumDelaySamples);
+}
+
+int32 FLongDelayAPF::GetNumInternalBufferSamples() const
+{
+	return NumInternalBufferSamples;
+}
+
diff --git a/Engine/Source/Runtime/AudioMixer/Private/DSP/ReverbFast.cpp b/Engine/Source/Runtime/AudioMixer/Private/DSP/ReverbFast.cpp
new file mode 100644
--- /dev/null
+++ b/Engine/Source/Runtime/AudioMixer/Private/DSP/ReverbFast.cpp
@@ -0,0 +1,294 @@
+
+#include "DSP/ReverbFast.h"
+#include "DSP/BufferVectorOperations.h"
+
+namespace Audio {
+	FPlateReverbFastSettings::FPlateReverbFastSettings()
+		: Wetness(0.5f)
+		, QuadBehavior(EQuadBehavior::StereoOnly)
+	{}
+
+	bool FPlateReverbFastSettings::operator==(const FPlateReverbFastSettings& Other) const
+	{
+		bool bIsEqual = (
+			(Other.EarlyReflections == EarlyReflections) &&
+			(Other.LateReflections == LateReflections) &&
+			(Other.Wetness == Wetness) &&
+			(Other.QuadBehavior == QuadBehavior));
+		
+		return bIsEqual;
+	}
+
+	bool FPlateReverbFastSettings::operator!=(const FPlateReverbFastSettings& Other) const
+	{
+		return !(*this == Other);
+	}
+
+
+	const float FPlateReverbFast::MaxWetness = 10.f;
+	const float FPlateReverbFast::MinWetness = 0.0f;
+
+	const FPlateReverbFastSettings FPlateReverbFast::DefaultSettings;
+
+	FPlateReverbFast::FPlateReverbFast(float InSampleRate, int32 InMaxInternalBufferSamples, const FPlateReverbFastSettings& InSettings)
+		: SampleRate(InSampleRate)
+		, EarlyReflections(InSampleRate, InMaxInternalBufferSamples)
+		, LateReflections(InSampleRate, InMaxInternalBufferSamples, InSettings.LateReflections)
+		, bEnableEarlyReflections(true)
+		, bEnableLateReflections(true)
+	{
+		SetSettings(InSettings);
+	}
+
+	FPlateReverbFast::~FPlateReverbFast()
+	{}
+
+	void FPlateReverbFast::SetSettings(const FPlateReverbFastSettings& InSettings)
+	{
+		// Copy, clamp and apply settings
+		Settings = InSettings;
+		ClampSettings(Settings);
+		ApplySettings();
+	}
+
+	const FPlateReverbFastSettings& FPlateReverbFast::GetSettings() const
+	{
+		return Settings;
+	}
+
+	// Whether or not to enable late reflections
+	void FPlateReverbFast::EnableLateReflections(const bool bInEnableLateReflections)
+	{
+		bEnableLateReflections = bInEnableLateReflections;
+	}
+
+	// Whether or not to enable late reflections
+	void FPlateReverbFast::EnableEarlyReflections(const bool bInEnableEarlyReflections)
+	{
+		bEnableEarlyReflections = bInEnableEarlyReflections;
+	}
+
+	// Process a buffer of input audio samples.
+	void FPlateReverbFast::ProcessAudio(const AlignedFloatBuffer& InSamples, const int32 InNumChannels, AlignedFloatBuffer& OutSamples, const int32 OutNumChannels)
+	{
+		checkf((1 == InNumChannels) || (2 == InNumChannels), TEXT("FPlateReverbFast only supports 1 or 2 channel inputs."))
+		checkf(OutNumChannels >= 2, TEXT("FPlateReverbFast requires at least 2 output channels."))
+
+		// Determine number of frames and output size.
+		const int32 InNum = InSamples.Num();
+		const int32 InNumFrames = InNum / InNumChannels;
+
+		if (!bEnableEarlyReflections && !bEnableLateReflections)
+		{
+			// Zero output buffers if all reverb is disabled. 
+			const int32 OutNum = InNumFrames * OutNumChannels;
+
+			OutSamples.Reset(OutNum);
+			OutSamples.AddUninitialized(OutNum);
+
+			FMemory::Memset(OutSamples.GetData(), 0, sizeof(float) * OutNum);
+			return;
+		}
+
+		// Resize internal buffers	
+		FrontLeftReverbSamples.Reset(InNumFrames);
+		FrontRightReverbSamples.Reset(InNumFrames);
+
+		FrontLeftReverbSamples.AddUninitialized(InNumFrames);
+		FrontRightReverbSamples.AddUninitialized(InNumFrames);
+
+		if (bEnableEarlyReflections && !bEnableLateReflections)
+		{
+			// Only generate early reflections.
+			EarlyReflections.ProcessAudio(InSamples, InNumChannels, FrontLeftReverbSamples, FrontRightReverbSamples);
+		}
+		else if (!bEnableEarlyReflections && bEnableLateReflections)
+		{
+			// Only generate late reflections.
+			LateReflections.ProcessAudio(InSamples, InNumChannels, FrontLeftReverbSamples, FrontRightReverbSamples);
+		}
+		else if (bEnableEarlyReflections && bEnableLateReflections)
+		{
+			// Resize internal buffers
+			FrontLeftLateReflectionsSamples.Reset(InNumFrames);
+			FrontRightLateReflectionsSamples.Reset(InNumFrames);
+			FrontLeftEarlyReflectionsSamples.Reset(InNumFrames);
+			FrontRightEarlyReflectionsSamples.Reset(InNumFrames);
+
+			FrontLeftLateReflectionsSamples.AddUninitialized(InNumFrames);
+			FrontRightLateReflectionsSamples.AddUninitialized(InNumFrames);
+			FrontLeftEarlyReflectionsSamples.AddUninitialized(InNumFrames);
+			FrontRightEarlyReflectionsSamples.AddUninitialized(InNumFrames);
+
+			// Generate both early reflections and late reflections 
+			EarlyReflections.ProcessAudio(InSamples, InNumChannels, FrontLeftEarlyReflectionsSamples, FrontRightEarlyReflectionsSamples);
+			LateReflections.ProcessAudio(InSamples, InNumChannels, FrontLeftLateReflectionsSamples, FrontRightLateReflectionsSamples);
+			// Add early and late reflections together.
+			SumBuffers(FrontLeftEarlyReflectionsSamples, FrontLeftLateReflectionsSamples, FrontLeftReverbSamples);
+			SumBuffers(FrontRightEarlyReflectionsSamples, FrontRightLateReflectionsSamples, FrontRightReverbSamples);
+		}
+
+		// Scale output by wetness
+		MultiplyBufferByConstantInPlace(FrontLeftReverbSamples, Settings.Wetness);
+		MultiplyBufferByConstantInPlace(FrontRightReverbSamples, Settings.Wetness);
+
+		// Interleave and upmix
+		InterleaveAndMixOutput(FrontLeftReverbSamples, FrontRightReverbSamples, OutSamples, OutNumChannels);
+	}
+
+	void FPlateReverbFast::ClampSettings(FPlateReverbFastSettings& InOutSettings)
+	{
+		// Clamp settings for this object and member objects.
+		InOutSettings.Wetness = FMath::Clamp(InOutSettings.Wetness, MinWetness, MaxWetness);
+		FLateReflectionsFast::ClampSettings(InOutSettings.LateReflections);
+		FEarlyReflectionsFast::ClampSettings(InOutSettings.EarlyReflections);
+	}
+
+	// Copy input samples to output samples. Remap channels if necessary.
+	void FPlateReverbFast::PassThroughAudio(const AlignedFloatBuffer& InSamples, const int32 InNumChannels, AlignedFloatBuffer& OutSamples, const int32 OutNumChannels)
+	{
+		const int32 InNum = InSamples.Num();
+		const int32 InNumFrames = InNum / InNumChannels;
+		const int32 OutNum = OutNumChannels * InNumFrames;
+
+		// Resize output buffer
+		OutSamples.Reset(OutNum);
+		OutSamples.AddUninitialized(OutNum);
+		FMemory::Memset(OutSamples.GetData(), 0, sizeof(float) * OutNum);
+
+		if (InNum > 0)
+		{
+			if (InNumChannels == OutNumChannels)
+			{
+				FMemory::Memcpy(OutSamples.GetData(), InSamples.GetData(), sizeof(float) * InNum);
+			}
+			else
+			{
+				// InNumChannels can only be 1 or 2 channels so we have a limited number of 
+				// upmix situations.
+				FMemory::Memset(OutSamples.GetData(), 0, sizeof(float) * OutNum);
+				if (1 == InNumChannels)
+				{
+					// Upmix a mono signal
+					float* OutSampleData = OutSamples.GetData();
+					const float* InSampleData = InSamples.GetData();
+					int32 OutPos = 0;
+					for (int32 i = 0; i < InNumFrames; i++, OutPos += OutNumChannels)
+					{
+						// Scale to keep loudnes consistent.
+						float value = InSampleData[i] * 0.5f;	
+						OutSampleData[OutPos] = value;
+						OutSampleData[OutPos + 1] = value;
+					}
+				}
+				else if (2 == InNumChannels)
+				{
+					// Upmix a stereo signal by copying:
+					// 		FrontLeft  -> FrontLeft
+					//		FrontRight -> FrontRight
+					float* OutSampleData = OutSamples.GetData();
+					const float* InSampleData = InSamples.GetData();
+					int32 OutPos = 0;
+					for (int32 i = 0; i < InNum; i += InNumChannels, OutPos += OutNumChannels)
+					{
+						OutSampleData[OutPos] = InSampleData[i];
+						OutSampleData[OutPos + 1] = InSampleData[i + 1];
+					}
+				}
+			}
+		}
+	}
+
+	// Copy reverberated samples to interleaved output samples. Map channels according to internal settings.
+	// InFrontLeftSamples and InFrontRightSamples may be modified in-place.
+	void FPlateReverbFast::InterleaveAndMixOutput(const AlignedFloatBuffer& InFrontLeftSamples, const AlignedFloatBuffer& InFrontRightSamples, AlignedFloatBuffer& OutSamples, const int32 OutNumChannels)
+	{
+		check(InFrontLeftSamples.Num() == InFrontRightSamples.Num())
+
+		const int32 InNumFrames = InFrontLeftSamples.Num();
+		const int32 OutNum = OutNumChannels * InNumFrames;
+
+		// Resize output buffer
+		OutSamples.Reset(OutNum);
+		OutSamples.AddUninitialized(OutNum);
+		FMemory::Memset(OutSamples.GetData(), 0, sizeof(float) * OutNum);
+
+		// Interleave / mix reverb audio into output buffer
+		if (2 == OutNumChannels)
+		{
+			// Stereo interleaved output
+			BufferInterleave2ChannelFast(FrontLeftReverbSamples, FrontRightReverbSamples, OutSamples);
+		}
+		else
+		{
+			if ((OutNumChannels < 5) || (FPlateReverbFastSettings::EQuadBehavior::StereoOnly == Settings.QuadBehavior))
+			{
+				// We do not handle any quad reverb mapping when OutNumChannels is less than 5
+
+				float* LeftSampleData = FrontLeftReverbSamples.GetData();
+				float* RightSampleData = FrontRightReverbSamples.GetData();
+				float* OutSampleData = OutSamples.GetData();
+				
+				int32 OutPos = 0;
+				for (int32 i = 0; i < InNumFrames; i++, OutPos += OutNumChannels)
+				{
+					OutSampleData[OutPos + EAudioMixerChannel::FrontLeft] = LeftSampleData[i];
+					OutSampleData[OutPos + EAudioMixerChannel::FrontRight] = RightSampleData[i];
+				}
+			}
+			else
+			{
+				// There are 5 or more output channels and quad mapping is enabled.
+
+				LeftAttenuatedSamples.Reset(InNumFrames);
+				RightAttenuatedSamples.Reset(InNumFrames);
+
+				LeftAttenuatedSamples.AddUninitialized(InNumFrames);
+				RightAttenuatedSamples.AddUninitialized(InNumFrames);
+
+				// Reduce volume of output reverbs.
+				BufferMultiplyByConstant(FrontLeftReverbSamples, 0.5f, LeftAttenuatedSamples);
+				BufferMultiplyByConstant(FrontRightReverbSamples, 0.5f, RightAttenuatedSamples);
+
+				const float* FrontLeftSampleData = LeftAttenuatedSamples.GetData();
+				const float* FrontRightSampleData = RightAttenuatedSamples.GetData();
+				// WARNING: this pointer will alias other pointers in this scope. Be conscious of RESTRICT keyword in any called functions.
+				const float* BackLeftSampleData = nullptr;
+				// WARNING: this pointer will alias other pointers in this scope. Be conscious of RESTRICT keyword in any called functions.
+				const float* BackRightSampleData = nullptr;
+
+				// Map quads by asigning pointers.
+				switch (Settings.QuadBehavior)
+				{
+					case FPlateReverbFastSettings::EQuadBehavior::QuadFlipped:
+						// Left and right are flipped.
+						BackLeftSampleData = FrontRightSampleData;
+						BackRightSampleData = FrontLeftSampleData;
+
+					case FPlateReverbFastSettings::EQuadBehavior::QuadMatched:
+					default:
+						// Left and right are matched.
+						BackLeftSampleData = FrontLeftSampleData;
+						BackRightSampleData = FrontRightSampleData;
+				}
+
+				// Interleave to output
+				float* OutSampleData = OutSamples.GetData();
+				int32 OutPos = 0;
+				for (int32 i = 0; i < InNumFrames; i++, OutPos += OutNumChannels)
+				{
+					OutSampleData[OutPos + EAudioMixerChannel::FrontLeft] = FrontLeftSampleData[i];
+					OutSampleData[OutPos + EAudioMixerChannel::FrontRight] = FrontRightSampleData[i];
+					OutSampleData[OutPos + EAudioMixerChannel::BackLeft] = BackLeftSampleData[i];
+					OutSampleData[OutPos + EAudioMixerChannel::BackRight] = BackRightSampleData[i];
+				}
+			}
+		}
+	}
+
+	void FPlateReverbFast::ApplySettings()
+	{
+		EarlyReflections.SetSettings(Settings.EarlyReflections);
+		LateReflections.SetSettings(Settings.LateReflections);
+	}
+}
diff --git a/Engine/Source/Runtime/AudioMixer/Private/Effects/AudioMixerSubmixEffectReverb.cpp b/Engine/Source/Runtime/AudioMixer/Private/Effects/AudioMixerSubmixEffectReverb.cpp
--- a/Engine/Source/Runtime/AudioMixer/Private/Effects/AudioMixerSubmixEffectReverb.cpp
+++ b/Engine/Source/Runtime/AudioMixer/Private/Effects/AudioMixerSubmixEffectReverb.cpp
@@ -7,24 +7,24 @@
 #include "AudioMixer.h"
 
 static int32 DisableSubmixReverbCVar = 0;
-FAutoConsoleVariableRef CVarDisableSubmixReverb(
+static FAutoConsoleVariableRef CVarDisableSubmixReverb(
 	TEXT("au.DisableReverbSubmix"),
 	DisableSubmixReverbCVar,
 	TEXT("Disables the reverb submix.\n")
 	TEXT("0: Not Disabled, 1: Disabled"),
 	ECVF_Default);
 
-static int32 EnableReverbStereoFlipForQuadCvar = 0;
-FAutoConsoleVariableRef CVarReverbStereoFlipForQuad(
+static int32 EnableReverbStereoFlipForQuadCVar = 0;
+static FAutoConsoleVariableRef CVarReverbStereoFlipForQuad(
 	TEXT("au.EnableReverbStereoFlipForQuad"),
-	EnableReverbStereoFlipForQuadCvar,
+	EnableReverbStereoFlipForQuadCVar,
 	TEXT("Enables doing a stereo flip for quad reverb when in surround.\n")
 	TEXT("0: Not Enabled, 1: Enabled"),
 	ECVF_Default);
 
 
 static int32 DisableQuadReverbCVar = 0;
-FAutoConsoleVariableRef CVarDisableQuadReverbCVar(
+static FAutoConsoleVariableRef CVarDisableQuadReverbCVar(
 	TEXT("au.DisableQuadReverb"),
 	DisableQuadReverbCVar,
 	TEXT("Disables quad reverb in surround.\n")
@@ -133,7 +133,7 @@ void FSubmixEffectReverb::OnProcessAudio(const FSoundEffectSubmixInputData& InDa
 				PlateReverb.ProcessAudioFrame(&AudioData[InSampleIndex], InData.NumChannels, &OutAudioData[OutSampleIndex], InData.NumChannels);
 			}
 		}
-		else if (EnableReverbStereoFlipForQuadCvar == 1)
+		else if (EnableReverbStereoFlipForQuadCVar == 1)
 		{
 			for (int32 InSampleIndex = 0, OutSampleIndex = 0; InSampleIndex < InData.AudioBuffer->Num(); InSampleIndex += InData.NumChannels, OutSampleIndex += OutData.NumChannels)
 			{
@@ -232,4 +232,4 @@ void USubmixEffectReverbPreset::SetSettingsWithReverbEffect(const UReverbEffect*
 void USubmixEffectReverbPreset::SetSettings(const FSubmixEffectReverbSettings& InSettings)
 {
 	UpdateSettings(InSettings);
-}
\ No newline at end of file
+}
diff --git a/Engine/Source/Runtime/AudioMixer/Private/Effects/AudioMixerSubmixEffectReverbFast.cpp b/Engine/Source/Runtime/AudioMixer/Private/Effects/AudioMixerSubmixEffectReverbFast.cpp
new file mode 100644
--- /dev/null
+++ b/Engine/Source/Runtime/AudioMixer/Private/Effects/AudioMixerSubmixEffectReverbFast.cpp
@@ -0,0 +1,235 @@
+// Copyright 1998-2019 Epic Games, Inc. All Rights Reserved.
+
+#include "SubmixEffects/AudioMixerSubmixEffectReverbFast.h"
+#include "AudioMixerEffectsManager.h"
+#include "Sound/ReverbEffect.h"
+#include "Audio.h"
+#include "AudioMixer.h"
+#include "DSP/ReverbFast.h"
+
+
+static int32 DisableSubmixReverbCVar = 0;
+static FAutoConsoleVariableRef CVarDisableSubmixReverb(
+	TEXT("au.DisableReverbSubmix"),
+	DisableSubmixReverbCVar,
+	TEXT("Disables the reverb submix.\n")
+	TEXT("0: Not Disabled, 1: Disabled"),
+	ECVF_Default);
+
+static int32 EnableReverbStereoFlipForQuadCVar = 0;
+static FAutoConsoleVariableRef CVarReverbStereoFlipForQuad(
+	TEXT("au.EnableReverbStereoFlipForQuad"),
+	EnableReverbStereoFlipForQuadCVar,
+	TEXT("Enables doing a stereo flip for quad reverb when in surround.\n")
+	TEXT("0: Not Enabled, 1: Enabled"),
+	ECVF_Default);
+
+static int32 DisableQuadReverbCVar = 0;
+static FAutoConsoleVariableRef CVarDisableQuadReverbCVar(
+	TEXT("au.DisableQuadReverb"),
+	DisableQuadReverbCVar,
+	TEXT("Disables quad reverb in surround.\n")
+	TEXT("0: Not Disabled, 1: Disabled"),
+	ECVF_Default);
+
+
+
+class UReverbEffect;
+
+FSubmixEffectReverbFast::FSubmixEffectReverbFast()
+{
+}
+
+void FSubmixEffectReverbFast::Init(const FSoundEffectSubmixInitData& InitData)
+{
+	LLM_SCOPE(ELLMTag::AudioMixer);
+
+	/* `FPlateReverbFast` produces slightly different quality effect than `FPlateReverb`. Comparing the Init
+	 * settings between FSubmixEffectReverb and FSubmixEffectReverbFast slight differences will arise.
+	 *
+	 * The delay line implementations significantly differ between the `FPlateReverb` and `FPlateReverbFast` classes.
+	 * Specifically, the `FPlateReverb` class utilizes linearly interpolated fractional delay line and fractional
+	 * delays while the `FPlateReverbFast` class uses integer delay lines and integer delays whenever possible.
+	 * Linearly interpolated fractional delay lines introduce a low pass filter dependent upon the fractional portion
+	 * of the delay value. As a result, the `FPlateReverb` class produces a darker reverb.
+	 */
+	Audio::FPlateReverbFastSettings NewSettings;
+
+
+	NewSettings.Wetness = 1.0f;
+
+	NewSettings.EarlyReflections.Decay = 0.9f;
+	NewSettings.EarlyReflections.Absorption = 0.7f;
+	NewSettings.EarlyReflections.Gain = 1.0f;
+	NewSettings.EarlyReflections.PreDelayMsec = 0.0f;
+	NewSettings.EarlyReflections.Bandwidth = 0.8f;
+
+	NewSettings.LateReflections.LateDelayMsec = 0.0f;
+	NewSettings.LateReflections.LateGainDB = 0.0f;
+	NewSettings.LateReflections.Bandwidth = 0.54f;
+	NewSettings.LateReflections.Diffusion = 0.60f;
+	NewSettings.LateReflections.Dampening = 0.35f;
+	NewSettings.LateReflections.Decay = 0.15f;
+	NewSettings.LateReflections.Density = 0.85f;
+
+	Params.SetParams(NewSettings);
+
+	DecayCurve.AddKey(0.0f, 0.99f);
+	DecayCurve.AddKey(2.0f, 0.45f);
+	DecayCurve.AddKey(5.0f, 0.15f);
+	DecayCurve.AddKey(10.0f, 0.1f);
+	DecayCurve.AddKey(18.0f, 0.01f);
+	DecayCurve.AddKey(19.0f, 0.002f);
+	DecayCurve.AddKey(20.0f, 0.0001f);
+
+	PlateReverb = MakeUnique<Audio::FPlateReverbFast>(InitData.SampleRate, 512, NewSettings);
+	PlateReverb->EnableEarlyReflections(false);
+	PlateReverb->EnableLateReflections(true);
+}
+
+void FSubmixEffectReverbFast::OnPresetChanged()
+{
+	LLM_SCOPE(ELLMTag::AudioMixer);
+
+	GET_EFFECT_SETTINGS(SubmixEffectReverbFast);
+
+	FAudioReverbEffect ReverbEffect;
+	ReverbEffect.Density = Settings.Density;
+	ReverbEffect.Diffusion = Settings.Diffusion;
+	ReverbEffect.Gain = Settings.Gain;
+	ReverbEffect.GainHF = Settings.GainHF;
+	ReverbEffect.DecayTime = Settings.DecayTime;
+	ReverbEffect.DecayHFRatio = Settings.DecayHFRatio;
+	ReverbEffect.ReflectionsGain = Settings.ReflectionsGain;
+	ReverbEffect.ReflectionsDelay = Settings.ReflectionsDelay;
+	ReverbEffect.LateGain = Settings.LateGain;
+	ReverbEffect.LateDelay = Settings.LateDelay;
+	ReverbEffect.AirAbsorptionGainHF = Settings.AirAbsorptionGainHF;
+	ReverbEffect.RoomRolloffFactor = 0.0f; // not used
+	ReverbEffect.Volume = Settings.WetLevel;
+
+
+	SetEffectParameters(ReverbEffect);
+}
+
+void FSubmixEffectReverbFast::OnProcessAudio(const FSoundEffectSubmixInputData& InData, FSoundEffectSubmixOutputData& OutData)
+{
+	LLM_SCOPE(ELLMTag::AudioMixer);
+
+	check(InData.NumChannels == 2);
+ 	if (OutData.NumChannels < 2 || DisableSubmixReverbCVar == 1) 
+	{
+		// Not supported
+		return;
+	}
+
+	UpdateParameters();
+
+	PlateReverb->ProcessAudio(*InData.AudioBuffer, InData.NumChannels, *OutData.AudioBuffer, OutData.NumChannels);
+}
+
+void FSubmixEffectReverbFast::SetEffectParameters(const FAudioReverbEffect& InParams)
+{
+	LLM_SCOPE(ELLMTag::AudioMixer);
+
+	/* `FPlateReverbFast` produces slightly different quality effect than `FPlateReverb`. Comparing the 
+	 * settings between FSubmixEffectReverb and FSubmixEffectReverbFast slight differences will arise.
+	 *
+	 * The delay line implementations significantly differ between the `FPlateReverb` and `FPlateReverbFast` classes.
+	 * Specifically, the `FPlateReverb` class utilizes linearly interpolated fractional delay line and fractional
+	 * delays while the `FPlateReverbFast` class uses integer delay lines and integer delays whenever possible.
+	 * Linearly interpolated fractional delay lines introduce a low pass filter dependent upon the fractional portion
+	 * of the delay value. As a result, the `FPlateReverb` class produces a darker reverb.
+	 */
+	Audio::FPlateReverbFastSettings NewSettings;
+
+	// General reverb settings
+	NewSettings.Wetness = FMath::GetMappedRangeValueClamped({ 0.0f, 10.0f }, { 0.0f, 10.0f }, InParams.Volume);
+
+	// Early Reflections
+	NewSettings.EarlyReflections.Gain = FMath::GetMappedRangeValueClamped({ 0.0f, 3.16f }, { 0.0f, 1.0f }, InParams.ReflectionsGain);
+	NewSettings.EarlyReflections.PreDelayMsec = FMath::GetMappedRangeValueClamped({ 0.0f, 0.3f }, { 0.0f, 300.0f }, InParams.ReflectionsDelay);
+	NewSettings.EarlyReflections.Bandwidth = FMath::GetMappedRangeValueClamped({ 0.0f, 1.0f }, { 0.0f, 1.0f }, InParams.GainHF);
+
+	// LateReflections
+	NewSettings.LateReflections.LateDelayMsec = FMath::GetMappedRangeValueClamped({ 0.0f, 0.1f }, { 0.0f, 100.0f }, InParams.LateDelay);
+	NewSettings.LateReflections.LateGainDB = FMath::GetMappedRangeValueClamped({ 0.0f, 1.0f }, { 0.0f, 1.0f }, InParams.Gain);
+	NewSettings.LateReflections.Bandwidth = FMath::GetMappedRangeValueClamped({ 0.0f, 1.0f }, { 0.1f, 0.6f }, InParams.AirAbsorptionGainHF);
+	NewSettings.LateReflections.Diffusion = FMath::GetMappedRangeValueClamped({ 0.05f, 1.0f }, { 0.0f, 0.95f }, InParams.Diffusion);
+	NewSettings.LateReflections.Dampening = FMath::GetMappedRangeValueClamped({ 0.05f, 1.95f }, { 0.0f, 0.999f }, InParams.DecayHFRatio);
+	NewSettings.LateReflections.Density = FMath::GetMappedRangeValueClamped({ 0.0f, 0.95f }, { 0.06f, 1.0f }, InParams.Density);
+
+	// Use mapping function to get decay time in seconds to internal linear decay scale value
+	const float DecayValue = DecayCurve.Eval(InParams.DecayTime);
+	NewSettings.LateReflections.Decay = DecayValue;
+
+	// Convert to db
+	NewSettings.LateReflections.LateGainDB = Audio::ConvertToDecibels(NewSettings.LateReflections.LateGainDB);
+
+	// Apply the settings the thread safe settings object
+	Params.SetParams(NewSettings);
+}
+
+void FSubmixEffectReverbFast::UpdateParameters()
+{
+	Audio::FPlateReverbFastSettings NewSettings;
+	if (Params.GetParams(&NewSettings))
+	{
+		PlateReverb->SetSettings(NewSettings);
+	}
+
+	// Check cVars for quad mapping
+	Audio::FPlateReverbFastSettings::EQuadBehavior TargetQuadBehavior;
+	if (DisableQuadReverbCVar)
+	{
+		// Disable quad mapping.
+ 		TargetQuadBehavior = Audio::FPlateReverbFastSettings::EQuadBehavior::StereoOnly;
+	}
+	else if (!DisableQuadReverbCVar && EnableReverbStereoFlipForQuadCVar)
+	{
+		// Enable quad flipped mapping
+		TargetQuadBehavior = Audio::FPlateReverbFastSettings::EQuadBehavior::QuadFlipped;
+	}
+	else
+	{
+		// Enable quad mapping
+		TargetQuadBehavior = Audio::FPlateReverbFastSettings::EQuadBehavior::QuadMatched;
+	}
+
+	// Check if settings need to be updated
+	const Audio::FPlateReverbFastSettings& ReverbFastSettings = PlateReverb->GetSettings();
+	if (ReverbFastSettings.QuadBehavior != TargetQuadBehavior)
+	{
+		// Update quad settings 
+		NewSettings = ReverbFastSettings;
+		NewSettings.QuadBehavior = TargetQuadBehavior;
+		PlateReverb->SetSettings(NewSettings);
+	}
+}
+
+void USubmixEffectReverbFastPreset::SetSettingsWithReverbEffect(const UReverbEffect* InReverbEffect, const float WetLevel, const float DryLevel)
+{
+	if (InReverbEffect)
+	{
+		Settings.Density = InReverbEffect->Density;
+		Settings.Diffusion = InReverbEffect->Diffusion;
+		Settings.Gain = InReverbEffect->Gain;
+		Settings.GainHF = InReverbEffect->GainHF;
+		Settings.DecayTime = InReverbEffect->DecayTime;
+		Settings.DecayHFRatio = InReverbEffect->DecayHFRatio;
+		Settings.ReflectionsGain = InReverbEffect->ReflectionsGain;
+		Settings.ReflectionsDelay = InReverbEffect->ReflectionsDelay;
+		Settings.LateGain = InReverbEffect->LateGain;
+		Settings.LateDelay = InReverbEffect->LateDelay;
+		Settings.AirAbsorptionGainHF = InReverbEffect->AirAbsorptionGainHF;
+		Settings.WetLevel = WetLevel;
+		Settings.DryLevel = DryLevel;
+
+		Update();
+	}
+}
+
+void USubmixEffectReverbFastPreset::SetSettings(const FSubmixEffectReverbFastSettings& InSettings)
+{
+	UpdateSettings(InSettings);
+}
diff --git a/Engine/Source/Runtime/AudioMixer/Public/DSP/AlignedBlockBuffer.h b/Engine/Source/Runtime/AudioMixer/Public/DSP/AlignedBlockBuffer.h
new file mode 100644
--- /dev/null
+++ b/Engine/Source/Runtime/AudioMixer/Public/DSP/AlignedBlockBuffer.h
@@ -0,0 +1,87 @@
+// Copyright 1998-2019 Epic Games, Inc. All Rights Reserved.
+
+#pragma once
+#include "CoreMinimal.h"
+#include "BufferVectorOperations.h"
+
+namespace Audio
+{	
+	// First In First Out Buffer designed for audio buffers. Generally performs best when buffers are greater
+	// than 16 samples, optimally for buffers between 256 to 4096 samples.  Memory operations are designed 
+	// for audio applications by ensuring memory allocations only occur in the constructor or when capacity 
+	// is altered, memory freeing only happens in the destructor or when capacity is altered. All assignments 
+	// are performed using block memory copies.
+	class AUDIOMIXER_API FAlignedBlockBuffer
+	{
+	public:
+		
+		// Constructor
+		FAlignedBlockBuffer(int32 InSampleCapacity=8192, int32 InMaxNumInspectSamples=1024, uint32 InByteAlignment=AUDIO_BUFFER_ALIGNMENT);
+
+		// Destructor
+		~FAlignedBlockBuffer() throw();
+		
+		// Maximum size of buffer
+		int32 GetSampleCapacity() const;
+		
+		// Amount of samples left over in buffer
+		int32 GetNumAvailable() const;
+
+		// Add zeros to the buffer
+		void AddZeros(int32 InNum);
+		
+		// Add samples to the buffer.
+		void AddSamples(const float* InSamples, int32 InNum);
+		
+		// Remove samples from the buffer
+		void RemoveSamples(int32 InNum);
+
+		// Clear entire buffer
+		void ClearSamples();
+
+		// Inspect samples in the buffer. Returned samples abide by the memory alignment specified during construction.
+		// Warning: Not thread safe. The returned pointer is only valid as long as this object exists 
+		// and before any mutating operations are made to this object.
+		const float* InspectSamples(int32 InNum, int32 InOffset=0);
+		
+		// Maximum number of elements that can be inspected
+		int32 GetMaxNumInspectSamples() const;
+		
+	protected:
+
+		// Set the maximum number of elements that can be inspected
+		void SetMaxNumInspectSamples(int32 InMax);
+
+		float* AllocateAlignedFloatArray(int32 InNum) const;
+		void FreeFloatArray(float*& Array) const;
+		void ZeroFloatArray(float* InArray, int32 InNum) const;
+		void CopyFloatArray(float* ToArray, const float* FromArray, int32 InNum) const;
+
+	private:
+
+		int32 CheckNumAndIncrementWriteCount(int32 InNum);
+
+		// Hide copy and assignment operators
+		FAlignedBlockBuffer& operator=(const FAlignedBlockBuffer& rhs);
+		FAlignedBlockBuffer(const FAlignedBlockBuffer& CopyBuff);
+		
+		// Alignment definition
+		uint32 ByteAlignment;
+		uint32 FloatAlignment;
+
+		// Internal memory
+		float* InternalBuffer;
+		float* RolloverBuffer;
+
+		// Read and write pointers
+		float* ReadPointer;
+		float* WritePointer;
+
+		int32 SampleCapacity;
+		int32 MaxNumInspectSamples;
+		int32 WriteCount;
+		int32 WritePos;
+		int32 ReadPos;
+	};
+}
+
diff --git a/Engine/Source/Runtime/AudioMixer/Public/DSP/AllPassFractionalDelay.h b/Engine/Source/Runtime/AudioMixer/Public/DSP/AllPassFractionalDelay.h
new file mode 100644
--- /dev/null
+++ b/Engine/Source/Runtime/AudioMixer/Public/DSP/AllPassFractionalDelay.h
@@ -0,0 +1,44 @@
+// Copyright 1998-2019 Epic Games, Inc. All Rights Reserved.
+
+#pragma once
+
+#include "CoreMinimal.h"
+#include "BufferVectorOperations.h"
+#include "AlignedBlockBuffer.h"
+
+namespace Audio
+{
+	// Performs a fractional delay utilzing a single order all pass filter.
+	class AUDIOMIXER_API FAllPassFractionalDelay
+	{
+	public:
+		// InMaxDelay sets the maximum allowable delay that this object can support. The minimum delay is 0.5
+		// InNumInternalBufferSamples sets the maximum block processing size.
+		FAllPassFractionalDelay(int32 InMaxDelay, int32 InNumInternalBufferSamples);
+
+		~FAllPassFractionalDelay();
+
+		// Apply delay to InSamples. Fill OutSamples with data from the delay line at a delay of InDelay.
+		// InDelay must be equal length to InSamples. 
+		void ProcessAudio(const AlignedFloatBuffer& InSamples, const AlignedFloatBuffer& InDelays, AlignedFloatBuffer& OutSamples);
+
+		// Set all values in internal delay line to zero. 
+		void Reset();
+
+	private:
+		void ProcessAudioBlock(const float* InSamples, const float* InDelays, const int32 InNum, float* OutSamples);
+
+		int32 MaxDelay;
+		int32 NumInternalBufferSamples;
+
+		// Delay element for all pass filter.
+		float Z1;
+
+		TUniquePtr<FAlignedBlockBuffer> DelayLine;
+
+		AlignedFloatBuffer Coefficients;
+		AlignedFloatBuffer FractionalDelays;
+		AlignedInt32Buffer IntegerDelays;
+		AlignedInt32Buffer IntegerDelayOffsets;
+	};
+}
diff --git a/Engine/Source/Runtime/AudioMixer/Public/DSP/BufferOnePoleLPF.h b/Engine/Source/Runtime/AudioMixer/Public/DSP/BufferOnePoleLPF.h
new file mode 100644
--- /dev/null
+++ b/Engine/Source/Runtime/AudioMixer/Public/DSP/BufferOnePoleLPF.h
@@ -0,0 +1,45 @@
+// Copyright 1998-2019 Epic Games, Inc. All Rights Reserved.
+
+#pragma once
+
+#include "CoreMinimal.h"
+#include "DSP/Dsp.h"
+#include "DSP/BufferVectorOperations.h"
+
+namespace Audio
+{
+
+	// Simple 1-pole lowpass filter
+	class AUDIOMIXER_API FBufferOnePoleLPF
+	{
+	public:
+
+		// Constructor 
+		FBufferOnePoleLPF(float InG = 0.0f);
+
+		// Set the LPF gain coefficient
+		void SetG(float InG);
+
+		// Resets the sample delay to 0
+		void Reset();
+
+		// Sets the filter frequency using normalized frequency (between 0.0 and 1.0f or 0.0 hz and Nyquist Frequency in Hz) 
+		void SetFrequency(const float InFrequency);
+
+		// Process InSamples and place filtered audio in OutSamples
+		void ProcessAudio(const AlignedFloatBuffer& InSamples, AlignedFloatBuffer& OutSamples);
+
+	protected:
+		float CutoffFrequency;
+
+		// Filter coefficients
+		float B1;
+		float A0;
+
+		// 1-sample delay
+		float Z1;
+	};
+
+	
+
+}
diff --git a/Engine/Source/Runtime/AudioMixer/Public/DSP/BufferVectorOperations.h b/Engine/Source/Runtime/AudioMixer/Public/DSP/BufferVectorOperations.h
--- a/Engine/Source/Runtime/AudioMixer/Public/DSP/BufferVectorOperations.h
+++ b/Engine/Source/Runtime/AudioMixer/Public/DSP/BufferVectorOperations.h
@@ -11,48 +11,93 @@
 #define AUDIO_BUFFER_ALIGNMENT 16
 #endif
 
+#define AUDIO_SIMD_BYTE_ALIGNMENT (16)
+#define AUDIO_SIMD_FLOAT_ALIGNMENT (4)
+
 namespace Audio
 {
+	
 	typedef TArray<float, TAlignedHeapAllocator<AUDIO_BUFFER_ALIGNMENT>> AlignedFloatBuffer;
 	typedef TArray<uint8, TAlignedHeapAllocator<AUDIO_BUFFER_ALIGNMENT>> AlignedByteBuffer;
+	typedef TArray<int32, TAlignedHeapAllocator<AUDIO_BUFFER_ALIGNMENT>> AlignedInt32Buffer;
+
 
 	/** CHANNEL-AGNOSTIC OPERATIONS */
 
+	/* Sets a values to zero if value is denormal. Denormal numbers significantly slow down floating point operations. */
+	void AUDIOMIXER_API BufferUnderflowClampFast(AlignedFloatBuffer& InOutBuffer);
+
+	/* Sets a values to zero if value is denormal. Denormal numbers significantly slow down floating point operations. */
+	void AUDIOMIXER_API BufferUnderflowClampFast(float* RESTRICT InOutBuffer, const int32 InNum);
+		
 	/** Multiplies the input aligned float buffer with the given value. */
-	void BufferMultiplyByConstant(const AlignedFloatBuffer& InFloatBuffer, float InValue, AlignedFloatBuffer& OutFloatBuffer);
+	void AUDIOMIXER_API BufferMultiplyByConstant(const AlignedFloatBuffer& InFloatBuffer, float InValue, AlignedFloatBuffer& OutFloatBuffer);
+	void AUDIOMIXER_API BufferMultiplyByConstant(const float* RESTRICT InFloatBuffer, float InValue, float* RESTRICT OutFloatBuffer, const int32 InNumSamples);
 
 	/** Similar to BufferMultiplyByConstant, but (a) assumes a buffer length divisible by 4 and (b) performs the multiply in place. */
-	void MultiplyBufferByConstantInPlace(AlignedFloatBuffer& InBuffer, float InGain);
-	void MultiplyBufferByConstantInPlace(float* RESTRICT InBuffer, int32 NumSamples, float InGain);
+	void AUDIOMIXER_API MultiplyBufferByConstantInPlace(AlignedFloatBuffer& InBuffer, float InGain);
+	void AUDIOMIXER_API MultiplyBufferByConstantInPlace(float* RESTRICT InBuffer, int32 NumSamples, float InGain);
+
+	/* Performs an element-wise weighted sum OutputBuffer = (InBuffer1 x InGain1) + (InBuffer2 x InGain2) */
+	void AUDIOMIXER_API BufferWeightedSumFast(const AlignedFloatBuffer& InBuffer1, float InGain1, const AlignedFloatBuffer& InBuffer2, float InGain2, AlignedFloatBuffer& OutBuffer);
+
+	/* Performs an element-wise weighted sum OutputBuffer = (InBuffer1 x InGain1) + InBuffer2 */
+	void AUDIOMIXER_API BufferWeightedSumFast(const AlignedFloatBuffer& InBuffer1, float InGain1, const AlignedFloatBuffer& InBuffer2, AlignedFloatBuffer& OutBuffer);
+
+	/* Performs an element-wise weighted sum OutputBuffer = (InBuffer1 x InGain1) + (InBuffer2 x InGain2) */
+	void AUDIOMIXER_API BufferWeightedSumFast(const float* RESTRICT InBuffer1, float InGain1, const float* RESTRICT InBuffer2, float InGain2, float* RESTRICT OutBuffer, int32 InNum);
+
+	/* Performs an element-wise weighted sum OutputBuffer = (InBuffer1 x InGain1) + InBuffer2 */
+	void AUDIOMIXER_API BufferWeightedSumFast(const float* RESTRICT InBuffer1, float InGain1, const float* RESTRICT InBuffer2, float* RESTRICT OutBuffer, int32 InNum);
 
 	/* Takes a float buffer and quickly interpolates it's gain from StartValue to EndValue. */
 	/* This operation completely ignores channel counts, so avoid using this function on buffers that are not mono, stereo or quad */
 	/* if the buffer needs to fade all channels uniformly. */
-	void FadeBufferFast(AlignedFloatBuffer& OutFloatBuffer, const float StartValue, const float EndValue);
-	void FadeBufferFast(float* RESTRICT OutFloatBuffer, int32 NumSamples, const float StartValue, const float EndValue);
+	void AUDIOMIXER_API FadeBufferFast(AlignedFloatBuffer& OutFloatBuffer, const float StartValue, const float EndValue);
+	void AUDIOMIXER_API FadeBufferFast(float* RESTRICT OutFloatBuffer, int32 NumSamples, const float StartValue, const float EndValue);
 
 	/** Takes buffer InFloatBuffer, optionally multiplies it by Gain, and adds it to BufferToSumTo. */
-	void MixInBufferFast(const AlignedFloatBuffer& InFloatBuffer, AlignedFloatBuffer& BufferToSumTo, const float Gain);
-	void MixInBufferFast(const float* RESTRICT InFloatBuffer, float* RESTRICT BufferToSumTo, int32 NumSamples, const float Gain);
-	void MixInBufferFast(const float* RESTRICT InFloatBuffer, float* RESTRICT BufferToSumTo, int32 NumSamples);
+	void AUDIOMIXER_API MixInBufferFast(const AlignedFloatBuffer& InFloatBuffer, AlignedFloatBuffer& BufferToSumTo, const float Gain);
+	void AUDIOMIXER_API MixInBufferFast(const float* RESTRICT InFloatBuffer, float* RESTRICT BufferToSumTo, int32 NumSamples, const float Gain);
+	void AUDIOMIXER_API MixInBufferFast(const AlignedFloatBuffer& InFloatBuffer, AlignedFloatBuffer& BufferToSumTo);
+	void AUDIOMIXER_API MixInBufferFast(const float* RESTRICT InFloatBuffer, float* RESTRICT BufferToSumTo, int32 NumSamples);
+
+	/* Subtracts two buffers together element-wise. */
+	void AUDIOMIXER_API BufferSubtractFast(const AlignedFloatBuffer& InMinuend, const AlignedFloatBuffer& InSubtrahend, AlignedFloatBuffer& OutputBuffer);
+
+	/* Subtracts two buffers together element-wise. */
+	void AUDIOMIXER_API BufferSubtractFast(const float* RESTRICT InMinuend, const float* RESTRICT InSubtrahend, float* RESTRICT OutputBuffer, int32 NumSamples);
+
+	/* Performs element-wise in-place subtraction placing the result in the subtrahend. InOutSubtrahend = InMinuend - InOutSubtrahend */
+	void AUDIOMIXER_API BufferSubtractInPlace1Fast(const AlignedFloatBuffer& InMinuend, AlignedFloatBuffer& InOutSubtrahend);
+
+	/* Performs element-wise in-place subtraction placing the result in the subtrahend. InOutSubtrahend = InMinuend - InOutSubtrahend */
+	void AUDIOMIXER_API BufferSubtractInPlace1Fast(const float* RESTRICT InMinuend, float* RESTRICT InOutSubtrahend, int32 NumSamples);
+
+
+	/* Performs element-wise in-place subtraction placing the result in the minuend. InOutMinuend = InOutMinuend - InSubtrahend */
+	void AUDIOMIXER_API BufferSubtractInPlace2Fast(AlignedFloatBuffer& InOutMinuend, const AlignedFloatBuffer& InSubtrahend);
+
+	/* Performs element-wise in-place subtraction placing the result in the minuend. InOutMinuend = InOutMinuend - InSubtrahend */
+	void AUDIOMIXER_API BufferSubtractInPlace2Fast(float* RESTRICT InOutMinuend, const float* RESTRICT InSubtrahend, int32 NumSamples);
 
 	/** Sums two buffers together and places the result in the resulting buffer. */
-	void SumBuffers(const AlignedFloatBuffer& InFloatBuffer1, const AlignedFloatBuffer& InFloatBuffer2, AlignedFloatBuffer& OutputBuffer);
-	void SumBuffers(const float* RESTRICT InFloatBuffer1, const float* RESTRICT InFloatBuffer2, float* RESTRICT OutputBuffer, int32 NumSamples);
+	void AUDIOMIXER_API SumBuffers(const AlignedFloatBuffer& InFloatBuffer1, const AlignedFloatBuffer& InFloatBuffer2, AlignedFloatBuffer& OutputBuffer);
+	void AUDIOMIXER_API SumBuffers(const float* RESTRICT InFloatBuffer1, const float* RESTRICT InFloatBuffer2, float* RESTRICT OutputBuffer, int32 NumSamples);
 
 	/** Multiply the second buffer in place by the first buffer. */
-	void MultiplyBuffersInPlace(const AlignedFloatBuffer& InFloatBuffer, AlignedFloatBuffer& BufferToMultiply);
-	void MultiplyBuffersInPlace(const float* RESTRICT InFloatBuffer, float* RESTRICT BufferToMultiply, int32 NumSamples);
+	void AUDIOMIXER_API MultiplyBuffersInPlace(const AlignedFloatBuffer& InFloatBuffer, AlignedFloatBuffer& BufferToMultiply);
+	void AUDIOMIXER_API MultiplyBuffersInPlace(const float* RESTRICT InFloatBuffer, float* RESTRICT BufferToMultiply, int32 NumSamples);
 
 	/** CHANNEL-AGNOSTIC ANALYSIS OPERATIONS */
 
 	/** Takes an audio buffer and returns the magnitude across that buffer. */
-	float GetMagnitude(const AlignedFloatBuffer& Buffer);
-	float GetMagnitude(const float* RESTRICT Buffer, int32 NumSamples);
+	float AUDIOMIXER_API GetMagnitude(const AlignedFloatBuffer& Buffer);
+	float AUDIOMIXER_API GetMagnitude(const float* RESTRICT Buffer, int32 NumSamples);
 
 	/** Takes an audio buffer and gets the average absolute amplitude across that buffer. */
-	float GetAverageAmplitude(const AlignedFloatBuffer& Buffer);
-	float GetAverageAmplitude(const float* RESTRICT Buffer, int32 NumSamples);
+	float AUDIOMIXER_API GetAverageAmplitude(const AlignedFloatBuffer& Buffer);
+	float AUDIOMIXER_API GetAverageAmplitude(const float* RESTRICT Buffer, int32 NumSamples);
 
 	/** CHANNEL-SPECIFIC OPERATIONS */
 
@@ -60,114 +105,135 @@ namespace Audio
 	 *  StereoBuffer must have an even number of frames.
 	 *  If StartGains and EndGains are provided, this function will interpolate between the two across the buffer.
 	*/
-	void Apply2ChannelGain(AlignedFloatBuffer& StereoBuffer, const float* RESTRICT Gains);
-	void Apply2ChannelGain(float* RESTRICT StereoBuffer, int32 NumSamples, const float* RESTRICT Gains);
-	void Apply2ChannelGain(AlignedFloatBuffer& StereoBuffer, const float* RESTRICT StartGains, const float* RESTRICT EndGains);
-	void Apply2ChannelGain(float* RESTRICT StereoBuffer, int32 NumSamples, const float* RESTRICT StartGains, const float* RESTRICT EndGains);
+	void AUDIOMIXER_API Apply2ChannelGain(AlignedFloatBuffer& StereoBuffer, const float* RESTRICT Gains);
+	void AUDIOMIXER_API Apply2ChannelGain(float* RESTRICT StereoBuffer, int32 NumSamples, const float* RESTRICT Gains);
+	void AUDIOMIXER_API Apply2ChannelGain(AlignedFloatBuffer& StereoBuffer, const float* RESTRICT StartGains, const float* RESTRICT EndGains);
+	void AUDIOMIXER_API Apply2ChannelGain(float* RESTRICT StereoBuffer, int32 NumSamples, const float* RESTRICT StartGains, const float* RESTRICT EndGains);
 
 	/** Takes a 1 channel buffer and mixes it to a stereo buffer using Gains. Gains is expected to point to a 2 float long buffer.
 	 *  these buffers must have an even number of frames.
 	 *  If StartGains and EndGains are provided, this function will interpolate between the two across the buffer.
 	*/
-	void MixMonoTo2ChannelsFast(const AlignedFloatBuffer& MonoBuffer, AlignedFloatBuffer& DestinationBuffer, const float* RESTRICT Gains);
-	void MixMonoTo2ChannelsFast(const float* RESTRICT MonoBuffer, float* RESTRICT DestinationBuffer, int32 NumFrames, const float* RESTRICT Gains);
-	void MixMonoTo2ChannelsFast(const AlignedFloatBuffer& MonoBuffer, AlignedFloatBuffer& DestinationBuffer, const float* RESTRICT StartGains, const float* RESTRICT EndGains);
-	void MixMonoTo2ChannelsFast(const float* RESTRICT MonoBuffer, float* RESTRICT DestinationBuffer, int32 NumFrames, const float* RESTRICT StartGains, const float* RESTRICT EndGains);
+	void AUDIOMIXER_API MixMonoTo2ChannelsFast(const AlignedFloatBuffer& MonoBuffer, AlignedFloatBuffer& DestinationBuffer, const float* RESTRICT Gains);
+	void AUDIOMIXER_API MixMonoTo2ChannelsFast(const float* RESTRICT MonoBuffer, float* RESTRICT DestinationBuffer, int32 NumFrames, const float* RESTRICT Gains);
+	void AUDIOMIXER_API MixMonoTo2ChannelsFast(const AlignedFloatBuffer& MonoBuffer, AlignedFloatBuffer& DestinationBuffer, const float* RESTRICT StartGains, const float* RESTRICT EndGains);
+	void AUDIOMIXER_API MixMonoTo2ChannelsFast(const float* RESTRICT MonoBuffer, float* RESTRICT DestinationBuffer, int32 NumFrames, const float* RESTRICT StartGains, const float* RESTRICT EndGains);
+	void AUDIOMIXER_API MixMonoTo2ChannelsFast(const AlignedFloatBuffer& MonoBuffer, AlignedFloatBuffer& DestinationBuffer);
+	void AUDIOMIXER_API MixMonoTo2ChannelsFast(const float* RESTRICT MonoBuffer, float* RESTRICT DestinationBuffer, int32 NumFrames);
 
 	/** Takes a 2 channel buffer and mixes it to an 2 channel interleaved buffer using Gains. Gains is expected to point to a 16 float long buffer.
 	*  Output gains for the left input channel should be the first 8 values in Gains, and Output gains for the right input channel should be rest.
 	*  NumFrames must be a multiple of 4.
 	*  If StartGains and EndGains are provided, this function will interpolate between the two across the buffer.
 	*/
-	void Mix2ChannelsTo2ChannelsFast(const AlignedFloatBuffer& SourceBuffer, AlignedFloatBuffer& DestinationBuffer, const float* RESTRICT Gains);
-	void Mix2ChannelsTo2ChannelsFast(const float* RESTRICT SourceBuffer, float* RESTRICT DestinationBuffer, int32 NumFrames, const float* RESTRICT Gains);
-	void Mix2ChannelsTo2ChannelsFast(const AlignedFloatBuffer& SourceBuffer, AlignedFloatBuffer& DestinationBuffer, const float* RESTRICT StartGains, const float* RESTRICT EndGains);
-	void Mix2ChannelsTo2ChannelsFast(const float* RESTRICT SourceBuffer, float* RESTRICT DestinationBuffer, int32 NumFrames, const float* RESTRICT StartGains, const float* RESTRICT EndGains);
+	void AUDIOMIXER_API Mix2ChannelsTo2ChannelsFast(const AlignedFloatBuffer& SourceBuffer, AlignedFloatBuffer& DestinationBuffer, const float* RESTRICT Gains);
+	void AUDIOMIXER_API Mix2ChannelsTo2ChannelsFast(const float* RESTRICT SourceBuffer, float* RESTRICT DestinationBuffer, int32 NumFrames, const float* RESTRICT Gains);
+	void AUDIOMIXER_API Mix2ChannelsTo2ChannelsFast(const AlignedFloatBuffer& SourceBuffer, AlignedFloatBuffer& DestinationBuffer, const float* RESTRICT StartGains, const float* RESTRICT EndGains);
+	void AUDIOMIXER_API Mix2ChannelsTo2ChannelsFast(const float* RESTRICT SourceBuffer, float* RESTRICT DestinationBuffer, int32 NumFrames, const float* RESTRICT StartGains, const float* RESTRICT EndGains);
 
 	/** Takes a 4 channel interleaved buffer and applies Gains to it. Gains is expected to point to a 2 float long buffer.
 	 *  If StartGains and EndGains are provided, this function will interpolate between the two across the buffer.
 	*/
-	void Apply4ChannelGain(AlignedFloatBuffer& InterleavedBuffer, const float* RESTRICT Gains);
-	void Apply4ChannelGain(float* RESTRICT InterleavedBuffer, int32 NumSamples, const float* RESTRICT Gains);
-	void Apply4ChannelGain(AlignedFloatBuffer& InterleavedBuffer, const float* RESTRICT StartGains, const float* RESTRICT EndGains);
-	void Apply4ChannelGain(float* RESTRICT InterleavedBuffer, int32 NumSamples, const float* RESTRICT StartGains, const float* RESTRICT EndGains);
+	void AUDIOMIXER_API Apply4ChannelGain(AlignedFloatBuffer& InterleavedBuffer, const float* RESTRICT Gains);
+	void AUDIOMIXER_API Apply4ChannelGain(float* RESTRICT InterleavedBuffer, int32 NumSamples, const float* RESTRICT Gains);
+	void AUDIOMIXER_API Apply4ChannelGain(AlignedFloatBuffer& InterleavedBuffer, const float* RESTRICT StartGains, const float* RESTRICT EndGains);
+	void AUDIOMIXER_API Apply4ChannelGain(float* RESTRICT InterleavedBuffer, int32 NumSamples, const float* RESTRICT StartGains, const float* RESTRICT EndGains);
 
 	/** Takes a 1 channel buffer and mixes it to an 8 channel interleaved buffer using Gains. Gains is expected to point to a 8 float long buffer.
 	*  these buffers must have an even number of frames.
 	*  If StartGains and EndGains are provided, this function will interpolate between the two across the buffer.
 	*/
-	void MixMonoTo4ChannelsFast(const AlignedFloatBuffer& MonoBuffer, AlignedFloatBuffer& DestinationBuffer, const float* RESTRICT Gains);
-	void MixMonoTo4ChannelsFast(const float* RESTRICT MonoBuffer, float* RESTRICT DestinationBuffer, int32 NumFrames, const float* RESTRICT Gains);
-	void MixMonoTo4ChannelsFast(const AlignedFloatBuffer& MonoBuffer, AlignedFloatBuffer& DestinationBuffer, const float* RESTRICT StartGains, const float* RESTRICT EndGains);
-	void MixMonoTo4ChannelsFast(const float* RESTRICT MonoBuffer, float* RESTRICT DestinationBuffer, int32 NumFrames, const float* RESTRICT StartGains, const float* RESTRICT EndGains);
+	void AUDIOMIXER_API MixMonoTo4ChannelsFast(const AlignedFloatBuffer& MonoBuffer, AlignedFloatBuffer& DestinationBuffer, const float* RESTRICT Gains);
+	void AUDIOMIXER_API MixMonoTo4ChannelsFast(const float* RESTRICT MonoBuffer, float* RESTRICT DestinationBuffer, int32 NumFrames, const float* RESTRICT Gains);
+	void AUDIOMIXER_API MixMonoTo4ChannelsFast(const AlignedFloatBuffer& MonoBuffer, AlignedFloatBuffer& DestinationBuffer, const float* RESTRICT StartGains, const float* RESTRICT EndGains);
+	void AUDIOMIXER_API MixMonoTo4ChannelsFast(const float* RESTRICT MonoBuffer, float* RESTRICT DestinationBuffer, int32 NumFrames, const float* RESTRICT StartGains, const float* RESTRICT EndGains);
 
 	/** Takes a 2 channel buffer and mixes it to an 8 channel interleaved buffer using Gains. Gains is expected to point to a 16 float long buffer.
 	*  Output gains for the left input channel should be the first 8 values in Gains, and Output gains for the right input channel should be rest.
 	*  NumFrames must be a multiple of 4.
 	*  If StartGains and EndGains are provided, this function will interpolate between the two across the buffer.
 	*/
-	void Mix2ChannelsTo4ChannelsFast(const AlignedFloatBuffer& SourceBuffer, AlignedFloatBuffer& DestinationBuffer, const float* RESTRICT Gains);
-	void Mix2ChannelsTo4ChannelsFast(const float* RESTRICT SourceBuffer, float* RESTRICT DestinationBuffer, int32 NumFrames, const float* RESTRICT Gains);
-	void Mix2ChannelsTo4ChannelsFast(const AlignedFloatBuffer& SourceBuffer, AlignedFloatBuffer& DestinationBuffer, const float* RESTRICT StartGains, const float* RESTRICT EndGains);
-	void Mix2ChannelsTo4ChannelsFast(const float* RESTRICT SourceBuffer, float* RESTRICT DestinationBuffer, int32 NumFrames, const float* RESTRICT StartGains, const float* RESTRICT EndGains);
+	void AUDIOMIXER_API Mix2ChannelsTo4ChannelsFast(const AlignedFloatBuffer& SourceBuffer, AlignedFloatBuffer& DestinationBuffer, const float* RESTRICT Gains);
+	void AUDIOMIXER_API Mix2ChannelsTo4ChannelsFast(const float* RESTRICT SourceBuffer, float* RESTRICT DestinationBuffer, int32 NumFrames, const float* RESTRICT Gains);
+	void AUDIOMIXER_API Mix2ChannelsTo4ChannelsFast(const AlignedFloatBuffer& SourceBuffer, AlignedFloatBuffer& DestinationBuffer, const float* RESTRICT StartGains, const float* RESTRICT EndGains);
+	void AUDIOMIXER_API Mix2ChannelsTo4ChannelsFast(const float* RESTRICT SourceBuffer, float* RESTRICT DestinationBuffer, int32 NumFrames, const float* RESTRICT StartGains, const float* RESTRICT EndGains);
 
 	/** Takes a 6 channel interleaved buffer and applies Gains to it. Gains is expected to point to a 2 float long buffer.
 	 *  InterleavedBuffer must have an even number of frames.
 	*/
-	void Apply6ChannelGain(AlignedFloatBuffer& InterleavedBuffer, const float* RESTRICT Gains);
-	void Apply6ChannelGain(float* RESTRICT InterleavedBuffer, int32 NumSamples, const float* RESTRICT Gains);
-	void Apply6ChannelGain(AlignedFloatBuffer& InterleavedBuffer, const float* RESTRICT StartGains, const float* RESTRICT EndGains);
-	void Apply6ChannelGain(float* RESTRICT InterleavedBuffer, int32 NumSamples, const float* RESTRICT StartGains, const float* RESTRICT EndGains);
+	void AUDIOMIXER_API Apply6ChannelGain(AlignedFloatBuffer& InterleavedBuffer, const float* RESTRICT Gains);
+	void AUDIOMIXER_API Apply6ChannelGain(float* RESTRICT InterleavedBuffer, int32 NumSamples, const float* RESTRICT Gains);
+	void AUDIOMIXER_API Apply6ChannelGain(AlignedFloatBuffer& InterleavedBuffer, const float* RESTRICT StartGains, const float* RESTRICT EndGains);
+	void AUDIOMIXER_API Apply6ChannelGain(float* RESTRICT InterleavedBuffer, int32 NumSamples, const float* RESTRICT StartGains, const float* RESTRICT EndGains);
 
 	/** Takes a 1 channel buffer and mixes it to an 8 channel interleaved buffer using Gains. Gains is expected to point to a 8 float long buffer.
 	 *  these buffers must have an even number of frames.
 	 *  If StartGains and EndGains are provided, this function will interpolate between the two across the buffer.
 	*/
-	void MixMonoTo6ChannelsFast(const AlignedFloatBuffer& MonoBuffer, AlignedFloatBuffer& DestinationBuffer, const float* RESTRICT Gains);
-	void MixMonoTo6ChannelsFast(const float* RESTRICT MonoBuffer, float* RESTRICT DestinationBuffer, int32 NumFrames, const float* RESTRICT Gains);
-	void MixMonoTo6ChannelsFast(const AlignedFloatBuffer& MonoBuffer, AlignedFloatBuffer& DestinationBuffer, const float* RESTRICT StartGains, const float* RESTRICT EndGains);
-	void MixMonoTo6ChannelsFast(const float* RESTRICT MonoBuffer, float* RESTRICT DestinationBuffer, int32 NumFrames, const float* RESTRICT StartGains, const float* RESTRICT EndGains);
+	void AUDIOMIXER_API MixMonoTo6ChannelsFast(const AlignedFloatBuffer& MonoBuffer, AlignedFloatBuffer& DestinationBuffer, const float* RESTRICT Gains);
+	void AUDIOMIXER_API MixMonoTo6ChannelsFast(const float* RESTRICT MonoBuffer, float* RESTRICT DestinationBuffer, int32 NumFrames, const float* RESTRICT Gains);
+	void AUDIOMIXER_API MixMonoTo6ChannelsFast(const AlignedFloatBuffer& MonoBuffer, AlignedFloatBuffer& DestinationBuffer, const float* RESTRICT StartGains, const float* RESTRICT EndGains);
+	void AUDIOMIXER_API MixMonoTo6ChannelsFast(const float* RESTRICT MonoBuffer, float* RESTRICT DestinationBuffer, int32 NumFrames, const float* RESTRICT StartGains, const float* RESTRICT EndGains);
 
 	/** Takes a 2 channel buffer and mixes it to an 8 channel interleaved buffer using Gains. Gains is expected to point to a 16 float long buffer.
 	 *  Output gains for the left input channel should be the first 8 values in Gains, and Output gains for the right input channel should be rest.
 	 *  NumFrames must be a multiple of 4.
 	 *  If StartGains and EndGains are provided, this function will interpolate between the two across the buffer.
 	*/
-	void Mix2ChannelsTo6ChannelsFast(const AlignedFloatBuffer& SourceBuffer, AlignedFloatBuffer& DestinationBuffer, const float* RESTRICT Gains);
-	void Mix2ChannelsTo6ChannelsFast(const float* RESTRICT SourceBuffer, float* RESTRICT DestinationBuffer, int32 NumFrames, const float* RESTRICT Gains);
-	void Mix2ChannelsTo6ChannelsFast(const AlignedFloatBuffer& SourceBuffer, AlignedFloatBuffer& DestinationBuffer, const float* RESTRICT StartGains, const float* RESTRICT EndGains);
-	void Mix2ChannelsTo6ChannelsFast(const float* RESTRICT SourceBuffer, float* RESTRICT DestinationBuffer, int32 NumFrames, const float* RESTRICT StartGains, const float* RESTRICT EndGains);
+	void AUDIOMIXER_API Mix2ChannelsTo6ChannelsFast(const AlignedFloatBuffer& SourceBuffer, AlignedFloatBuffer& DestinationBuffer, const float* RESTRICT Gains);
+	void AUDIOMIXER_API Mix2ChannelsTo6ChannelsFast(const float* RESTRICT SourceBuffer, float* RESTRICT DestinationBuffer, int32 NumFrames, const float* RESTRICT Gains);
+	void AUDIOMIXER_API Mix2ChannelsTo6ChannelsFast(const AlignedFloatBuffer& SourceBuffer, AlignedFloatBuffer& DestinationBuffer, const float* RESTRICT StartGains, const float* RESTRICT EndGains);
+	void AUDIOMIXER_API Mix2ChannelsTo6ChannelsFast(const float* RESTRICT SourceBuffer, float* RESTRICT DestinationBuffer, int32 NumFrames, const float* RESTRICT StartGains, const float* RESTRICT EndGains);
 
 	/** Takes an 8 channel interleaved buffer and applies Gains to it. Gains is expected to point to an 8 float long buffer. */
-	void Apply8ChannelGain(AlignedFloatBuffer& InterleavedBuffer, const float* RESTRICT Gains);
-	void Apply8ChannelGain(float* RESTRICT InterleavedBuffer, int32 NumSamples, const float* RESTRICT Gains);
-	void Apply8ChannelGain(AlignedFloatBuffer& InterleavedBuffer, const float* RESTRICT StartGains, const float* RESTRICT EndGains);
-	void Apply8ChannelGain(float* RESTRICT InterleavedBuffer, int32 NumSamples, const float* RESTRICT StartGains, const float* RESTRICT EndGains);
+	void AUDIOMIXER_API Apply8ChannelGain(AlignedFloatBuffer& InterleavedBuffer, const float* RESTRICT Gains);
+	void AUDIOMIXER_API Apply8ChannelGain(float* RESTRICT InterleavedBuffer, int32 NumSamples, const float* RESTRICT Gains);
+	void AUDIOMIXER_API Apply8ChannelGain(AlignedFloatBuffer& InterleavedBuffer, const float* RESTRICT StartGains, const float* RESTRICT EndGains);
+	void AUDIOMIXER_API Apply8ChannelGain(float* RESTRICT InterleavedBuffer, int32 NumSamples, const float* RESTRICT StartGains, const float* RESTRICT EndGains);
 
 	/** Takes a 1 channel buffer and mixes it to an 8 channel interleaved buffer using Gains. Gains is expected to point to a 8 float long buffer.
 	 *  these buffers must have an even number of frames.
 	 *  If StartGains and EndGains are provided, this function will interpolate between the two across the buffer.
 	*/
-	void MixMonoTo8ChannelsFast(const AlignedFloatBuffer& MonoBuffer, AlignedFloatBuffer& DestinationBuffer, const float* RESTRICT Gains);
-	void MixMonoTo8ChannelsFast(const float* RESTRICT MonoBuffer, float* RESTRICT DestinationBuffer, int32 NumFrames, const float* RESTRICT Gains);
-	void MixMonoTo8ChannelsFast(const AlignedFloatBuffer& MonoBuffer, AlignedFloatBuffer& DestinationBuffer, const float* RESTRICT StartGains, const float* RESTRICT EndGains);
-	void MixMonoTo8ChannelsFast(const float* RESTRICT MonoBuffer, float* RESTRICT DestinationBuffer, int32 NumFrames, const float* RESTRICT StartGains, const float* RESTRICT EndGains);
+	void AUDIOMIXER_API MixMonoTo8ChannelsFast(const AlignedFloatBuffer& MonoBuffer, AlignedFloatBuffer& DestinationBuffer, const float* RESTRICT Gains);
+	void AUDIOMIXER_API MixMonoTo8ChannelsFast(const float* RESTRICT MonoBuffer, float* RESTRICT DestinationBuffer, int32 NumFrames, const float* RESTRICT Gains);
+	void AUDIOMIXER_API MixMonoTo8ChannelsFast(const AlignedFloatBuffer& MonoBuffer, AlignedFloatBuffer& DestinationBuffer, const float* RESTRICT StartGains, const float* RESTRICT EndGains);
+	void AUDIOMIXER_API MixMonoTo8ChannelsFast(const float* RESTRICT MonoBuffer, float* RESTRICT DestinationBuffer, int32 NumFrames, const float* RESTRICT StartGains, const float* RESTRICT EndGains);
 
 	/** Takes a 2 channel buffer and mixes it to an 8 channel interleaved buffer using Gains. Gains is expected to point to a 16 float long buffer.
 	 *  Output gains for the left input channel should be the first 8 values in Gains, and Output gains for the right input channel should be rest.
 	 *  these buffers must have an even number of frames.
 	 *  If StartGains and EndGains are provided, this function will interpolate between the two across the buffer.
 	*/
-	void Mix2ChannelsTo8ChannelsFast(const AlignedFloatBuffer& SourceBuffer, AlignedFloatBuffer& DestinationBuffer, const float* RESTRICT Gains);
-	void Mix2ChannelsTo8ChannelsFast(const float* RESTRICT SourceBuffer, float* RESTRICT DestinationBuffer, int32 NumFrames, const float* RESTRICT Gains);
-	void Mix2ChannelsTo8ChannelsFast(const AlignedFloatBuffer& SourceBuffer, AlignedFloatBuffer& DestinationBuffer, const float* RESTRICT StartGains, const float* RESTRICT EndGains);
-	void Mix2ChannelsTo8ChannelsFast(const float* RESTRICT SourceBuffer, float* RESTRICT DestinationBuffer, int32 NumFrames, const float* RESTRICT StartGains, const float* RESTRICT EndGains);
+	void AUDIOMIXER_API Mix2ChannelsTo8ChannelsFast(const AlignedFloatBuffer& SourceBuffer, AlignedFloatBuffer& DestinationBuffer, const float* RESTRICT Gains);
+	void AUDIOMIXER_API Mix2ChannelsTo8ChannelsFast(const float* RESTRICT SourceBuffer, float* RESTRICT DestinationBuffer, int32 NumFrames, const float* RESTRICT Gains);
+	void AUDIOMIXER_API Mix2ChannelsTo8ChannelsFast(const AlignedFloatBuffer& SourceBuffer, AlignedFloatBuffer& DestinationBuffer, const float* RESTRICT StartGains, const float* RESTRICT EndGains);
+	void AUDIOMIXER_API Mix2ChannelsTo8ChannelsFast(const float* RESTRICT SourceBuffer, float* RESTRICT DestinationBuffer, int32 NumFrames, const float* RESTRICT StartGains, const float* RESTRICT EndGains);
 
-	/** This is a generalized operation that uses the channel gain matrix provied in Gains to mix an interleaved source buffer to the interleaved downmix buffer.
+	/** This is a generalized operation that uses the channel gain matrix provided in Gains to mix an interleaved source buffer to the interleaved downmix buffer.
 	 *  This operation is not explicitly vectorized and will almost always be slower than using one of the functions above.
 	*/
-	void DownmixBuffer(int32 NumSourceChannels, int32 NumDestinationChannels, const AlignedFloatBuffer& SourceBuffer, AlignedFloatBuffer& DestinationBuffer, const float* RESTRICT Gains);
-	void DownmixBuffer(int32 NumSourceChannels, int32 NumDestinationChannels, const float* RESTRICT SourceBuffer, float* RESTRICT DestinationBuffer, int32 NumFrames, const float* RESTRICT Gains);
-	void DownmixBuffer(int32 NumSourceChannels, int32 NumDestinationChannels, const AlignedFloatBuffer& SourceBuffer, AlignedFloatBuffer& DestinationBuffer, float* RESTRICT StartGains, const float* RESTRICT EndGains);
-	void DownmixBuffer(int32 NumSourceChannels, int32 NumDestinationChannels, const float* RESTRICT SourceBuffer, float* RESTRICT DestinationBuffer, int32 NumFrames, float* RESTRICT StartGains, const float* RESTRICT EndGains);
+	void AUDIOMIXER_API DownmixBuffer(int32 NumSourceChannels, int32 NumDestinationChannels, const AlignedFloatBuffer& SourceBuffer, AlignedFloatBuffer& DestinationBuffer, const float* RESTRICT Gains);
+	void AUDIOMIXER_API DownmixBuffer(int32 NumSourceChannels, int32 NumDestinationChannels, const float* RESTRICT SourceBuffer, float* RESTRICT DestinationBuffer, int32 NumFrames, const float* RESTRICT Gains);
+	void AUDIOMIXER_API DownmixBuffer(int32 NumSourceChannels, int32 NumDestinationChannels, const AlignedFloatBuffer& SourceBuffer, AlignedFloatBuffer& DestinationBuffer, float* RESTRICT StartGains, const float* RESTRICT EndGains);
+	void AUDIOMIXER_API DownmixBuffer(int32 NumSourceChannels, int32 NumDestinationChannels, const float* RESTRICT SourceBuffer, float* RESTRICT DestinationBuffer, int32 NumFrames, float* RESTRICT StartGains, const float* RESTRICT EndGains);
+
+
+	/** Interleaves samples from two input buffers */
+	void AUDIOMIXER_API BufferInterleave2ChannelFast(const AlignedFloatBuffer& InBuffer1, const AlignedFloatBuffer& InBuffer2, AlignedFloatBuffer& OutBuffer);
+
+	/** Interleaves samples from two input buffers */
+	void AUDIOMIXER_API BufferInterleave2ChannelFast(const float* RESTRICT InBuffer1, const float* RESTRICT InBuffer2, float* RESTRICT OutBuffer, const int32 InNum);
+
+	/** Deinterleaves samples from a 2 channel input buffer */
+	void AUDIOMIXER_API BufferDeinterleave2ChannelFast(const AlignedFloatBuffer& InBuffer, AlignedFloatBuffer& OutBuffer1, AlignedFloatBuffer& OutBuffer2);
+
+	/** Deinterleaves samples from a 2 channel input buffer */
+	void AUDIOMIXER_API BufferDeinterleave2ChannelFast(const float* RESTRICT InBuffer, float* RESTRICT OutBuffer1, float* RESTRICT OutBuffer2, const int32 InNum);
+
+	/** Sums 2 channel interleaved input samples. OutSamples[n] = InSamples[2n] + InSamples[2n + 1] */
+	void AUDIOMIXER_API BufferSum2ChannelToMonoFast(const AlignedFloatBuffer& InSamples, AlignedFloatBuffer& OutSamples);
+
+	/** Sums 2 channel interleaved input samples. OutSamples[n] = InSamples[2n] + InSamples[2n + 1] */
+	void AUDIOMIXER_API BufferSum2ChannelToMonoFast(const float* RESTRICT InSamples, float* RESTRICT OutSamples, const int32 InNumFrames);
 }
diff --git a/Engine/Source/Runtime/AudioMixer/Public/DSP/Dsp.h b/Engine/Source/Runtime/AudioMixer/Public/DSP/Dsp.h
--- a/Engine/Source/Runtime/AudioMixer/Public/DSP/Dsp.h
+++ b/Engine/Source/Runtime/AudioMixer/Public/DSP/Dsp.h
@@ -59,6 +59,7 @@ namespace Audio
 	}
 
 	// Low precision, high performance approximation of sine using parabolic polynomial approx
+	// Valid on interval [-PI, PI]
 	static FORCEINLINE float FastSin(const float X)
 	{
 		return (4.0f * X) / PI * (1.0f - FMath::Abs(X) / PI);
diff --git a/Engine/Source/Runtime/AudioMixer/Public/DSP/DynamicDelayAPF.h b/Engine/Source/Runtime/AudioMixer/Public/DSP/DynamicDelayAPF.h
new file mode 100644
--- /dev/null
+++ b/Engine/Source/Runtime/AudioMixer/Public/DSP/DynamicDelayAPF.h
@@ -0,0 +1,62 @@
+// Copyright 1998-2019 Epic Games, Inc. All Rights Reserved.
+
+#pragma once
+
+#include "CoreMinimal.h"
+#include "IntegerDelay.h"
+#include "BufferVectorOperations.h"
+#include "AllPassFractionalDelay.h"
+#include "LinearInterpFractionalDelay.h"
+
+namespace Audio
+{
+	// All Pass Filter with a long fractional delay which can be set per a sample. This filter is specifically designed for 
+	// reverb applications where filter delay lines are long.
+	class AUDIOMIXER_API FDynamicDelayAPF
+	{
+	public:
+		// InG is the filter coefficient used in the long delay all pass filter.
+		// InMinDelay is the minimum allowable delay of the all pass filter.
+		// InMaxDelay is the maximum allowable delay of the all pass filter.
+		// InMaxNumInternalBufferSamples is the maximum internal block size used internally. 
+		FDynamicDelayAPF(float InG, int32 InMinDelay, int32 InMaxDelay, int32 InMaxNumInternalBufferSamples);
+
+		// Destructor
+		~FDynamicDelayAPF();
+
+		// Set the APF feedback/feedforward gain coefficient
+		void SetG(float InG) { G = InG; }
+
+		// Processes InSamples through the all pass filter and populates OutSamples with the filter output.
+		// InDelays denotes the per-sample delay of the allpass. It must have an equal number of elements as InSamples
+		void ProcessAudio(const AlignedFloatBuffer& InSamples, const AlignedFloatBuffer& InDelays, AlignedFloatBuffer& OutSamples);
+
+		// Zeros the internal delay line.
+		void Reset();
+
+	protected:
+
+		// Process one block of audio.
+		void ProcessAudioBlock(const float* InSamples, const AlignedFloatBuffer& InFractionalDelays, const int32 InNum, float* OutSamples);
+
+	private:
+		// Feedback/Feedforward gain coefficient
+		float G;
+		int32 MinDelay;
+		int32 MaxDelay;
+		int32 NumDelaySamples;
+		int32 NumInternalBufferSamples;
+
+		// Buffers for block processing
+		AlignedFloatBuffer FractionalDelays;
+		AlignedFloatBuffer DelayLineInput;
+		AlignedFloatBuffer WorkBufferA;
+		AlignedFloatBuffer WorkBufferB;
+
+		// Delay line memory.
+		TUniquePtr<FAlignedBlockBuffer> IntegerDelayLine;
+		TUniquePtr<FLinearInterpFractionalDelay> FractionalDelayLine;
+
+	};
+
+}
diff --git a/Engine/Source/Runtime/AudioMixer/Public/DSP/EarlyReflectionsFast.h b/Engine/Source/Runtime/AudioMixer/Public/DSP/EarlyReflectionsFast.h
new file mode 100644
--- /dev/null
+++ b/Engine/Source/Runtime/AudioMixer/Public/DSP/EarlyReflectionsFast.h
@@ -0,0 +1,97 @@
+// Copyright 1998-2019 Epic Games, Inc. All Rights Reserved.
+
+#pragma once
+
+#include "CoreMinimal.h"
+#include "DSP/BufferVectorOperations.h"
+#include "DSP/IntegerDelay.h"
+#include "DSP/LongDelayAPF.h"
+#include "DSP/BufferOnePoleLPF.h"
+#include "DSP/FeedbackDelayNetwork.h"
+
+namespace Audio
+{
+	struct AUDIOMIXER_API FEarlyReflectionsFastSettings
+	{
+		// Early reflections gain
+		float Gain;
+
+		// Delay between input signal and early reflections
+		float PreDelayMsec;
+
+		// Input sample bandwidth before entering early reflections
+		float Bandwidth;
+
+		// Early reflections decay (lower value is longer)
+		float Decay;
+
+		// Early reflection high frequency absorption factor
+		float Absorption;
+
+		FEarlyReflectionsFastSettings();
+
+		bool operator==(const FEarlyReflectionsFastSettings& Other) const;
+
+		bool operator!=(const FEarlyReflectionsFastSettings& Other) const;
+	};
+
+	// Basic implementation of early reflections using a predelay, low pass filter and feedback delay network (FDN). The FDN
+	// utilizes four delay lines where each delay line consists of an all pass filter and low pass filter to control diffusion
+	// and absorption.
+	class AUDIOMIXER_API FEarlyReflectionsFast
+	{
+	public:
+		
+		// Limits for settings
+		static const float MaxGain;
+		static const float MinGain;
+		static const float MaxPreDelay;
+		static const float MinPreDelay;
+		static const float MaxBandwidth;
+		static const float MinBandwidth;
+		static const float MaxDecay;
+		static const float MinDecay;
+		static const float MaxAbsorption;
+		static const float MinAbsorption;
+		
+		static const FEarlyReflectionsFastSettings DefaultSettings;
+		
+		// InMaxNumInternalBufferSamples sets the maximum possible samples in an internal buffer.
+		FEarlyReflectionsFast(float InSampleRate, int32 InMaxNumInternalBufferSamples, const FEarlyReflectionsFastSettings& InSettings=DefaultSettings);
+		~FEarlyReflectionsFast();
+
+		// Sets the reverb settings, clamps, applies, and updates
+		void SetSettings(const FEarlyReflectionsFastSettings& InSettings);
+
+		// Process the single audio frame
+		void ProcessAudio(const AlignedFloatBuffer& InSamples, const int32 InNumChannels, AlignedFloatBuffer& OutLeftSamples, AlignedFloatBuffer& OutRightSamples);
+
+		// Clamps settings to acceptable values. 
+		static void ClampSettings(FEarlyReflectionsFastSettings& InOutSettings);
+
+	private:
+		void ApplySettings();
+
+		FEarlyReflectionsFastSettings Settings;
+		float SampleRate;
+
+		AlignedFloatBuffer LeftInputBuffer;
+		AlignedFloatBuffer LeftWorkBufferA;
+		AlignedFloatBuffer LeftWorkBufferB;
+		AlignedFloatBuffer RightInputBuffer;
+		AlignedFloatBuffer RightWorkBufferA;
+		AlignedFloatBuffer RightWorkBufferB;
+
+		FFDNCoefficients LeftCoefficients;
+		FFDNCoefficients RightCoefficients;
+
+		FFeedbackDelayNetwork LeftFDN;
+		FFeedbackDelayNetwork RightFDN;
+
+		FIntegerDelay LeftPreDelay;
+		FIntegerDelay RightPreDelay;
+
+		FBufferOnePoleLPF LeftInputLPF;
+		FBufferOnePoleLPF RightInputLPF;
+	};
+}
diff --git a/Engine/Source/Runtime/AudioMixer/Public/DSP/FeedbackDelayNetwork.h b/Engine/Source/Runtime/AudioMixer/Public/DSP/FeedbackDelayNetwork.h
new file mode 100644
--- /dev/null
+++ b/Engine/Source/Runtime/AudioMixer/Public/DSP/FeedbackDelayNetwork.h
@@ -0,0 +1,94 @@
+
+// Copyright 1998-2019 Epic Games, Inc. All Rights Reserved.
+
+#pragma once
+
+#include "CoreMinimal.h"
+#include "DSP/BufferVectorOperations.h"
+#include "DSP/IntegerDelay.h"
+#include "DSP/LongDelayAPF.h"
+#include "DSP/BufferOnePoleLPF.h"
+
+namespace Audio
+{
+	// Delay settings of 4 channel feedback delay network
+	struct AUDIOMIXER_API FFDNDelaySettings
+	{
+		// Number of samples for delay line 0
+		int32 APF0DelayNumSamples;
+		// Number of samples for delay line 1
+		int32 APF1DelayNumSamples;
+		// Number of samples for delay line 2
+		int32 APF2DelayNumSamples;
+		// Number of samples for delay line 3
+		int32 APF3DelayNumSamples;
+		
+		// Default delay line settings for left channel.
+		static FFDNDelaySettings DefaultLeftDelays(float InSampleRate);
+		// Default delay line settings for right channel. 
+		static FFDNDelaySettings DefaultRightDelays(float InSampleRate);
+	};
+
+	// Filter coefficients of 4 channel feedback delay network.
+	struct AUDIOMIXER_API FFDNCoefficients
+	{
+		// Sample multiplier of input samples before they enter the delay line.
+		float InputScale;
+		// All pass filter coefficients. One coefficient for each filter.
+		float APFG[4];
+		// Low pass filter feed back coefficients. One coefficient for each filter.
+		float LPFA[4];
+		// Low pass filter feed forward coefficients. One coefficient for each filter.
+		float LPFB[4];
+		// Feedback multiplier of mixing matrix output to network input.
+		float Feedback;
+	};
+
+	// 4 channel feedback delay network (FDN) for artificial reverberation.
+	class AUDIOMIXER_API FFeedbackDelayNetwork
+	{
+	public:
+		// InMaxNumInternalBufferSamples controls the internal buffer size used for vector operations. 
+		// InSettings controls the delay line lengths.
+		FFeedbackDelayNetwork(int32 InMaxNumInternalBufferSamples, const FFDNDelaySettings& InSettings);
+
+		~FFeedbackDelayNetwork();
+
+		// Sets the coefficient values of the all pass filters, low pass filters, input scalers and feedback scalers.
+		void SetCoefficients(const FFDNCoefficients& InCoefficients);
+
+		// Generates artificial reverberation for InSamples and places results in OutSamples.
+		void ProcessAudio(const AlignedFloatBuffer& InSamples, AlignedFloatBuffer& OutSamples);
+
+		// Returns the actual number of samples in an internal buffer.
+		int32 GetNumInternalBufferSamples() const;
+
+
+	private:
+
+		// Process single internal buffer of audio.
+		void ProcessAudioBuffer(const float* InSampleData, const int32 InNum, float* OutSampleData);
+
+		// Internal buffer size.
+		int32 NumInternalBufferSamples;
+
+		// Coefficients structure for APF, LPF, input scaling and feedback.
+		FFDNCoefficients Coefficients;
+
+		// Previous output of low pass filters.
+		float LPFZ[4];
+
+		// Previous output of mixing matrix.
+		float FMO[4];
+
+		TUniquePtr<FAlignedBlockBuffer> DelayLine0;
+		TUniquePtr<FAlignedBlockBuffer> DelayLine1;
+		TUniquePtr<FAlignedBlockBuffer> DelayLine2;
+		TUniquePtr<FAlignedBlockBuffer> DelayLine3;
+
+		AlignedFloatBuffer WorkBuffer0;
+		AlignedFloatBuffer WorkBuffer1;
+		AlignedFloatBuffer WorkBuffer2;
+		AlignedFloatBuffer WorkBuffer3;
+	};
+}
diff --git a/Engine/Source/Runtime/AudioMixer/Public/DSP/IntegerDelay.h b/Engine/Source/Runtime/AudioMixer/Public/DSP/IntegerDelay.h
new file mode 100644
--- /dev/null
+++ b/Engine/Source/Runtime/AudioMixer/Public/DSP/IntegerDelay.h
@@ -0,0 +1,63 @@
+// Copyright 1998-2019 Epic Games, Inc. All Rights Reserved.
+
+#pragma once
+
+#include "CoreMinimal.h"
+#include "DSP/Dsp.h"
+#include "DSP/AlignedBlockBuffer.h"
+#include "DSP/BufferVectorOperations.h"
+
+namespace Audio
+{
+
+	// An adjustable delay line. Delays values are limited to integer values. 
+	class AUDIOMIXER_API FIntegerDelay
+	{
+	public:
+
+		// InMaxDelaySamples is the maximum supported delay.
+		// InDelaySamples is the initial delay in samples.
+		FIntegerDelay(int32 InMaxNumDelaySamples, int32 InNumDelaySamples, int32 InNumInternalBufferSamples = 256);
+
+		// Destructor
+		~FIntegerDelay();
+
+		// Sets the current delay in samples. InDelay must be less than or equal to the InMaxDelaySamples set in the constructor.
+		void SetDelayLengthSamples(int32 InNumDelaySamples);
+
+		// Resets the delay line state, flushes buffer and resets read/write pointers.
+		void Reset();
+
+		// Returns the current delay length in samples.
+		int32 GetNumDelaySamples() const;
+
+		// Process InSamples, placing delayed versions in OutSamples.
+		void ProcessAudio(const Audio::AlignedFloatBuffer& InSamples, Audio::AlignedFloatBuffer& OutSamples);
+
+		// Retrieve a copy of the internal delay line.
+		// InNum must be less than or equal to InMaxNumDelaySamples.
+		void PeekDelayLine(int32 InNum, Audio::AlignedFloatBuffer& OutSamples);
+
+	private:
+		// Process a block of audio. InNum is always less than or equal to NumInternalBufferSamples.
+		void ProcessAudioBlock(const float* InSamples, const int32 InNum, float* OutSamples);
+
+		// Maximum supported number of delay samples
+		int32 MaxNumDelaySamples;
+
+		// Current number of delay samples
+		int32 NumDelaySamples;
+
+		// Offset of output delayed samples relative to DelayLine block buffer
+		int32 NumDelayLineOffsetSamples;
+
+		// Offset to handle delay values which are not multiple of required buffer alignment.
+		int32 NumBufferOffsetSamples;
+
+		// Number of samples in an internal buffer
+		int32 NumInternalBufferSamples;
+
+		// Buffer for holding delayed data.
+		TUniquePtr<FAlignedBlockBuffer> DelayLine;
+	};
+}
diff --git a/Engine/Source/Runtime/AudioMixer/Public/DSP/LateReflectionsFast.h b/Engine/Source/Runtime/AudioMixer/Public/DSP/LateReflectionsFast.h
new file mode 100644
--- /dev/null
+++ b/Engine/Source/Runtime/AudioMixer/Public/DSP/LateReflectionsFast.h
@@ -0,0 +1,316 @@
+// Copyright 1998-2019 Epic Games, Inc. All Rights Reserved.
+
+#pragma once
+
+#include "CoreMinimal.h"
+#include "DSP/BufferVectorOperations.h"
+#include "DSP/IntegerDelay.h"
+#include "DSP/LongDelayAPF.h"
+#include "DSP/BufferOnePoleLPF.h"
+#include "DSP/DynamicDelayAPF.h"
+
+namespace Audio
+{
+	// Structure to hold the various delay line tap outputs produced from a plate.
+	struct AUDIOMIXER_API FLateReflectionsPlateOutputs
+	{
+		// Number of tap locations in plate reverb
+		static const int32 NumTaps = 7;
+
+		// Tap sample buffers.
+		AlignedFloatBuffer Taps[NumTaps];
+
+		// Output sample buffer.
+		AlignedFloatBuffer Output;
+
+		// Resizes and zero all buffers in structure.
+		void ResizeAndZero(int32 InNumSamples);
+	};
+
+	// Delay line settings for reverb plate delay
+	// Tap locations are determined by analyzing these delay values. 
+	struct AUDIOMIXER_API FLateReflectionsPlateDelays
+	{
+		/*
+			Setting the delay values affects the tap locations. The audio flowchart below shows the
+			relationship between delay values and tap locations. 
+
+			+---------------+
+			|    Input      |
+			+------+--------+
+				   |
+			+------v--------+
+			| Modulated APF |
+			+------+--------+
+				   |
+			+------+-------+
+			|    DelayA    |
+			+------+-------+  +---------+
+				   |--------->+  Tap 0  |
+			+------v-------+  +---------+
+			|    DelayB    |
+			+------+-------+  +---------+
+				   |---------->  Tap 1  |
+			+------v-------+  +---------+
+			|    DelayC    |
+			+------+-------+  +---------+
+				   |---------->  Tap 2  |
+			+------v-------+  +---------+
+			|    DelayD    |
+			+------+-------+
+				   |
+			+------v--------+
+			|     LPF       |
+			+------+--------+
+				   |
+			+------v-------+  +---------+  +---------+  +---------+  +--------+
+			|     APF      +--> DelayE  +--> |Tap 3  +--> DelayF  +-->  Tap 4 |
+			+------+-------+  +---------+  +---------+  +---------+  +--------+
+				   |
+			+------v-------+
+			|    DelayG    |
+			+------+-------+  +---------+
+				   |---------->  Tap 5  |
+			+------v-------+  +---------+
+			|    DelayH    |
+			+------+-------+  +---------+
+				   |--------->+  Tap 6  |
+			+------v-------+  +---------+
+			|    DelayI    |
+			+------+-------+
+				   |
+			+------v-------+
+			|    Output    |
+			+--------------+
+		 */
+
+		// Base delay samples for modulated delay
+		int32 NumSamplesModulatedBase;
+		// Max delta samples for modulated delay
+		int32 NumSamplesModulatedDelta;
+		// Tap 0 exists at the output of DelayA
+		int32 NumSamplesDelayA;
+		// Tap 1 exists at the output of DelayB
+		int32 NumSamplesDelayB;
+		// Tap 2 exists at the output of DelayC
+		int32 NumSamplesDelayC;
+		// Delay line length
+		int32 NumSamplesDelayD;
+		// All pass filter length
+		int32 NumSamplesAPF;
+		// Tap 3 exists at the output of DelayE
+		int32 NumSamplesDelayE;
+		// Tap 4 exists at the output of DelayF
+		int32 NumSamplesDelayF;
+		// Tap 5 exists at the output of DelayG
+		int32 NumSamplesDelayG;
+		// Tap 6 exists at the output of DelayH
+		int32 NumSamplesDelayH;
+		// Output tap exists at the output of DelayI
+		int32 NumSamplesDelayI;
+
+		// Default delay settings for left channel.
+		static FLateReflectionsPlateDelays DefaultLeftDelays(float InSampleRate);
+
+		// Default delay settings for right channel.
+		static FLateReflectionsPlateDelays DefaultRightDelays(float InSampleRate);
+	};
+
+
+	// Single plate channel for plate reverberation.
+	class AUDIOMIXER_API FLateReflectionsPlate {
+	public:
+		// InMaxNumInternalBufferSamples sets the mmaximum number of samples in an internal buffer.
+		FLateReflectionsPlate(
+			float InSampleRate,
+			int32 InMaxNumInternalBufferSamples,
+			const FLateReflectionsPlateDelays& InDelays);
+
+		~FLateReflectionsPlate();
+
+		// Processing input samples. All input buffers must be of equal length.
+		// InSamples contains previously unseen audio samples.
+		// InFeedbackSamples contains the delay line output from the other plate.
+		// InDelayModulations contains the fractional delay values per a sample to modulate the first all pass filters delay line.
+		// OutPlateSamples is filled with audio from the various tap points and delay line output.
+		void ProcessAudioFrames(
+			const AlignedFloatBuffer& InSamples,
+			const AlignedFloatBuffer& InFeedbackSamples,
+			const AlignedFloatBuffer& InDelayModulations,
+			FLateReflectionsPlateOutputs& OutPlateSamples);
+
+		void SetDensity(float InDensity);
+		void SetDampening(float InDampening);
+		void SetDecay(float InDecay);
+
+		// The number of samples in an internal buffer. This describes the limit of how many values
+		// are accessible via PeekDelayLineOutput(...)
+		int32 GetNumInternalBufferSamples() const;
+
+		// Retrieve the internal delay line data. The number of samples retrieved is limited to the
+		// number of samples in the internal buffer. That amount can be determined by calling
+		// "GetNumInternalBufferSamples()".
+		void PeekDelayLine(int32 InNum, AlignedFloatBuffer& OutSamples);
+
+	protected:
+		// Current parameter settings of reverb
+		float SampleRate;
+		int32 NumInternalBufferSamples;
+		float Dampening;
+		float Decay;
+		float Density;
+
+		// Sample delay values
+		FLateReflectionsPlateDelays PlateDelays;
+
+		// Sample rate used for hard-coded delay line values
+		static const int32 PresetSampleRate = 29761;
+
+		TUniquePtr<FDynamicDelayAPF> ModulatedAPF;
+		TUniquePtr<FIntegerDelay> DelayA;
+		TUniquePtr<FIntegerDelay> DelayB;
+		TUniquePtr<FIntegerDelay> DelayC;
+		TUniquePtr<FIntegerDelay> DelayD;
+		TUniquePtr<FBufferOnePoleLPF> LPF;
+		TUniquePtr<FLongDelayAPF> APF;
+		TUniquePtr<FIntegerDelay> DelayE;
+		TUniquePtr<FIntegerDelay> DelayF;
+		TUniquePtr<FIntegerDelay> DelayG;
+		TUniquePtr<FIntegerDelay> DelayH;
+		TUniquePtr<FIntegerDelay> DelayI;
+
+		Audio::AlignedFloatBuffer WorkBufferA;
+		Audio::AlignedFloatBuffer WorkBufferB;
+		Audio::AlignedFloatBuffer WorkBufferC;
+	};
+
+	// Settings for controlling the FLateReflections
+	struct AUDIOMIXER_API FLateReflectionsFastSettings
+	{
+		// Milliseconds for the predelay
+		float LateDelayMsec;
+
+		// Initial attenuation of audio after it leaves the predelay
+		float LateGainDB;
+
+		// Frequency bandwidth of audio going into input diffusers. 0.999 is full bandwidth
+		float Bandwidth;
+
+		// Amount of input diffusion (larger value results in more diffusion)
+		float Diffusion;
+
+		// The amount of high-frequency dampening in plate feedback paths
+		float Dampening;
+
+		// The amount of decay in the feedback path. Lower value is larger reverb time.
+		float Decay;
+
+		// The amount of diffusion in decay path. Larger values is a more dense reverb.
+		float Density;
+
+		FLateReflectionsFastSettings();
+
+		bool operator==(const FLateReflectionsFastSettings& Other) const;
+
+		bool operator!=(const FLateReflectionsFastSettings& Other) const;
+	};
+
+
+	// FLateReflections generates the long tail reverb of an input audio signal using a relatively
+	// fast algorithm using all pass filter delay lines. It supports 1 or 2 channel input and always 
+	// produces 2 channel output. 
+	class AUDIOMIXER_API FLateReflectionsFast {
+	public:
+
+		// Limits on settings.
+		static const float MaxLateDelayMsec;
+		static const float MinLateDelayMsec;
+		static const float MaxLateGainDB;
+		static const float MaxBandwidth;
+		static const float MinBandwidth;
+		static const float MaxDampening;
+		static const float MinDampening;
+		static const float MaxDiffusion;
+		static const float MinDiffusion;
+		static const float MaxDecay;
+		static const float MinDecay;
+		static const float MaxDensity;
+		static const float MinDensity;
+
+		static const FLateReflectionsFastSettings DefaultSettings;
+
+		// InMaxNumInternalBufferSamples sets the maximum possible number of samples in an internal
+		// buffer. This only affectsinternal chunking and does not affect the number of samples that
+		// can be sent to ProcessAudio(...), 
+		FLateReflectionsFast(float InSampleRate, int InMaxNumInternalBufferSamples, const FLateReflectionsFastSettings& InSettings=DefaultSettings);
+		
+		~FLateReflectionsFast();
+
+		// Copies the reverb settings internally, clamps the internal copy and applies the clamped settings.
+		void SetSettings(const FLateReflectionsFastSettings& InSettings);
+
+		// Alters settings to ensure they are within acceptable ranges. 
+		static void ClampSettings(FLateReflectionsFastSettings& InOutSettings);
+
+		// Create reverberation in OutSamples based upon InSamples. 
+		// OutSamples is a 2 channel audio buffer no matter the value of InNumChannels.
+		// OutSamples is mixed with InSamples based upon gain.
+		void ProcessAudio(const AlignedFloatBuffer& InSamples, const int32 InNumChannels, AlignedFloatBuffer& OutLeftSamples, AlignedFloatBuffer& OutRightSamples);
+
+	private:
+		void ProcessAudioBuffer(const float* InSampleData, const int32 InNumFrames, const int32 InNumChannels, float* OutLeftSampleData, float* OutRightSampleData);
+
+		void ApplySettings();
+
+		void GeneraterPlateModulations(const int32 InNum, AlignedFloatBuffer& OutLeftDelays, AlignedFloatBuffer& OutRightDelays);
+
+		float SampleRate;
+		float Gain;
+		float ModulationPhase;
+		float ModulationQuadPhase;
+		float ModulationPhaseIncrement;
+		int32 NumInternalBufferSamples;
+
+		// Current parameter settings of reverb
+		FLateReflectionsFastSettings Settings;
+
+		// Sample rate used for hard-coded delay line values
+		static const int32 PresetSampleRate = 29761;
+
+		// A simple pre-delay line to emulate large delays for late-reflections
+		TUniquePtr<FIntegerDelay> PreDelay;
+
+		// Input diffusion
+		TUniquePtr<FBufferOnePoleLPF> InputLPF;
+		TUniquePtr<FLongDelayAPF> APF1;
+		TUniquePtr<FLongDelayAPF> APF2;
+		TUniquePtr<FLongDelayAPF> APF3;
+		TUniquePtr<FLongDelayAPF> APF4;
+
+		// The plate delay data
+		FLateReflectionsPlateDelays LeftPlateDelays;
+		FLateReflectionsPlateDelays RightPlateDelays;
+
+		// The plates.
+		TUniquePtr<FLateReflectionsPlate> LeftPlate;
+		TUniquePtr<FLateReflectionsPlate> RightPlate;
+
+		// The plate outputs
+		FLateReflectionsPlateOutputs LeftPlateOutputs;
+		FLateReflectionsPlateOutputs RightPlateOutputs;
+
+		// Plate output buffering.
+		TUniquePtr<FAlignedBlockBuffer> LeftPlateOutputBuffer;
+		TUniquePtr<FAlignedBlockBuffer> RightPlateOutputBuffer;
+
+		// Buffer for delay modulation of left and right plates.
+		AlignedFloatBuffer LeftDelayModSamples;
+		AlignedFloatBuffer RightDelayModSamples;
+
+		AlignedFloatBuffer WorkBufferA;
+		AlignedFloatBuffer WorkBufferB;
+		AlignedFloatBuffer WorkBufferC;
+
+	};
+
+}
diff --git a/Engine/Source/Runtime/AudioMixer/Public/DSP/LinearInterpFractionalDelay.h b/Engine/Source/Runtime/AudioMixer/Public/DSP/LinearInterpFractionalDelay.h
new file mode 100644
--- /dev/null
+++ b/Engine/Source/Runtime/AudioMixer/Public/DSP/LinearInterpFractionalDelay.h
@@ -0,0 +1,39 @@
+// Copyright 1998-2019 Epic Games, Inc. All Rights Reserved.
+
+#pragma once
+
+#include "CoreMinimal.h"
+#include "DSP/AlignedBlockBuffer.h"
+
+namespace Audio
+{
+	// Fractional delay using linear interpolation.
+	class AUDIOMIXER_API FLinearInterpFractionalDelay
+	{
+	public:
+		
+		// InMaxDelay sets the maximum allowable delay that this object can support. The minimum delay is 0.0
+		// InNumInternalBufferSamples sets the maximum block processing size.
+		FLinearInterpFractionalDelay(int32 InMaxDelay, int32 InMaxNumInternalBufferSamples);
+		~FLinearInterpFractionalDelay();
+
+		// Apply delay to InSamples. Fill OutSamples with data from the delay line at a delay of InDelay.
+		// InDelay must be equal length to InSamples. 
+		void ProcessAudio(const AlignedFloatBuffer& InSamples, const AlignedFloatBuffer& InDelays, AlignedFloatBuffer& OutSamples);
+
+		// Set all values in internal delay line to zero. 
+		void Reset();
+
+	private:
+		void ProcessAudioBlock(const float* InSamples, const float* InDelays, const int32 InNum, float* OutSamples);
+
+		int32 MaxDelay;
+		int32 NumInternalBufferSamples;
+
+		TUniquePtr<FAlignedBlockBuffer> DelayLine;
+		AlignedInt32Buffer IntegerDelayOffsets;
+		int* UpperDelayPos;
+		int* LowerDelayPos;
+
+	};
+}
diff --git a/Engine/Source/Runtime/AudioMixer/Public/DSP/LongDelayAPF.h b/Engine/Source/Runtime/AudioMixer/Public/DSP/LongDelayAPF.h
new file mode 100644
--- /dev/null
+++ b/Engine/Source/Runtime/AudioMixer/Public/DSP/LongDelayAPF.h
@@ -0,0 +1,60 @@
+// Copyright 1998-2019 Epic Games, Inc. All Rights Reserved.
+
+#pragma once
+
+#include "CoreMinimal.h"
+#include "DSP/BufferVectorOperations.h"
+#include "DSP/AlignedBlockBuffer.h"
+
+namespace Audio
+{
+	// All Pass Filter with a long delay. This filter is specifically designed for 
+	// reverb applications where filter delay lines are long.
+	class AUDIOMIXER_API FLongDelayAPF
+	{
+	public:
+		// InG is the filter coefficient used in the long delay all pass filter.
+		// InNumDelaySamples is the delay line length in samples.
+		// InMaxNumInternalBufferSamples is the maximum internal block size used internally. 
+		FLongDelayAPF(float InG, int32 InNumDelaySamples, int32 InMaxNumInternalBufferSamples);
+
+		// Destructor
+		~FLongDelayAPF();
+
+		// Set the APF feedback/feedforward gain coefficient
+		void SetG(float InG) { G = InG; }
+
+		// Process InSamples and place filtered data in OutSamples
+		void ProcessAudio(const AlignedFloatBuffer& InSamples, AlignedFloatBuffer& OutSamples);
+
+		// Process InSamples and place filtered data in OutSamples and delay line samples in OutDelaySamples.
+		// The OutDelaySamples correspond to "w[n] = InSamples[n] + InG * w[n - InDelay]"
+		void ProcessAudio(const AlignedFloatBuffer& InSamples, AlignedFloatBuffer& OutSamples, AlignedFloatBuffer& OutDelaySamples);
+
+		// Sets delay line values to zero.
+		void Reset();
+
+	private:
+		// Performs all pass filtering on an audio block of BlockSize samples are less. 
+		// InNum must be less than or equal to BlockSize.
+		void ProcessAudioBlock(const float* InSamples, const float* InDelaySamples, const int32 InNum, float* OutSamples, float* OutDelaySamples);
+
+		int32 GetNumInternalBufferSamples() const;
+
+		// Feedback/Feedforward gain coefficient
+		float G;
+
+		// Delay in samples
+		int32 NumDelaySamples;
+
+		// Buffer size for internal block processing.
+		int32 NumInternalBufferSamples;
+
+		// Buffer for block operations.
+		AlignedFloatBuffer WorkBuffer;
+
+		// Filter delay line memory.
+		TUniquePtr<FAlignedBlockBuffer> DelayLine;
+	};
+
+}
diff --git a/Engine/Source/Runtime/AudioMixer/Public/DSP/ReverbFast.h b/Engine/Source/Runtime/AudioMixer/Public/DSP/ReverbFast.h
new file mode 100644
--- /dev/null
+++ b/Engine/Source/Runtime/AudioMixer/Public/DSP/ReverbFast.h
@@ -0,0 +1,114 @@
+// Copyright 1998-2019 Epic Games, Inc. All Rights Reserved.
+
+#pragma once
+
+#include "CoreMinimal.h"
+#include "DSP/BufferVectorOperations.h"
+#include "DSP/IntegerDelay.h"
+#include "DSP/LongDelayAPF.h"
+#include "DSP/BufferOnePoleLPF.h"
+#include "DSP/DynamicDelayAPF.h"
+#include "DSP/LateReflectionsFast.h"
+#include "DSP/EarlyReflectionsFast.h"
+
+namespace Audio
+{
+	// Settings for plate reverb
+	struct AUDIOMIXER_API FPlateReverbFastSettings 
+	{
+		// EQuadBehavior describes how reverb is generated when there are 5 or more output channels.
+		enum class EQuadBehavior : uint8
+		{
+			// Only produce reverb in front left and front right.
+			StereoOnly,
+			// Produce reverb in front left and front right. Copy front left to rear left and front right to rear right.
+			QuadMatched,
+			// Produce reverb in front left and front right. Copy front left to rear right and front right to rear left.
+			QuadFlipped
+		};
+
+		FPlateReverbFastSettings();
+
+		bool operator==(const FPlateReverbFastSettings& Other) const;
+
+		bool operator!=(const FPlateReverbFastSettings& Other) const;
+
+		// EarlyReflectionSettings controls the initial perceived echoes from a sound, modeling the first few
+		// orders of reflections from a sound source to the listener's ears. 
+		FEarlyReflectionsFastSettings EarlyReflections;
+		// LateReflectionSettings controls the long tail diffused echo modeling the higher order reflections
+		// from a sound source to the listener's ears. 
+		FLateReflectionsFastSettings LateReflections;
+
+		// Mix amount between dry and wet signals.
+		float Wetness;
+
+		// Set how reverb module generates reverb when there are 5 or more output channels. 
+		EQuadBehavior QuadBehavior;
+	};
+
+	// The Plate Reverb emulates the interactions between a sound, the listener and the space they share. Early reflections
+	// are modeled using a feedback delay network while late reflections are modeled using a plate reverb. This 
+	// class aims to support a flexible and pleasant sounding reverb balanced with computational efficiency. 
+	class AUDIOMIXER_API FPlateReverbFast {
+		public:
+			static const float MaxWetness;
+			static const float MinWetness;
+			static const FPlateReverbFastSettings DefaultSettings;
+			
+			// InMaxInternalBufferSamples sets the maximum number of samples used in internal buffers.
+			FPlateReverbFast(float InSampleRate, int32 InMaxInternalBufferSamples = 512, const FPlateReverbFastSettings& InSettings=DefaultSettings);
+
+			~FPlateReverbFast();
+
+			// Copies, clamps and applies settings.
+			void SetSettings(const FPlateReverbFastSettings& InSettings);
+
+			const FPlateReverbFastSettings& GetSettings() const;
+
+			// Whether or not to enable late reflections
+			void EnableLateReflections(const bool bInEnableLateReflections);
+
+			// Whether or not to enable late reflections
+			void EnableEarlyReflections(const bool bInEnableEarlyReflections);
+
+			// Creates reverberated audio in OutSamples based upon InSamples
+			// InNumChannels can be 1 or 2 channels.
+			// OutSamples must be greater or equal to 2.
+			void ProcessAudio(const AlignedFloatBuffer& InSamples, const int32 InNumChannels, AlignedFloatBuffer& OutSamples, const int32 OutNumChannels);
+
+			// Clamp individual settings to values supported by this class.
+			static void ClampSettings(FPlateReverbFastSettings& InOutSettings);
+
+		private:
+			// Copy input samples to output samples. Remap channels if necessary.
+			void PassThroughAudio(const AlignedFloatBuffer& InSamples, const int32 InNumChannels, AlignedFloatBuffer& OutSamples, const int32 OutNumChannels);
+
+			
+			// Copy reverberated samples to interleaved output samples. Map channels according to internal settings.
+			void InterleaveAndMixOutput(const AlignedFloatBuffer& InFrontLeftSamples, const AlignedFloatBuffer& InFrontRightSamples, AlignedFloatBuffer& OutSamples, const int32 OutNumChannels);
+
+			void ApplySettings();
+
+
+
+
+			float SampleRate;
+
+			FPlateReverbFastSettings Settings;
+			FEarlyReflectionsFast EarlyReflections;
+			FLateReflectionsFast LateReflections;
+
+			AlignedFloatBuffer FrontLeftLateReflectionsSamples;
+			AlignedFloatBuffer FrontRightLateReflectionsSamples;
+			AlignedFloatBuffer FrontLeftEarlyReflectionsSamples;
+			AlignedFloatBuffer FrontRightEarlyReflectionsSamples;
+			AlignedFloatBuffer FrontLeftReverbSamples;
+			AlignedFloatBuffer FrontRightReverbSamples;
+			AlignedFloatBuffer LeftAttenuatedSamples;
+			AlignedFloatBuffer RightAttenuatedSamples;
+
+			bool bEnableEarlyReflections;
+			bool bEnableLateReflections;
+	};
+}
diff --git a/Engine/Source/Runtime/Engine/Classes/Sound/AudioSettings.h b/Engine/Source/Runtime/Engine/Classes/Sound/AudioSettings.h
--- a/Engine/Source/Runtime/Engine/Classes/Sound/AudioSettings.h
+++ b/Engine/Source/Runtime/Engine/Classes/Sound/AudioSettings.h
@@ -135,6 +135,10 @@ class ENGINE_API UAudioSettings : public UDeveloperSettings
 	UPROPERTY(config, EditAnywhere, Category = "Audio", AdvancedDisplay)
 	float DefaultReverbSendLevel;
 
+	/** Enables legacy version of reverb. The legacy reverb runs more slowly, but by most other measures is functionally equivalent. It has a slight perceptual difference. */
+	UPROPERTY(config, EditAnywhere, Category = "Audio", AdvancedDisplay)
+	uint32 bEnableLegacyReverb: 1;
+
 	/** How many streaming sounds can be played at the same time (if more are played they will be sorted by priority) */
 	UPROPERTY(config, EditAnywhere, Category="Audio", meta=(ClampMin=0))
 	int32 MaximumConcurrentStreams;
diff --git a/Engine/Source/Runtime/Engine/Private/AudioSettings.cpp b/Engine/Source/Runtime/Engine/Private/AudioSettings.cpp
--- a/Engine/Source/Runtime/Engine/Private/AudioSettings.cpp
+++ b/Engine/Source/Runtime/Engine/Private/AudioSettings.cpp
@@ -69,6 +69,7 @@ void UAudioSettings::AddDefaultSettings()
 	QualityLevels.Add(DefaultSettings);
 	bAllowVirtualizedSounds = true;
 	DefaultReverbSendLevel = 0.0f;
+	bEnableLegacyReverb = false;
 	VoiPSampleRate = EVoiceSampleRate::Low16000Hz;
 	VoipBufferingDelay = 0.2f;
 	MaxWaveInstances = 100;
