[NO NON-PII COMMIT MESSAGE]

diff --git a/Engine/Source/Runtime/AudioMixer/Private/DSP/EarlyReflectionsFast.cpp b/Engine/Source/Runtime/AudioMixer/Private/DSP/EarlyReflectionsFast.cpp
--- a/Engine/Source/Runtime/AudioMixer/Private/DSP/EarlyReflectionsFast.cpp
+++ b/Engine/Source/Runtime/AudioMixer/Private/DSP/EarlyReflectionsFast.cpp
@@ -2,14 +2,6 @@
 
 using namespace Audio;
 
-// The presets for the late refelections reverb were based on a strange sample rate. This 
-// function helps translate delay values in the paper to delay values for this reverb.
-// https://ccrma.stanford.edu/~dattorro/EffectDesignPart1.pdf
-static int32 GetDelaySamples(float InSampleRate, float InPresetValue) {
-	
-	return (int32)((float)InSampleRate / 29761.f * InPresetValue);
-}
-
 FEarlyReflectionsFastSettings::FEarlyReflectionsFastSettings()
 	: Gain(1.0f)
 	, PreDelayMsec(0.0f)
diff --git a/Engine/Source/Runtime/AudioMixer/Private/DSP/FeedbackDelayNetwork.cpp b/Engine/Source/Runtime/AudioMixer/Private/DSP/FeedbackDelayNetwork.cpp
--- a/Engine/Source/Runtime/AudioMixer/Private/DSP/FeedbackDelayNetwork.cpp
+++ b/Engine/Source/Runtime/AudioMixer/Private/DSP/FeedbackDelayNetwork.cpp
@@ -6,7 +6,7 @@ using namespace Audio;
 // The presets for the late refelections reverb were based on a strange sample rate. This 
 // function helps translate delay values in the paper to delay values for this reverb.
 // https://ccrma.stanford.edu/~dattorro/EffectDesignPart1.pdf
-static int32 GetDelaySamples(float InSampleRate, float InPresetValue) {
+static int32 FeedbackDelayNetworkGetDelaySamples(float InSampleRate, float InPresetValue) {
 	
 	return (int32)((float)InSampleRate / 29761.f * InPresetValue);
 }
@@ -14,20 +14,20 @@ static int32 GetDelaySamples(float InSampleRate, float InPresetValue) {
 FFDNDelaySettings FFDNDelaySettings::DefaultLeftDelays(float InSampleRate)
 {
 	FFDNDelaySettings Settings;
-	Settings.APF0DelayNumSamples = GetDelaySamples(InSampleRate, 5 * 142);
-	Settings.APF1DelayNumSamples = GetDelaySamples(InSampleRate, 5 * 107);
-	Settings.APF2DelayNumSamples = GetDelaySamples(InSampleRate, 5 * 379);
-	Settings.APF3DelayNumSamples = GetDelaySamples(InSampleRate, 5 * 277);
+	Settings.APF0DelayNumSamples = FeedbackDelayNetworkGetDelaySamples(InSampleRate, 5 * 142);
+	Settings.APF1DelayNumSamples = FeedbackDelayNetworkGetDelaySamples(InSampleRate, 5 * 107);
+	Settings.APF2DelayNumSamples = FeedbackDelayNetworkGetDelaySamples(InSampleRate, 5 * 379);
+	Settings.APF3DelayNumSamples = FeedbackDelayNetworkGetDelaySamples(InSampleRate, 5 * 277);
 	return Settings;
 }
 
 FFDNDelaySettings FFDNDelaySettings::DefaultRightDelays(float InSampleRate)
 {
 	FFDNDelaySettings Settings;
-	Settings.APF0DelayNumSamples = GetDelaySamples(InSampleRate, 5 * 279);
-	Settings.APF1DelayNumSamples = GetDelaySamples(InSampleRate, 5 * 137);
-	Settings.APF2DelayNumSamples = GetDelaySamples(InSampleRate, 5 * 213);
-	Settings.APF3DelayNumSamples = GetDelaySamples(InSampleRate, 5 * 327);
+	Settings.APF0DelayNumSamples = FeedbackDelayNetworkGetDelaySamples(InSampleRate, 5 * 279);
+	Settings.APF1DelayNumSamples = FeedbackDelayNetworkGetDelaySamples(InSampleRate, 5 * 137);
+	Settings.APF2DelayNumSamples = FeedbackDelayNetworkGetDelaySamples(InSampleRate, 5 * 213);
+	Settings.APF3DelayNumSamples = FeedbackDelayNetworkGetDelaySamples(InSampleRate, 5 * 327);
 	return Settings;
 }
 
diff --git a/Engine/Source/Runtime/AudioMixer/Private/DSP/LateReflectionsFast.cpp b/Engine/Source/Runtime/AudioMixer/Private/DSP/LateReflectionsFast.cpp
--- a/Engine/Source/Runtime/AudioMixer/Private/DSP/LateReflectionsFast.cpp
+++ b/Engine/Source/Runtime/AudioMixer/Private/DSP/LateReflectionsFast.cpp
@@ -7,7 +7,7 @@ using namespace Audio;
 // The presets for the late refelections reverb were based on a strange sample rate. This 
 // function helps translate delay values in the paper to delay values for this reverb.
 // https://ccrma.stanford.edu/~dattorro/EffectDesignPart1.pdf
-static int32 GetDelaySamples(float InSampleRate, float InPresetValue) {
+static int32 LateReflectionsGetDelaySamples(float InSampleRate, float InPresetValue) {
 	
 	return (int32)((float)InSampleRate / 29761.f * InPresetValue);
 }
@@ -60,21 +60,21 @@ FLateReflectionsPlateDelays FLateReflectionsPlateDelays::DefaultLeftDelays(float
 {
 	FLateReflectionsPlateDelays Delays;
 
-	Delays.NumSamplesModulatedBase = GetDelaySamples(InSampleRate, 908);
-	Delays.NumSamplesModulatedDelta = GetDelaySamples(InSampleRate, 16);
+	Delays.NumSamplesModulatedBase = LateReflectionsGetDelaySamples(InSampleRate, 908);
+	Delays.NumSamplesModulatedDelta = LateReflectionsGetDelaySamples(InSampleRate, 16);
 	// Left Delay 1 should add up to 4217 and have tap outs at 353 (R0), 1190 (L4), and 3627 (R1)
-	Delays.NumSamplesDelayA = GetDelaySamples(InSampleRate, 353);
-	Delays.NumSamplesDelayB = GetDelaySamples(InSampleRate, 837);
-	Delays.NumSamplesDelayC = GetDelaySamples(InSampleRate, 2437);
-	Delays.NumSamplesDelayD = GetDelaySamples(InSampleRate, 590);
-	Delays.NumSamplesAPF = GetDelaySamples(InSampleRate, 2656);
+	Delays.NumSamplesDelayA = LateReflectionsGetDelaySamples(InSampleRate, 353);
+	Delays.NumSamplesDelayB = LateReflectionsGetDelaySamples(InSampleRate, 837);
+	Delays.NumSamplesDelayC = LateReflectionsGetDelaySamples(InSampleRate, 2437);
+	Delays.NumSamplesDelayD = LateReflectionsGetDelaySamples(InSampleRate, 590);
+	Delays.NumSamplesAPF = LateReflectionsGetDelaySamples(InSampleRate, 2656);
 	// APF Delay should have tapouts at 187 (L5) and 1228 (R2)
-	Delays.NumSamplesDelayE = GetDelaySamples(InSampleRate, 187);
-	Delays.NumSamplesDelayF = GetDelaySamples(InSampleRate, 1041);
+	Delays.NumSamplesDelayE = LateReflectionsGetDelaySamples(InSampleRate, 187);
+	Delays.NumSamplesDelayF = LateReflectionsGetDelaySamples(InSampleRate, 1041);
 	// Left Delay 2 should add up to 3136 and have tap outs at 1066 (L6) and 2673 (R3)
-	Delays.NumSamplesDelayG = GetDelaySamples(InSampleRate, 1066);
-	Delays.NumSamplesDelayH = GetDelaySamples(InSampleRate, 1607);
-	Delays.NumSamplesDelayI = GetDelaySamples(InSampleRate, 463);
+	Delays.NumSamplesDelayG = LateReflectionsGetDelaySamples(InSampleRate, 1066);
+	Delays.NumSamplesDelayH = LateReflectionsGetDelaySamples(InSampleRate, 1607);
+	Delays.NumSamplesDelayI = LateReflectionsGetDelaySamples(InSampleRate, 463);
 
 	return Delays;
 }
@@ -83,21 +83,21 @@ FLateReflectionsPlateDelays FLateReflectionsPlateDelays::DefaultRightDelays(floa
 {
 	FLateReflectionsPlateDelays Delays;
 
-	Delays.NumSamplesModulatedBase = GetDelaySamples(InSampleRate, 672);
-	Delays.NumSamplesModulatedDelta = GetDelaySamples(InSampleRate, 16);
+	Delays.NumSamplesModulatedBase = LateReflectionsGetDelaySamples(InSampleRate, 672);
+	Delays.NumSamplesModulatedDelta = LateReflectionsGetDelaySamples(InSampleRate, 16);
 	// Right Delay 1 should add up to 4453 and have tap outs at 266 (L0), 2111 (R4), 2974 (L1)
-	Delays.NumSamplesDelayA = GetDelaySamples(InSampleRate, 266);
-	Delays.NumSamplesDelayB = GetDelaySamples(InSampleRate, 1845);
-	Delays.NumSamplesDelayC = GetDelaySamples(InSampleRate, 863);
-	Delays.NumSamplesDelayD = GetDelaySamples(InSampleRate, 1479);
-	Delays.NumSamplesAPF = GetDelaySamples(InSampleRate, 1800);
+	Delays.NumSamplesDelayA = LateReflectionsGetDelaySamples(InSampleRate, 266);
+	Delays.NumSamplesDelayB = LateReflectionsGetDelaySamples(InSampleRate, 1845);
+	Delays.NumSamplesDelayC = LateReflectionsGetDelaySamples(InSampleRate, 863);
+	Delays.NumSamplesDelayD = LateReflectionsGetDelaySamples(InSampleRate, 1479);
+	Delays.NumSamplesAPF = LateReflectionsGetDelaySamples(InSampleRate, 1800);
 	// APF Delay should have tapouts at 335 (R5) and 1913 (L2)
-	Delays.NumSamplesDelayE = GetDelaySamples(InSampleRate, 335);
-	Delays.NumSamplesDelayF = GetDelaySamples(InSampleRate, 1578);
+	Delays.NumSamplesDelayE = LateReflectionsGetDelaySamples(InSampleRate, 335);
+	Delays.NumSamplesDelayF = LateReflectionsGetDelaySamples(InSampleRate, 1578);
 	// Right Delay 2 should add up to 3720 and have tap outs at 121 (R6) and 1996 (L3)
-	Delays.NumSamplesDelayG = GetDelaySamples(InSampleRate, 121);
-	Delays.NumSamplesDelayH = GetDelaySamples(InSampleRate, 1875);
-	Delays.NumSamplesDelayI = GetDelaySamples(InSampleRate, 1724);
+	Delays.NumSamplesDelayG = LateReflectionsGetDelaySamples(InSampleRate, 121);
+	Delays.NumSamplesDelayH = LateReflectionsGetDelaySamples(InSampleRate, 1875);
+	Delays.NumSamplesDelayI = LateReflectionsGetDelaySamples(InSampleRate, 1724);
 
 	return Delays;
 }
@@ -294,10 +294,10 @@ FLateReflectionsFast::FLateReflectionsFast(float InSampleRate, int32 InMaxNumInt
 	InputLPF = MakeUnique<FBufferOnePoleLPF>(1.0f - Settings.Bandwidth);
 
 	// make the signal decorrelators
-	APF1 = MakeUnique<FLongDelayAPF>(Settings.Diffusion, GetDelaySamples(SampleRate, 142), NumInternalBufferSamples);
-	APF2 = MakeUnique<FLongDelayAPF>(Settings.Diffusion, GetDelaySamples(SampleRate, 107), NumInternalBufferSamples);
-	APF3 = MakeUnique<FLongDelayAPF>(Settings.Diffusion - 0.125f, GetDelaySamples(SampleRate, 379), NumInternalBufferSamples);
-	APF4 = MakeUnique<FLongDelayAPF>(Settings.Diffusion - 0.125f, GetDelaySamples(SampleRate, 277), NumInternalBufferSamples);
+	APF1 = MakeUnique<FLongDelayAPF>(Settings.Diffusion, LateReflectionsGetDelaySamples(SampleRate, 142), NumInternalBufferSamples);
+	APF2 = MakeUnique<FLongDelayAPF>(Settings.Diffusion, LateReflectionsGetDelaySamples(SampleRate, 107), NumInternalBufferSamples);
+	APF3 = MakeUnique<FLongDelayAPF>(Settings.Diffusion - 0.125f, LateReflectionsGetDelaySamples(SampleRate, 379), NumInternalBufferSamples);
+	APF4 = MakeUnique<FLongDelayAPF>(Settings.Diffusion - 0.125f, LateReflectionsGetDelaySamples(SampleRate, 277), NumInternalBufferSamples);
 
 	// Set modulation rate to 1 Hz. 
 	ModulationPhaseIncrement = 2.0 * PI / SampleRate;
diff --git a/Engine/Source/Runtime/AudioMixer/Private/Effects/AudioMixerSubmixEffectReverbFast.cpp b/Engine/Source/Runtime/AudioMixer/Private/Effects/AudioMixerSubmixEffectReverbFast.cpp
--- a/Engine/Source/Runtime/AudioMixer/Private/Effects/AudioMixerSubmixEffectReverbFast.cpp
+++ b/Engine/Source/Runtime/AudioMixer/Private/Effects/AudioMixerSubmixEffectReverbFast.cpp
@@ -8,26 +8,27 @@
 #include "DSP/ReverbFast.h"
 
 
-static int32 DisableSubmixReverbCVar = 0;
+static int32 DisableSubmixReverbCVarFast = 0;
 static FAutoConsoleVariableRef CVarDisableSubmixReverb(
 	TEXT("au.DisableReverbSubmix"),
-	DisableSubmixReverbCVar,
+	DisableSubmixReverbCVarFast,
 	TEXT("Disables the reverb submix.\n")
 	TEXT("0: Not Disabled, 1: Disabled"),
 	ECVF_Default);
 
-static int32 EnableReverbStereoFlipForQuadCVar = 0;
+
+static int32 EnableReverbStereoFlipForQuadCVarFast = 0;
 static FAutoConsoleVariableRef CVarReverbStereoFlipForQuad(
 	TEXT("au.EnableReverbStereoFlipForQuad"),
-	EnableReverbStereoFlipForQuadCVar,
+	EnableReverbStereoFlipForQuadCVarFast,
 	TEXT("Enables doing a stereo flip for quad reverb when in surround.\n")
 	TEXT("0: Not Enabled, 1: Enabled"),
 	ECVF_Default);
 
-static int32 DisableQuadReverbCVar = 0;
-static FAutoConsoleVariableRef CVarDisableQuadReverbCVar(
+static int32 DisableQuadReverbCVarFast = 0;
+static FAutoConsoleVariableRef CVarDisableQuadReverbCVarFast(
 	TEXT("au.DisableQuadReverb"),
-	DisableQuadReverbCVar,
+	DisableQuadReverbCVarFast,
 	TEXT("Disables quad reverb in surround.\n")
 	TEXT("0: Not Disabled, 1: Disabled"),
 	ECVF_Default);
@@ -117,7 +118,7 @@ void FSubmixEffectReverbFast::OnProcessAudio(const FSoundEffectSubmixInputData&
 	LLM_SCOPE(ELLMTag::AudioMixer);
 
 	check(InData.NumChannels == 2);
- 	if (OutData.NumChannels < 2 || DisableSubmixReverbCVar == 1) 
+ 	if (OutData.NumChannels < 2 || DisableSubmixReverbCVarFast == 1) 
 	{
 		// Not supported
 		return;
@@ -180,12 +181,12 @@ void FSubmixEffectReverbFast::UpdateParameters()
 
 	// Check cVars for quad mapping
 	Audio::FPlateReverbFastSettings::EQuadBehavior TargetQuadBehavior;
-	if (DisableQuadReverbCVar)
+	if (DisableQuadReverbCVarFast)
 	{
 		// Disable quad mapping.
  		TargetQuadBehavior = Audio::FPlateReverbFastSettings::EQuadBehavior::StereoOnly;
 	}
-	else if (!DisableQuadReverbCVar && EnableReverbStereoFlipForQuadCVar)
+	else if (!DisableQuadReverbCVarFast && EnableReverbStereoFlipForQuadCVarFast)
 	{
 		// Enable quad flipped mapping
 		TargetQuadBehavior = Audio::FPlateReverbFastSettings::EQuadBehavior::QuadFlipped;
