[NO NON-PII COMMIT MESSAGE]

diff --git a/Engine/Source/Runtime/AudioMixer/Private/Effects/AudioMixerSubmixEffectReverbFast.cpp b/Engine/Source/Runtime/AudioMixer/Private/Effects/AudioMixerSubmixEffectReverbFast.cpp
--- a/Engine/Source/Runtime/AudioMixer/Private/Effects/AudioMixerSubmixEffectReverbFast.cpp
+++ b/Engine/Source/Runtime/AudioMixer/Private/Effects/AudioMixerSubmixEffectReverbFast.cpp
@@ -9,7 +9,7 @@
 
 
 static int32 DisableSubmixReverbCVarFast = 0;
-static FAutoConsoleVariableRef CVarDisableSubmixReverb(
+static FAutoConsoleVariableRef CVarDisableSubmixReverbFast(
 	TEXT("au.DisableReverbSubmix"),
 	DisableSubmixReverbCVarFast,
 	TEXT("Disables the reverb submix.\n")
@@ -18,7 +18,7 @@ static FAutoConsoleVariableRef CVarDisableSubmixReverb(
 
 
 static int32 EnableReverbStereoFlipForQuadCVarFast = 0;
-static FAutoConsoleVariableRef CVarReverbStereoFlipForQuad(
+static FAutoConsoleVariableRef CVarReverbStereoFlipForQuadFast(
 	TEXT("au.EnableReverbStereoFlipForQuad"),
 	EnableReverbStereoFlipForQuadCVarFast,
 	TEXT("Enables doing a stereo flip for quad reverb when in surround.\n")
