remove duplicate asset references reverb related console variables. #jira none #rb none

diff --git a/Engine/Source/Runtime/AudioMixer/Private/Effects/AudioMixerSubmixEffectReverb.cpp b/Engine/Source/Runtime/AudioMixer/Private/Effects/AudioMixerSubmixEffectReverb.cpp
--- a/Engine/Source/Runtime/AudioMixer/Private/Effects/AudioMixerSubmixEffectReverb.cpp
+++ b/Engine/Source/Runtime/AudioMixer/Private/Effects/AudioMixerSubmixEffectReverb.cpp
@@ -8,7 +8,7 @@
 
 static int32 DisableSubmixReverbCVar = 0;
 static FAutoConsoleVariableRef CVarDisableSubmixReverb(
-	TEXT("au.DisableReverbSubmix"),
+	TEXT("au.DisableReverbSubmixLegacy"),
 	DisableSubmixReverbCVar,
 	TEXT("Disables the reverb submix.\n")
 	TEXT("0: Not Disabled, 1: Disabled"),
@@ -16,7 +16,7 @@ static FAutoConsoleVariableRef CVarDisableSubmixReverb(
 
 static int32 EnableReverbStereoFlipForQuadCVar = 0;
 static FAutoConsoleVariableRef CVarReverbStereoFlipForQuad(
-	TEXT("au.EnableReverbStereoFlipForQuad"),
+	TEXT("au.EnableReverbStereoFlipForQuadLegacy"),
 	EnableReverbStereoFlipForQuadCVar,
 	TEXT("Enables doing a stereo flip for quad reverb when in surround.\n")
 	TEXT("0: Not Enabled, 1: Enabled"),
@@ -25,7 +25,7 @@ static FAutoConsoleVariableRef CVarReverbStereoFlipForQuad(
 
 static int32 DisableQuadReverbCVar = 0;
 static FAutoConsoleVariableRef CVarDisableQuadReverbCVar(
-	TEXT("au.DisableQuadReverb"),
+	TEXT("au.DisableQuadReverbLegacy"),
 	DisableQuadReverbCVar,
 	TEXT("Disables quad reverb in surround.\n")
 	TEXT("0: Not Disabled, 1: Disabled"),
