adding copyright to several files and fixing missing Audio.h include #jira none #rb none

diff --git a/Engine/Source/Runtime/AudioMixer/Private/DSP/AlignedBlockBuffer.cpp b/Engine/Source/Runtime/AudioMixer/Private/DSP/AlignedBlockBuffer.cpp
--- a/Engine/Source/Runtime/AudioMixer/Private/DSP/AlignedBlockBuffer.cpp
+++ b/Engine/Source/Runtime/AudioMixer/Private/DSP/AlignedBlockBuffer.cpp
@@ -1,3 +1,5 @@
+// Copyright 1998-2019 Epic Games, Inc. All Rights Reserved.
+
 #include "DSP/AlignedBlockBuffer.h"
 #include "CoreMinimal.h"
 #include "AudioMixer.h"
diff --git a/Engine/Source/Runtime/AudioMixer/Private/DSP/BufferOnePoleLPF.cpp b/Engine/Source/Runtime/AudioMixer/Private/DSP/BufferOnePoleLPF.cpp
--- a/Engine/Source/Runtime/AudioMixer/Private/DSP/BufferOnePoleLPF.cpp
+++ b/Engine/Source/Runtime/AudioMixer/Private/DSP/BufferOnePoleLPF.cpp
@@ -1,3 +1,5 @@
+// Copyright 1998-2019 Epic Games, Inc. All Rights Reserved.
+
 #include "DSP/BufferOnePoleLPF.h"
 #include "DSP/BufferVectorOperations.h"
 
diff --git a/Engine/Source/Runtime/AudioMixer/Private/DSP/EarlyReflectionsFast.cpp b/Engine/Source/Runtime/AudioMixer/Private/DSP/EarlyReflectionsFast.cpp
--- a/Engine/Source/Runtime/AudioMixer/Private/DSP/EarlyReflectionsFast.cpp
+++ b/Engine/Source/Runtime/AudioMixer/Private/DSP/EarlyReflectionsFast.cpp
@@ -1,3 +1,5 @@
+// Copyright 1998-2019 Epic Games, Inc. All Rights Reserved.
+
 #include "DSP/EarlyReflectionsFast.h"
 
 using namespace Audio;
diff --git a/Engine/Source/Runtime/AudioMixer/Private/DSP/FeedbackDelayNetwork.cpp b/Engine/Source/Runtime/AudioMixer/Private/DSP/FeedbackDelayNetwork.cpp
--- a/Engine/Source/Runtime/AudioMixer/Private/DSP/FeedbackDelayNetwork.cpp
+++ b/Engine/Source/Runtime/AudioMixer/Private/DSP/FeedbackDelayNetwork.cpp
@@ -1,3 +1,5 @@
+// Copyright 1998-2019 Epic Games, Inc. All Rights Reserved.
+
 #include "DSP/FeedbackDelayNetwork.h"
 
 
diff --git a/Engine/Source/Runtime/AudioMixer/Private/DSP/LateReflectionsFast.cpp b/Engine/Source/Runtime/AudioMixer/Private/DSP/LateReflectionsFast.cpp
--- a/Engine/Source/Runtime/AudioMixer/Private/DSP/LateReflectionsFast.cpp
+++ b/Engine/Source/Runtime/AudioMixer/Private/DSP/LateReflectionsFast.cpp
@@ -1,3 +1,5 @@
+// Copyright 1998-2019 Epic Games, Inc. All Rights Reserved.
+
 #include "DSP/LateReflectionsFast.h"
 #include "CoreMinimal.h"
 #include "DSP/BufferVectorOperations.h"
diff --git a/Engine/Source/Runtime/AudioMixer/Private/DSP/ReverbFast.cpp b/Engine/Source/Runtime/AudioMixer/Private/DSP/ReverbFast.cpp
--- a/Engine/Source/Runtime/AudioMixer/Private/DSP/ReverbFast.cpp
+++ b/Engine/Source/Runtime/AudioMixer/Private/DSP/ReverbFast.cpp
@@ -1,6 +1,8 @@
+// Copyright 1998-2019 Epic Games, Inc. All Rights Reserved.
 
 #include "DSP/ReverbFast.h"
 #include "DSP/BufferVectorOperations.h"
+#include "Audio.h"
 
 namespace Audio {
 	FPlateReverbFastSettings::FPlateReverbFastSettings()
