New version of single-pole low and high pass filters.
Audio mixer now uses these for the per-source one pole filters.

diff --git a/Engine/Source/Runtime/AudioMixer/Private/AudioMixerSourceManager.cpp b/Engine/Source/Runtime/AudioMixer/Private/AudioMixerSourceManager.cpp
--- a/Engine/Source/Runtime/AudioMixer/Private/AudioMixerSourceManager.cpp
+++ b/Engine/Source/Runtime/AudioMixer/Private/AudioMixerSourceManager.cpp
@@ -437,8 +437,6 @@ namespace Audio
 
 		SourceInfo.DistanceAttenuationSourceStart = -1.0f;
 		SourceInfo.DistanceAttenuationSourceDestination = -1.0f;
-		SourceInfo.LPFCutoffFrequencyParam.Init();
-		SourceInfo.HPFCutoffFrequencyParam.Init();
 
 		SourceInfo.LowPassFilter.Reset();
 		SourceInfo.HighPassFilter.Reset();
@@ -610,9 +608,7 @@ namespace Audio
 
 			// Initialize the number of per-source LPF filters based on input channels
 			SourceInfo.LowPassFilter.Init(MixerDevice->SampleRate, InitParams.NumInputChannels);
-
-			SourceInfo.HighPassFilter.Init(MixerDevice->SampleRate, InitParams.NumInputChannels, 0, nullptr);
-			SourceInfo.HighPassFilter.SetFilterType(EFilter::Type::HighPass);
+			SourceInfo.HighPassFilter.Init(MixerDevice->SampleRate, InitParams.NumInputChannels);
 
 			SourceInfo.SourceEnvelopeFollower = Audio::FEnvelopeFollower(MixerDevice->SampleRate / NumOutputFrames, (float)InitParams.EnvelopeFollowerAttackTime, (float)InitParams.EnvelopeFollowerReleaseTime, Audio::EPeakMode::Peak);
 
@@ -1057,7 +1053,7 @@ namespace Audio
 		{
 			AUDIO_MIXER_CHECK_AUDIO_PLAT_THREAD(MixerDevice);
 
-			SourceInfos[SourceId].LPFCutoffFrequencyParam.SetValue(InLPFFrequency, NumOutputFrames);
+			SourceInfos[SourceId].LowPassFilter.StartFrequencyInterpolation(InLPFFrequency, NumOutputFrames);
 		});
 	}
 
@@ -1071,7 +1067,7 @@ namespace Audio
 		{
 			AUDIO_MIXER_CHECK_AUDIO_PLAT_THREAD(MixerDevice);
 
-			SourceInfos[SourceId].HPFCutoffFrequencyParam.SetValue(InHPFFrequency, NumOutputFrames);
+			SourceInfos[SourceId].HighPassFilter.StartFrequencyInterpolation(InHPFFrequency, NumOutputFrames);
 		});
 	}
 
@@ -2009,57 +2005,20 @@ namespace Audio
 			{
 				float* PostDistanceAttenBufferPtr = SourceInfo.SourceBuffer.GetData();
 
-				// Process the filters after the source effects
-				for (int32 Frame = 0; Frame < NumOutputFrames; ++Frame)
+				// Process the filters after the source effects (Using optimized filters)
+				for (int32 SampleIndex = 0; SampleIndex < NumOutputFrames * SourceInfo.NumInputChannels; SampleIndex += SourceInfo.NumInputChannels)
 				{
-					const float LPFFreq = SourceInfo.LPFCutoffFrequencyParam.Update();
-					const float HPFFreq = SourceInfo.HPFCutoffFrequencyParam.Update();
-
-					SourceInfo.LowPassFilter.SetFrequency(LPFFreq);
-
-					SourceInfo.HighPassFilter.SetFrequency(HPFFreq);
-					SourceInfo.HighPassFilter.Update();
-
-					int32 SampleIndex = SourceInfo.NumInputChannels * Frame;
-
-					// Apply filters, if necessary.
-					if (LPFFreq < (MAX_FILTER_FREQUENCY - KINDA_SMALL_NUMBER))
-					{
-						//If we stopped processing the low pass filter, we need to clear the filter's memory to prevent pops.
-						if (SourceInfo.bIsBypassingLPF)
-						{
-							SourceInfo.LowPassFilter.ClearMemory();
-							SourceInfo.bIsBypassingLPF = false;
-						}
+					SourceInfo.LowPassFilter.ProcessAudioFrame(&PreDistanceAttenBufferPtr[SampleIndex], &PostDistanceAttenBufferPtr[SampleIndex]);
 
-						SourceInfo.LowPassFilter.ProcessAudio(&PreDistanceAttenBufferPtr[SampleIndex], &PostDistanceAttenBufferPtr[SampleIndex]);
-					}
-					else
+					if (!DisableHPFilteringCvar)
 					{
-						SourceInfo.bIsBypassingLPF = true;
-						FMemory::Memcpy(&PostDistanceAttenBufferPtr[SampleIndex], &PreDistanceAttenBufferPtr[SampleIndex], SourceInfo.NumInputChannels * sizeof(float));
-					}
-
-					if (!DisableHPFilteringCvar && HPFFreq > KINDA_SMALL_NUMBER)
-					{
-						//If we stopped processing the low pass filter, we need to clear the filter's memory to prevent pops.
-						if (SourceInfo.bIsBypassingHPF)
-						{
-							SourceInfo.HighPassFilter.Reset();
-							SourceInfo.bIsBypassingHPF = false;
-						}
-
 						SourceInfo.HighPassFilter.ProcessAudioFrame(&PostDistanceAttenBufferPtr[SampleIndex], &PostDistanceAttenBufferPtr[SampleIndex]);
 					}
-					else
-					{
-						SourceInfo.bIsBypassingHPF = true;
-					}
 				}
 
-				// Reset the volume and LPF param interpolations
-				SourceInfo.LPFCutoffFrequencyParam.Reset();
-				SourceInfo.HPFCutoffFrequencyParam.Reset();
+				// We manually reset interpolation to avoid branches in filter code
+				SourceInfo.LowPassFilter.StopFrequencyInterpolation();
+				SourceInfo.HighPassFilter.StopFrequencyInterpolation();
 
 				// Apply distance attenuation
 				ApplyDistanceAttenuation(SourceInfo, NumSamples);
diff --git a/Engine/Source/Runtime/AudioMixer/Private/AudioMixerSourceManager.h b/Engine/Source/Runtime/AudioMixer/Private/AudioMixerSourceManager.h
--- a/Engine/Source/Runtime/AudioMixer/Private/AudioMixerSourceManager.h
+++ b/Engine/Source/Runtime/AudioMixer/Private/AudioMixerSourceManager.h
@@ -9,7 +9,7 @@
 #include "AudioMixerBuffer.h"
 #include "AudioMixerSubmix.h"
 #include "AudioMixerBus.h"
-#include "DSP/OnePole.h"
+#include "DSP/InterpolatedOnePole.h"
 #include "DSP/Filter.h"
 #include "DSP/EnvelopeFollower.h"
 #include "DSP/ParamInterpolator.h"
@@ -445,12 +445,10 @@ namespace Audio
 
 			float DistanceAttenuationSourceStart;
 			float DistanceAttenuationSourceDestination;
-			FParam LPFCutoffFrequencyParam;
-			FParam HPFCutoffFrequencyParam;
 
 			// One-Pole LPFs and HPFs per source
-			Audio::FOnePoleLPFBank LowPassFilter;
-			Audio::FOnePoleFilter HighPassFilter;
+			Audio::FInterpolatedLPF LowPassFilter;
+			Audio::FInterpolatedHPF HighPassFilter;
 
 			// Source effect instances
 			uint32 SourceEffectChainId;
diff --git a/Engine/Source/Runtime/AudioMixer/Private/DSP/InterpolatedOnePole.cpp b/Engine/Source/Runtime/AudioMixer/Private/DSP/InterpolatedOnePole.cpp
new file mode 100644
--- /dev/null
+++ b/Engine/Source/Runtime/AudioMixer/Private/DSP/InterpolatedOnePole.cpp
@@ -0,0 +1,148 @@
+#include "DSP/InterpolatedOnePole.h"
+
+namespace Audio
+{
+
+	// INTERPOLATED ONE-POLE LOW-PASS IMPLEMENTATION
+	FInterpolatedLPF::FInterpolatedLPF()
+	{
+		Z1.Init(0.0f, NumChannels);
+		Reset();
+	}
+
+
+	void FInterpolatedLPF::Init(float InSampleRate, int32 InNumChannels)
+	{
+		SampleRate = InSampleRate;
+		NumChannels = InNumChannels;
+		CutoffFrequency = -1.0f;
+		Z1.Init(0.0f, NumChannels);
+		Reset();
+	}
+
+	void FInterpolatedLPF::StartFrequencyInterpolation(const float InTargetFrequency, const int32 InterpLength)
+	{
+		CurrInterpLength = InterpLength;
+		CurrInterpCounter = 0;
+
+		if (!FMath::IsNearlyEqual(InTargetFrequency, CutoffFrequency))
+		{
+			CutoffFrequency = InTargetFrequency;
+			float NormalizedFreq = FMath::Clamp(0.5f * InTargetFrequency / SampleRate, 0.0f, 1.0f);
+			B1Target = FMath::Exp(-PI * NormalizedFreq);
+			B1Delta = (B1Target - B1Curr) / static_cast<float>(InterpLength);
+		}
+	}
+
+	void FInterpolatedLPF::ProcessAudioFrame(float* RESTRICT InputFrame, float* RESTRICT OutputFrame)
+	{
+		B1Curr += B1Delta; // step forward coefficient
+		++CurrInterpCounter;
+
+		/*
+			[absorbing A0 coefficient]
+			-----------------------------%
+			Yn = Xn*A0 + B1*Z1;                <- old way
+			A0 = (1-B1)
+
+			Yn = Xn*(1-B1) + B1*Z1             <- (1 add, 1 sub, 2 mult)
+			Yn = Xn - B1*Xn + B1*Z1
+			Yn = Xn + B1*Z1 - B1*Xn
+			Yn = Xn + B1*(Z1 - Xn)             <- (1 add, 1 sub, 1 mult)
+		*/
+
+		for (int32 i = 0; i < NumChannels; ++i)
+		{
+			float* Z1Data = Z1.GetData();
+			const float InputSample = InputFrame[i];
+			float Yn = InputSample + B1Curr * (Z1Data[i] - InputSample); // LPF
+			Yn = UnderflowClamp(Yn);
+			Z1Data[i] = Yn;
+			OutputFrame[i] = Yn;
+		}
+	}
+
+	void FInterpolatedLPF::Reset()
+	{
+		B1Curr = 0.0f;
+		B1Delta = 0.0f;
+		B1Target = B1Curr;
+		CurrInterpLength = 0;
+		CurrInterpCounter = 0;
+		ClearMemory();
+	}
+
+	void FInterpolatedLPF::ClearMemory()
+	{
+		Z1.Reset();
+		Z1.AddZeroed(NumChannels);
+	}
+
+
+	// INTERPOLATED ONE-POLE HIGH-PASS IMPLEMENTATION
+	FInterpolatedHPF::FInterpolatedHPF()
+	{
+		Z1.Init(0.0f, NumChannels);
+		Reset();
+	}
+
+	void FInterpolatedHPF::Init(float InSampleRate, int32 InNumChannels)
+	{
+		SampleRate = InSampleRate;
+		NyquistLimit = 0.5f * SampleRate - 1.0f;
+		NumChannels = InNumChannels;
+		CutoffFrequency = -1.0f;
+		Z1.Init(0.0f, NumChannels);
+		Reset();
+	}
+
+	void FInterpolatedHPF::StartFrequencyInterpolation(const float InTargetFrequency, const int32 InterpLength)
+	{
+		CurrInterpLength = InterpLength;
+		CurrInterpCounter = 0;
+
+		if (!FMath::IsNearlyEqual(InTargetFrequency, CutoffFrequency))
+		{
+			CutoffFrequency = FMath::Min(InTargetFrequency, NyquistLimit);
+
+			const float G = GetGCoefficient();
+			A0Target = G / (1.0f + G);
+
+			A0Delta = (A0Target - A0Curr) / static_cast<float>(InterpLength);
+		}
+	}
+
+	void FInterpolatedHPF::ProcessAudioFrame(float* RESTRICT InputFrame, float* RESTRICT OutputFrame)
+	{
+		A0Curr += A0Delta; // step forward coefficient
+		++CurrInterpCounter;
+
+		for (int32 i = 0; i < NumChannels; ++i)
+		{
+			float* Z1Data = Z1.GetData();
+			const float InputSample = InputFrame[i];
+			const float Vn = (InputSample - Z1Data[i]) * A0Curr;
+			const float LPF = Vn + Z1Data[i];
+			Z1Data[i] = Vn + LPF;
+
+			OutputFrame[i] = InputSample - LPF;
+		}
+	}
+
+	void FInterpolatedHPF::Reset()
+	{
+		A0Curr = 0.0f;
+		A0Delta = 0.0f;
+		CurrInterpLength = 0;
+		CurrInterpCounter = 0;
+		ClearMemory();
+	}
+
+	void FInterpolatedHPF::ClearMemory()
+	{
+		Z1.Reset();
+		Z1.AddZeroed(NumChannels);
+	}
+
+
+} // namespace Audio
\ No newline at end of file
diff --git a/Engine/Source/Runtime/AudioMixer/Public/DSP/InterpolatedOnePole.h b/Engine/Source/Runtime/AudioMixer/Public/DSP/InterpolatedOnePole.h
new file mode 100644
--- /dev/null
+++ b/Engine/Source/Runtime/AudioMixer/Public/DSP/InterpolatedOnePole.h
@@ -0,0 +1,137 @@
+// Copyright 1998-2019 Epic Games, Inc. All Rights Reserved.
+
+#pragma once
+#include "CoreMinimal.h"
+#include "DSP/Dsp.h"
+
+
+namespace Audio
+{
+	class AUDIOMIXER_API FInterpolatedLPF
+	{
+	public:
+		// ctor
+		FInterpolatedLPF();
+
+		FORCEINLINE float GetCutoffFrequency() const
+		{
+			return CutoffFrequency;
+		}
+
+		void Init(float InSampleRate, int32 InNumChannels);
+
+		/*
+			StartFrequencyInterpolation():
+				Intended use: call once at the beginning of callback, pass the callback length
+				then call ProcessAudioFrame() exactly "InterpLength" number of times
+				At the end of the callback, call StopFrequencyInterpolation()
+
+				Failing to call StopFrequencyInterpolation() after more than InterpLength calls to ProcessAudioFrame
+				will cause the filter to become unstable and blow up to +/- inf!
+		*/
+		void StartFrequencyInterpolation(const float InTargetFrequency, const int32 InterpLength = 1);
+
+		// interpolates coefficient and processes a sample
+		void ProcessAudioFrame(float* RESTRICT InputFrame, float* RESTRICT OutputFrame);
+
+		/*
+			StopFrequencyInterpolation() needs to be called manually when the interpolation should be done.
+			Snaps the coefficient to the target value.
+			Setting the lerp delta to 0.0f avoids introducing a branch in ProcessAudioFrame()
+		*/
+		FORCEINLINE void StopFrequencyInterpolation()
+		{
+			B1Delta = 0.0f;
+			B1Curr = B1Target;
+		}
+
+		// Resets the sample delay to 0
+		void Reset();
+
+		// Clears memory without reevaluating coefficients.
+		// This function is useful when there is a break between ProcessAudio calls.
+		void ClearMemory();
+
+	private:
+		float CutoffFrequency{ -1.0f };
+		float B1Curr{ 0.0f }; // coefficient
+		float B1Delta{ 0.0f }; // coefficient step size 
+		float B1Target{ 0.0f };
+		TArray<float> Z1; // multi-channel delay terms
+		int32 CurrInterpLength{ 0 };
+		int32 CurrInterpCounter{ 0 };
+		int32 NumInterpSteps;
+		float SampleRate{ 0 };
+		int32 NumChannels{ 1 };
+
+	}; // class Interpolated Low pass filter
+
+
+
+	class AUDIOMIXER_API FInterpolatedHPF
+	{
+	public:
+		// ctor
+		FInterpolatedHPF();
+
+		FORCEINLINE float GetCutoffFrequency() const
+		{
+			return CutoffFrequency;
+		}
+
+		void Init(float InSampleRate, int32 InNumChannels);
+
+		/*
+			StartFrequencyInterpolation():
+				Intended use: call once at the beginning of callback, pass the callback length
+				then call ProcessAudioFrame() exactly "InterpLength" number of times
+				At the end of the callback, call StopFrequencyInterpolation()
+
+				Failing to call StopFrequencyInterpolation() after more than InterpLength calls to ProcessAudioFrame
+				will cause the filter to become unstable and blow up to +/- inf!
+		*/
+		void StartFrequencyInterpolation(const float InTargetFrequency, const int32 InterpLength = 1);
+
+		// interpolates coefficient and processes a sample
+		void ProcessAudioFrame(float* RESTRICT InputFrame, float* RESTRICT OutputFrame);
+
+		/*
+			StopFrequencyInterpolation() needs to be called manually when the interpolation should be done.
+			Snaps the coefficient to the target value.
+			Setting the lerp delta to 0.0f avoids introducing a branch in ProcessAudioFrame()
+		*/
+		FORCEINLINE void StopFrequencyInterpolation()
+		{
+			A0Delta = 0.0f;
+			A0Curr = A0Target;
+		}
+
+		// Resets the sample delay to 0
+		void Reset();
+
+		// Clears memory without reevaluating coefficients.
+		// This function is useful when there is a break between ProcessAudio calls.
+		void ClearMemory();
+		
+	private:
+		FORCEINLINE float GetGCoefficient() const
+		{
+			const float OmegaDigital = 2.0f * PI * CutoffFrequency;
+			const float OmegaAnalog = 2.0f * SampleRate * Audio::FastTan(0.5f * OmegaDigital / SampleRate);
+			const float G = 0.5f * OmegaAnalog / SampleRate;
+			return G;
+		}
+
+		float CutoffFrequency{ -1.0f };
+		float A0Curr{ 0.0f }; // coefficient
+		float A0Delta{ 0.0f }; // coefficient step size
+		float A0Target{ 0.0f };
+		TArray<float> Z1; // multi-channel delay terms
+		int32 CurrInterpLength{ 0 };
+		int32 CurrInterpCounter{ 0 };
+		int32 NumInterpSteps;
+		float SampleRate{ 0 };
+		float NyquistLimit{ 0.0f };
+		int32 NumChannels{ 1 };
+	}; // class Interpolated High pass filter
+} // namespace Audio
