Added missing copyright notice on new file

diff --git a/Engine/Source/Runtime/AudioMixer/Private/DSP/InterpolatedOnePole.cpp b/Engine/Source/Runtime/AudioMixer/Private/DSP/InterpolatedOnePole.cpp
--- a/Engine/Source/Runtime/AudioMixer/Private/DSP/InterpolatedOnePole.cpp
+++ b/Engine/Source/Runtime/AudioMixer/Private/DSP/InterpolatedOnePole.cpp
@@ -1,3 +1,5 @@
+// Copyright 1998-2019 Epic Games, Inc. All Rights Reserved.
+
 #include "DSP/InterpolatedOnePole.h"
 
 namespace Audio
