Added control of Submix and Sorce Bus send levels based on listener distance.

diff --git a/Engine/Source/Runtime/AudioMixer/Private/AudioMixerBus.cpp b/Engine/Source/Runtime/AudioMixer/Private/AudioMixerBus.cpp
--- a/Engine/Source/Runtime/AudioMixer/Private/AudioMixerBus.cpp
+++ b/Engine/Source/Runtime/AudioMixer/Private/AudioMixerBus.cpp
@@ -56,6 +56,18 @@ namespace Audio
 
 	void FMixerBus::AddBusSend(EBusSendType BusSendType, const FBusSend& InBusSend)
 	{
+		// Make sure we don't have duplicates in the bus sends
+		for (FBusSend& BusSend : BusSends[(int32)BusSendType])
+		{
+			// If it's already added, just update the send level
+			if (BusSend.SourceId == InBusSend.SourceId)
+			{
+				BusSend.SendLevel = InBusSend.SendLevel;
+				return;
+			}
+		}
+
+		// It's a new source id so just add it
 		BusSends[(int32)BusSendType].Add(InBusSend);
 	}
 
diff --git a/Engine/Source/Runtime/AudioMixer/Private/AudioMixerSource.cpp b/Engine/Source/Runtime/AudioMixer/Private/AudioMixerSource.cpp
--- a/Engine/Source/Runtime/AudioMixer/Private/AudioMixerSource.cpp
+++ b/Engine/Source/Runtime/AudioMixer/Private/AudioMixerSource.cpp
@@ -16,7 +16,7 @@ FAutoConsoleVariableRef CVarDisableHRTF(
 	TEXT("0: Not Disabled, 1: Disabled"),
 	ECVF_Default);
 
-static int32 UseListenerOverrideForSpreadCVar = 1;
+static int32 UseListenerOverrideForSpreadCVar = 0;
 FAutoConsoleVariableRef CVarUseListenerOverrideForSpread(
 	TEXT("au.UseListenerOverrideForSpread"),
 	UseListenerOverrideForSpreadCVar,
@@ -47,6 +47,7 @@ namespace Audio
 		, bDebugMode(false)
 		, bIsVorbis(false)
 		, bIsStoppingVoicesEnabled(InAudioDevice->IsStoppingVoicesEnabled())
+		, bSendingAudioToBuses(false)
 	{
 	}
 
@@ -167,6 +168,7 @@ namespace Audio
 				// Toggle muting the source if sending only to output bus. 
 				// This can get set even if the source doesn't have bus sends since bus sends can be dynamically enabled.
 				InitParams.bOutputToBusOnly = WaveInstance->bOutputToBusOnly;
+				DynamicBusSendInfos.Reset();
 
 				// If this source is sending its audio to a bus
 				for (int32 BusSendType = 0; BusSendType < (int32)EBusSendType::Count; ++BusSendType)
@@ -180,6 +182,23 @@ namespace Audio
 							BusSend.BusId = SendInfo.SoundSourceBus->GetUniqueID();
 							BusSend.SendLevel = SendInfo.SendLevel;
 							InitParams.BusSends[BusSendType].Add(BusSend);
+
+							FDynamicBusSendInfo DynamicBusSendInfo;
+							DynamicBusSendInfo.SendLevel = SendInfo.SendLevel;
+							DynamicBusSendInfo.BusId = BusSend.BusId;
+							DynamicBusSendInfo.BusSendLevelControlMethod = SendInfo.SourceBusSendLevelControlMethod;
+							DynamicBusSendInfo.BusSendType = (EBusSendType)BusSendType;
+							DynamicBusSendInfo.MinSendLevel = SendInfo.MinSendLevel;
+							DynamicBusSendInfo.MaxSendLevel = SendInfo.MaxSendLevel;
+							DynamicBusSendInfo.MinSendDistance = SendInfo.MinSendDistance;
+							DynamicBusSendInfo.MaxSendDistance = SendInfo.MaxSendDistance;
+							DynamicBusSendInfo.CustomSendLevelCurve = SendInfo.CustomSendLevelCurve;
+
+							// Copy the bus SourceBusSendInfo structs to a local copy so we can update it in the update tick
+							DynamicBusSendInfos.Add(DynamicBusSendInfo);
+
+							// Flag that we're sending audio to buses so we can check for updates to send levels
+							bSendingAudioToBuses = true;
 						}
 					}
 				}
@@ -381,6 +400,8 @@ namespace Audio
 
 		UpdateEffects();
 
+		UpdateSourceBusSends();
+
 		UpdateChannelMaps();
 
 		FSoundSource::DrawDebugInfo();
@@ -900,18 +921,97 @@ namespace Audio
 			}
 		}
 
+		// Update submix send levels
 		for (FSoundSubmixSendInfo& SendInfo : WaveInstance->SoundSubmixSends)
 		{
 			if (SendInfo.SoundSubmix != nullptr)
 			{
 				FMixerSubmixWeakPtr SubmixInstance = MixerDevice->GetSubmixInstance(SendInfo.SoundSubmix);
-				MixerSourceVoice->SetSubmixSendInfo(SubmixInstance, SendInfo.SendLevel);
+				float SendLevel = 0.0f;
+
+				// calculate send level based on distance if that method is enabled
+				if (SendInfo.SendLevelControlMethod == ESendLevelControlMethod::Manual)
+				{
+					SendLevel = FMath::Clamp(SendInfo.SendLevel, 0.0f, 1.0f);
+				}
+				else
+				{
+					// The alpha value is determined identically between manual and custom curve methods
+					const FVector2D SendRadialRange = { SendInfo.MinSendDistance, SendInfo.MaxSendDistance};
+					const FVector2D SendLevelRange = { SendInfo.MinSendLevel, SendInfo.MaxSendLevel };
+					const float Denom = FMath::Max(SendRadialRange.Y - SendRadialRange.X, 1.0f);
+					const float Alpha = FMath::Clamp((WaveInstance->ListenerToSoundDistance - SendRadialRange.X) / Denom, 0.0f, 1.0f);
+
+					if (SendInfo.SendLevelControlMethod == ESendLevelControlMethod::Linear)
+					{
+						SendLevel = FMath::Clamp(FMath::Lerp(SendLevelRange.X, SendLevelRange.Y, Alpha), 0.0f, 1.0f);
+					}
+					else // use curve
+					{
+						SendLevel = FMath::Clamp(SendInfo.CustomSendLevelCurve.GetRichCurveConst()->Eval(Alpha), 0.0f, 1.0f);
+					}
+				}
+
+				// set the level for this send
+				MixerSourceVoice->SetSubmixSendInfo(SubmixInstance, SendLevel);
 
 				// Make sure we flag that we're using this submix sends since these can be dynamically added from BP
 				// If we don't flag this then these channel maps won't be generated for this channel format
 				ChannelMaps[(int32)SendInfo.SoundSubmix->ChannelFormat].bUsed = true;
 			}
 		}
+ 	}
+
+	void FMixerSource::UpdateSourceBusSends()
+	{
+		// 1) loop through all bus sends
+		// 2) check for any bus sends that are set to update non-manually
+		// 3) Cache previous send level and only do update if it's changed in any significant amount
+
+		if (!bSendingAudioToBuses)
+		{
+			return;
+		}
+
+		// If this source is sending its audio to a bus, we need to check if it needs to be updated
+		for (FDynamicBusSendInfo& DynamicBusSendInfo : DynamicBusSendInfos)
+		{
+			float SendLevel = 0.0f;
+
+			if (DynamicBusSendInfo.BusSendLevelControlMethod == ESourceBusSendLevelControlMethod::Manual)
+			{
+				SendLevel = FMath::Clamp(DynamicBusSendInfo.SendLevel, 0.0f, 1.0f);
+			}
+			else
+			{
+				// The alpha value is determined identically between linear and custom curve methods
+				const FVector2D SendRadialRange = { DynamicBusSendInfo.MinSendDistance, DynamicBusSendInfo.MaxSendDistance};
+				const FVector2D SendLevelRange = { DynamicBusSendInfo.MinSendLevel, DynamicBusSendInfo.MaxSendLevel };
+				const float Denom = FMath::Max(SendRadialRange.Y - SendRadialRange.X, 1.0f);
+				const float Alpha = FMath::Clamp((WaveInstance->ListenerToSoundDistance - SendRadialRange.X) / Denom, 0.0f, 1.0f);
+
+				if (DynamicBusSendInfo.BusSendLevelControlMethod == ESourceBusSendLevelControlMethod::Linear)
+				{
+					SendLevel = FMath::Clamp(FMath::Lerp(SendLevelRange.X, SendLevelRange.Y, Alpha), 0.0f, 1.0f);
+				}
+				else // use curve
+				{
+					SendLevel = FMath::Clamp(DynamicBusSendInfo.CustomSendLevelCurve.GetRichCurveConst()->Eval(Alpha), 0.0f, 1.0f);
+				}
+			}
+
+			// If the send level changed, then we need to send an update to the audio render thread
+			if (!FMath::IsNearlyEqual(SendLevel, DynamicBusSendInfo.SendLevel))
+			{
+				DynamicBusSendInfo.SendLevel = SendLevel;
+
+				FMixerBusSend BusSend;
+				BusSend.BusId = DynamicBusSendInfo.BusId;
+				BusSend.SendLevel = SendLevel;
+
+				MixerSourceVoice->SetBusSendInfo(DynamicBusSendInfo.BusSendType, BusSend);
+			}
+		}
 	}
 
 	void FMixerSource::UpdateChannelMaps()
diff --git a/Engine/Source/Runtime/AudioMixer/Private/AudioMixerSource.h b/Engine/Source/Runtime/AudioMixer/Private/AudioMixerSource.h
--- a/Engine/Source/Runtime/AudioMixer/Private/AudioMixerSource.h
+++ b/Engine/Source/Runtime/AudioMixer/Private/AudioMixerSource.h
@@ -77,6 +77,9 @@ namespace Audio
 		/** Updates and source effect on this voice. */
 		void UpdateEffects();
 
+		/** Updates source bus send levels based on game data. */
+		void UpdateSourceBusSends();
+
 		/** Updates the channel map of the sound if its a 3d sound.*/
 		void UpdateChannelMaps();
 
@@ -118,6 +121,32 @@ namespace Audio
 			{}
 		};
 
+		// This holds data copied from FSoundSourceBusSendInfo when a new sound starts playing
+		// so that distance-based level control can be calculated during rendering
+		struct FDynamicBusSendInfo
+		{
+			float SendLevel;
+			uint32 BusId;
+			ESourceBusSendLevelControlMethod BusSendLevelControlMethod;
+			EBusSendType BusSendType;
+			float MinSendLevel;
+			float MaxSendLevel;
+			float MinSendDistance;
+			float MaxSendDistance;
+			FRuntimeFloatCurve CustomSendLevelCurve;
+
+			FDynamicBusSendInfo()
+				: SendLevel(0.0f)
+				, BusId(0)
+				, BusSendLevelControlMethod(ESourceBusSendLevelControlMethod::Manual)
+				, BusSendType(EBusSendType::PreEffect)
+				, MinSendLevel(0.0f)
+				, MaxSendLevel(0.0f)
+				, MinSendDistance(0.0f)
+				, MaxSendDistance(0.0f)
+			{}
+		};
+
 		// Mapping of channel map types to channel maps. Determined by what submixes this source sends its audio to.
 		FChannelMapInfo ChannelMaps[(int32) ESubmixChannelFormat::Count];
 
@@ -136,6 +165,10 @@ namespace Audio
 		FThreadSafeBool bIsPlayingEffectTails;
 		FThreadSafeBool bFreeAsyncTask;
 
+		// Array of copied FSoundSourceBusSendInfo data for all the bus sends this
+		// source may need to live-update during its lifespan
+		TArray<FDynamicBusSendInfo> DynamicBusSendInfos;
+
 		// Whether or not we're currently releasing our resources. Prevents recycling the source until release is finished.
 		FThreadSafeBool bIsReleasing;
 
@@ -145,5 +178,6 @@ namespace Audio
 		uint32 bDebugMode : 1;
 		uint32 bIsVorbis : 1;
 		uint32 bIsStoppingVoicesEnabled : 1;
+		uint32 bSendingAudioToBuses : 1;
 	};
 }
diff --git a/Engine/Source/Runtime/AudioMixer/Private/AudioMixerSourceManager.cpp b/Engine/Source/Runtime/AudioMixer/Private/AudioMixerSourceManager.cpp
--- a/Engine/Source/Runtime/AudioMixer/Private/AudioMixerSourceManager.cpp
+++ b/Engine/Source/Runtime/AudioMixer/Private/AudioMixerSourceManager.cpp
@@ -1112,7 +1112,59 @@ namespace Audio
 					ChannelInfo.bInUse = true;
 				}
 
+			}
 				InSubmixPtr->AddOrSetSourceVoice(MixerSources[SourceId], InSubmixSend.SendLevel);
+		});
+	}
+
+	void FMixerSourceManager::SetBusSendInfo(const int32 SourceId, EBusSendType InBusSendType, FMixerBusSend& InBusSend)
+	{
+		AUDIO_MIXER_CHECK(SourceId < NumTotalSources);
+		AUDIO_MIXER_CHECK(GameThreadInfo.bIsBusy[SourceId]);
+		AUDIO_MIXER_CHECK_GAME_THREAD(MixerDevice);
+
+		AudioMixerThreadCommand([this, SourceId, InBusSendType, InBusSend]()
+		{
+			// Create mapping of source id to bus send level
+			FBusSend BusSend{ SourceId, InBusSend.SendLevel };
+			FSourceInfo& SourceInfo = SourceInfos[SourceId];
+
+			// Retrieve the bus we want to send audio to
+			FMixerBus* Bus = Buses.Find(InBusSend.BusId);
+
+			// If we already have a bus, we update the amount of audio we want to send to it
+			if (Bus)
+			{
+				Bus->AddBusSend(InBusSendType, BusSend);			
+			}
+			else
+			{
+				// If the bus is not registered, make a new entry on the send
+				FMixerBus NewBusData(this, SourceInfo.NumInputChannels, NumOutputFrames);
+
+				// Add a send to it. This will not have a bus instance id (i.e. won't output audio), but 
+				// we register the send anyway in the event that this bus does play, we'll know to send this
+				// source's audio to it.
+				NewBusData.AddBusSend(InBusSendType, BusSend);
+
+				Buses.Add(InBusSend.BusId, NewBusData);
+			}
+
+			// Check to see if we need to create new bus data. If we are not playing a bus with this id, then we
+			// need to create a slot for it such that when a bus does play, it'll start rendering audio from this source
+			bool bExisted = false;
+			for (uint32 BusId : SourceInfo.BusSends[(int32)InBusSendType])
+			{
+				if (BusId == InBusSend.BusId)
+				{
+					bExisted = true;
+					break;
+				}
+			}
+
+			if (!bExisted)
+			{
+				SourceInfo.BusSends[(int32)InBusSendType].Add(InBusSend.BusId);
 			}
 		});
 	}
diff --git a/Engine/Source/Runtime/AudioMixer/Private/AudioMixerSourceManager.h b/Engine/Source/Runtime/AudioMixer/Private/AudioMixerSourceManager.h
--- a/Engine/Source/Runtime/AudioMixer/Private/AudioMixerSourceManager.h
+++ b/Engine/Source/Runtime/AudioMixer/Private/AudioMixerSourceManager.h
@@ -230,6 +230,7 @@ namespace Audio
 		void MixOutputBuffers(const int32 SourceId, const ESubmixChannelFormat InSubmixChannelType, const float SendLevel, AlignedFloatBuffer& OutWetBuffer) const;
 
 		void SetSubmixSendInfo(const int32 SourceId, const FMixerSourceSubmixSend& SubmixSend);
+		void SetBusSendInfo(const int32 SourceId, EBusSendType InBusSendType, FMixerBusSend& BusSend);
 
 		void UpdateDeviceChannelCount(const int32 InNumOutputChannels);
 
diff --git a/Engine/Source/Runtime/AudioMixer/Private/AudioMixerSourceVoice.cpp b/Engine/Source/Runtime/AudioMixer/Private/AudioMixerSourceVoice.cpp
--- a/Engine/Source/Runtime/AudioMixer/Private/AudioMixerSourceVoice.cpp
+++ b/Engine/Source/Runtime/AudioMixer/Private/AudioMixerSourceVoice.cpp
@@ -275,6 +275,13 @@ namespace Audio
 		}
 	}
 
+	void FMixerSourceVoice::SetBusSendInfo(EBusSendType InBusSendType, FMixerBusSend& BusSend)
+	{
+		AUDIO_MIXER_CHECK_GAME_THREAD(MixerDevice);
+
+		SourceManager->SetBusSendInfo(SourceId, InBusSendType, BusSend);
+	}
+
 	void FMixerSourceVoice::OnMixBus(FMixerSourceVoiceBuffer* OutMixerSourceBuffer)
 	{
 		AUDIO_MIXER_CHECK_AUDIO_PLAT_THREAD(MixerDevice);
diff --git a/Engine/Source/Runtime/AudioMixer/Private/AudioMixerSourceVoice.h b/Engine/Source/Runtime/AudioMixer/Private/AudioMixerSourceVoice.h
--- a/Engine/Source/Runtime/AudioMixer/Private/AudioMixerSourceVoice.h
+++ b/Engine/Source/Runtime/AudioMixer/Private/AudioMixerSourceVoice.h
@@ -92,6 +92,9 @@ namespace Audio
 		// Sets the submix send levels
 		void SetSubmixSendInfo(FMixerSubmixWeakPtr Submix, const float SendLevel);
 
+		// Set the source bus send levels
+		void SetBusSendInfo(EBusSendType InBusSendType, FMixerBusSend& BusSend);
+
 		// Called when the source is a bus and needs to mix other sources together to generate output
 		void OnMixBus(FMixerSourceVoiceBuffer* OutMixerSourceBuffer);
 
diff --git a/Engine/Source/Runtime/Engine/Classes/Components/AudioComponent.h b/Engine/Source/Runtime/Engine/Classes/Components/AudioComponent.h
--- a/Engine/Source/Runtime/Engine/Classes/Components/AudioComponent.h
+++ b/Engine/Source/Runtime/Engine/Classes/Components/AudioComponent.h
@@ -425,6 +425,16 @@ public:
 	UFUNCTION(BlueprintCallable, Category = "Audio|Components|Audio")
 	void SetSubmixSend(USoundSubmix* Submix, float SendLevel);
 
+	/** Sets how much audio the sound should send to the given Source Bus (PRE Source Effects).
+		if the Bus Send doesn't already exist, it will be added to the overrides on the active sound */
+	UFUNCTION(BlueprintCallable, Category = "Audio|Components|Audio")
+	void SetSourceBusSendPreEffect(USoundSourceBus* SoundSourceBus, float SourceBusSendLevel);
+
+	/** Sets how much audio the sound should send to the given Source Bus (POST Source Effects).
+		if the Bus Send doesn't already exist, it will be added to the overrides on the active sound */
+	UFUNCTION(BlueprintCallable, Category = "Audio|Components|Audio")
+	void SetSourceBusSendPostEffect(USoundSourceBus* SoundSourceBus, float SourceBusSendLevel);
+
 	/** Sets whether or not the low pass filter is enabled on the audio component. */
 	UFUNCTION(BlueprintCallable, Category = "Audio|Components|Audio")
 	void SetLowPassFilterEnabled(bool InLowPassFilterEnabled);
diff --git a/Engine/Source/Runtime/Engine/Classes/Sound/SoundSourceBusSend.h b/Engine/Source/Runtime/Engine/Classes/Sound/SoundSourceBusSend.h
--- a/Engine/Source/Runtime/Engine/Classes/Sound/SoundSourceBusSend.h
+++ b/Engine/Source/Runtime/Engine/Classes/Sound/SoundSourceBusSend.h
@@ -7,16 +7,68 @@
 
 class USoundSourceBus;
 
+UENUM(BlueprintType)
+enum class ESourceBusSendLevelControlMethod : uint8
+{
+	// A send based on linear interpolation between a distance range and send-level range
+	Linear,
+
+	// A send based on a supplied curve
+	CustomCurve,
+
+	// A manual send level (Uses the specified constant send level value. Useful for 2D sounds.)
+	Manual,
+};
+
 USTRUCT(BlueprintType)
 struct ENGINE_API FSoundSourceBusSendInfo
 {
 	GENERATED_USTRUCT_BODY()
 
-	// The amount of audio to send to the source bus
+	/*
+		Manual: Use Send Level only
+		Linear: Interpolate between Min and Max Send Levels based on listener distance (between Distance Min and Distance Max)
+		Custom Curve: Use the float curve to map Send Level to distance (0.0-1.0 on curve maps to Distance Min - Distance Max)
+	*/
 	UPROPERTY(EditAnywhere, Category = SourceBusSend)
-	float SendLevel;
+	ESourceBusSendLevelControlMethod SourceBusSendLevelControlMethod;
 
-	// The source bus to send the audio to
+	// The Source Bus to send the audio to
 	UPROPERTY(EditAnywhere, Category = SourceBusSend)
 	USoundSourceBus* SoundSourceBus;
+
+	// The amount of audio to send to the source bus
+	UPROPERTY(EditAnywhere, Category = SourceBusSend)
+	float SendLevel;
+
+	// The amount to send to the Source Bus when sound is located at a distance equal to value specified in the min send distance.
+	UPROPERTY(EditAnywhere, BlueprintReadWrite, Category = SourceBusSend)
+	float MinSendLevel;
+
+	// The amount to send to the Source Bus when sound is located at a distance equal to value specified in the max send distance.
+	UPROPERTY(EditAnywhere, BlueprintReadWrite, Category = SourceBusSend)
+	float MaxSendLevel;
+
+	// The distance at which the Min Send Level is sent to the source bus
+	UPROPERTY(EditAnywhere, BlueprintReadWrite, Category = SourceBusSend)
+	float MinSendDistance;
+
+	// The distance at which the Max Send Level is sent to the source bus
+	UPROPERTY(EditAnywhere, BlueprintReadWrite, Category = SourceBusSend)
+	float MaxSendDistance;
+
+	// The custom curve to use for distance-based Source Bus send level.
+	UPROPERTY(EditAnywhere, BlueprintReadWrite, Category = SourceBusSend)
+	FRuntimeFloatCurve CustomSendLevelCurve;
+
+	FSoundSourceBusSendInfo()
+		: SourceBusSendLevelControlMethod(ESourceBusSendLevelControlMethod::Manual)
+		, SoundSourceBus(nullptr)
+		, SendLevel(0.0f)
+		, MinSendLevel(0.0f)
+		, MaxSendLevel(1.0f)
+		, MinSendDistance(100.0f)
+		, MaxSendDistance(1000.0f)
+	{
+	}
 };
diff --git a/Engine/Source/Runtime/Engine/Classes/Sound/SoundSubmix.h b/Engine/Source/Runtime/Engine/Classes/Sound/SoundSubmix.h
--- a/Engine/Source/Runtime/Engine/Classes/Sound/SoundSubmix.h
+++ b/Engine/Source/Runtime/Engine/Classes/Sound/SoundSubmix.h
@@ -7,18 +7,20 @@
 #include "Sound/SampleBuffer.h"
 #include "SoundEffectSubmix.h"
 #include "IAmbisonicsMixer.h"
+#include "Curves/CurveFloat.h"
 #include "SoundSubmix.generated.h"
 
+
 class UEdGraph;
 class USoundEffectSubmixPreset;
 class USoundSubmix;
 class ISubmixBufferListener;
 
-/* Submix channel format. 
+/* Submix channel format.
  Allows submixes to have sources mix to a particular channel configuration for potential effect chain requirements.
- Master submix will always render at the device channel count. All child submixes will be down-mixed (or up-mixed) to 
+ Master submix will always render at the device channel count. All child submixes will be down-mixed (or up-mixed) to
  the device channel count. This feature exists to allow specific submix effects to do their work on multi-channel mixes
- of audio. 
+ of audio.
 */
 UENUM(BlueprintType)
 enum class ESubmixChannelFormat : uint8
@@ -54,28 +56,81 @@ enum class EAudioRecordingExportType : uint8
 	WavFile
 };
 
+UENUM(BlueprintType)
+enum class ESendLevelControlMethod : uint8
+{
+	// A send based on linear interpolation between a distance range and send-level range
+	Linear,
+
+	// A send based on a supplied curve
+	CustomCurve,
+
+	// A manual send level (Uses the specified constant send level value. Useful for 2D sounds.)
+	Manual,
+};
+
+
 // Class used to send audio to submixes from USoundBase
 USTRUCT(BlueprintType)
 struct ENGINE_API FSoundSubmixSendInfo
 {
 	GENERATED_USTRUCT_BODY()
 
-	// The amount of audio to send
+	/* 
+		Manual: Use Send Level only
+		Linear: Interpolate between Min and Max Send Levels based on listener distance (between Distance Min and Distance Max)
+		Custom Curve: Use the float curve to map Send Level to distance (0.0-1.0 on curve maps to Distance Min - Distance Max)
+	*/
 	UPROPERTY(EditAnywhere, Category = SubmixSend)
-	float SendLevel;
-	
+	ESendLevelControlMethod SendLevelControlMethod;
+
 	// The submix to send the audio to
 	UPROPERTY(EditAnywhere, Category = SubmixSend)
 	USoundSubmix* SoundSubmix;
+
+	// The amount of audio to send
+	UPROPERTY(EditAnywhere, Category = SubmixSend)
+	float SendLevel;
+
+	// The amount to send to master when sound is located at a distance equal to value specified in the min send distance.
+	UPROPERTY(EditAnywhere, BlueprintReadWrite, Category = SubmixSend)
+	float MinSendLevel;
+
+	// The amount to send to master when sound is located at a distance equal to value specified in the max send distance.
+	UPROPERTY(EditAnywhere, BlueprintReadWrite, Category = SubmixSend)
+	float MaxSendLevel;
+
+	// The min distance to send to the master
+	UPROPERTY(EditAnywhere, BlueprintReadWrite, Category = SubmixSend)
+	float MinSendDistance;
+
+	// The max distance to send to the master
+	UPROPERTY(EditAnywhere, BlueprintReadWrite, Category = SubmixSend)
+	float MaxSendDistance;
+
+	// The custom reverb send curve to use for distance-based send level.
+	UPROPERTY(EditAnywhere, BlueprintReadWrite, Category = SubmixSend)
+	FRuntimeFloatCurve CustomSendLevelCurve;
+
+	FSoundSubmixSendInfo()
+		: SendLevelControlMethod(ESendLevelControlMethod::Manual)
+		, SoundSubmix(nullptr)
+		, SendLevel(0.0f)
+		, MinSendLevel(0.0f)
+		, MaxSendLevel(1.0f)
+		, MinSendDistance(100.0f)
+		, MaxSendDistance(1000.0f)
+		{
+		}
 };
 
 /**
 * Called when a recorded file has finished writing to disk.
-* 
+*
 */
-DECLARE_DYNAMIC_MULTICAST_DELEGATE_OneParam(FOnSubmixRecordedFileDone,const USoundWave*, ResultingSoundWave);
+DECLARE_DYNAMIC_MULTICAST_DELEGATE_OneParam(FOnSubmixRecordedFileDone, const USoundWave*, ResultingSoundWave);
 
-/** 
+/**
 * Called when a new submix envelope value is generated on the given audio device id (different for multiple PIE). Array is an envelope value for each channel.
 */
 DECLARE_DYNAMIC_MULTICAST_DELEGATE_OneParam(FOnSubmixEnvelope, const TArray<float>&, Envelope);
@@ -96,7 +151,7 @@ public:
 };
 #endif
 
-UCLASS(config=Engine, hidecategories=Object, editinlinenew, BlueprintType)
+UCLASS(config = Engine, hidecategories = Object, editinlinenew, BlueprintType)
 class ENGINE_API USoundSubmix : public UObject
 {
 	GENERATED_UCLASS_BODY()
@@ -182,7 +237,7 @@ protected:
 	virtual void PostEditChangeProperty(struct FPropertyChangedEvent& PropertyChangedEvent) override;
 #endif
 	//~ End UObject Interface.
-	
+
 	// State handling for bouncing output.
 	TUniquePtr<Audio::FAudioRecordingData> RecordingData;
 
diff --git a/Engine/Source/Runtime/Engine/Private/Components/AudioComponent.cpp b/Engine/Source/Runtime/Engine/Private/Components/AudioComponent.cpp
--- a/Engine/Source/Runtime/Engine/Private/Components/AudioComponent.cpp
+++ b/Engine/Source/Runtime/Engine/Private/Components/AudioComponent.cpp
@@ -1145,6 +1145,48 @@ void UAudioComponent::SetSubmixSend(USoundSubmix* Submix, float SendLevel)
 	}
 }
 
+// BP function to set source bus sends (pre effect)
+void UAudioComponent::SetSourceBusSendPreEffect(USoundSourceBus* SoundSourceBus, float SourceBusSendLevel)
+{
+	if (FAudioDevice* AudioDevice = GetAudioDevice())
+	{
+		const uint64 MyAudioComponentID = AudioComponentID;
+		FAudioThread::RunCommandOnAudioThread([AudioDevice, MyAudioComponentID, SoundSourceBus, SourceBusSendLevel]()
+		{
+			FActiveSound* ActiveSound = AudioDevice->FindActiveSound(MyAudioComponentID);
+			if (ActiveSound)
+			{
+				FSoundSourceBusSendInfo SourceBusSendInfo;
+				SourceBusSendInfo.SoundSourceBus = SoundSourceBus;
+				SourceBusSendInfo.SendLevel = SourceBusSendLevel;
+
+				ActiveSound->SetSourceBusSend(EBusSendType::PreEffect, SourceBusSendInfo);
+			}
+		});
+	}
+}
+
+// BP function to set source bus sends (post effect)
+void UAudioComponent::SetSourceBusSendPostEffect(USoundSourceBus * SoundSourceBus, float SourceBusSendLevel)
+{
+	if (FAudioDevice* AudioDevice = GetAudioDevice())
+	{
+		const uint64 MyAudioComponentID = AudioComponentID;
+		FAudioThread::RunCommandOnAudioThread([AudioDevice, MyAudioComponentID, SoundSourceBus, SourceBusSendLevel]()
+		{
+			FActiveSound* ActiveSound = AudioDevice->FindActiveSound(MyAudioComponentID);
+			if (ActiveSound)
+			{
+				FSoundSourceBusSendInfo SourceBusSendInfo;
+				SourceBusSendInfo.SoundSourceBus = SoundSourceBus;
+				SourceBusSendInfo.SendLevel = SourceBusSendLevel;
+
+				ActiveSound->SetSourceBusSend(EBusSendType::PostEffect, SourceBusSendInfo);
+			}
+		});
+	}
+}
+
 void UAudioComponent::SetLowPassFilterEnabled(bool InLowPassFilterEnabled)
 {
 	if (FAudioDevice* AudioDevice = GetAudioDevice())
