Scale fade over smaller more perceptible dB range for concurrency ducking

diff --git a/Engine/Source/Runtime/AudioMixer/Public/DSP/Dsp.h b/Engine/Source/Runtime/AudioMixer/Public/DSP/Dsp.h
--- a/Engine/Source/Runtime/AudioMixer/Public/DSP/Dsp.h
+++ b/Engine/Source/Runtime/AudioMixer/Public/DSP/Dsp.h
@@ -7,7 +7,7 @@
 
 // Macros which can be enabled to cause DSP sample checking
 #if 0
-#define CHECK_SAMPLE(VALUE) 
+#define CHECK_SAMPLE(VALUE)
 #define CHECK_SAMPLE2(VALUE)
 #else
 #define CHECK_SAMPLE(VALUE)  Audio::CheckSample(VALUE)
@@ -16,10 +16,10 @@
 
 namespace Audio
 {
-	// Utility to check for sample clipping. Put breakpoint in conditional to find 
+	// Utility to check for sample clipping. Put breakpoint in conditional to find
 	// DSP code that's not behaving correctly
 	static void CheckSample(float InSample, float Threshold = 0.001f)
-	{	
+	{
 		if (InSample > Threshold || InSample < -Threshold)
 		{
 			UE_LOG(LogTemp, Log, TEXT("SampleValue Was %.2f"), InSample);
@@ -37,9 +37,9 @@ namespace Audio
 	}
 
 	// Function converts linear scale volume to decibels
-	static FORCEINLINE float ConvertToDecibels(const float InLinear)
+	static FORCEINLINE float ConvertToDecibels(const float InLinear, const float InFloor = SMALL_NUMBER)
 	{
-		return 20.0f * FMath::LogX(10.0f, FMath::Max(InLinear, SMALL_NUMBER));
+		return 20.0f * FMath::LogX(10.0f, FMath::Max(InLinear, InFloor));
 	}
 
 	// Function converts decibel to linear scale
@@ -73,7 +73,7 @@ namespace Audio
 		return X2;
 	}
 
-	// Sine approximation using Bhaskara I technique discovered in 7th century. 
+	// Sine approximation using Bhaskara I technique discovered in 7th century.
 	// https://en.wikipedia.org/wiki/Bh%C4%81skara_I
 	static FORCEINLINE float FastSin3(const float X)
 	{
@@ -187,7 +187,7 @@ namespace Audio
 		return OutQ;
 	}
 
-	// Polynomial interpolation using lagrange polynomials. 
+	// Polynomial interpolation using lagrange polynomials.
 	// https://en.wikipedia.org/wiki/Lagrange_polynomial
 	static FORCEINLINE float LagrangianInterpolation(const TArray<FVector2D> Points, const float Alpha)
 	{
@@ -288,9 +288,9 @@ namespace Audio
 		// Percentage to move toward target value from current value each tick
 		float EaseFactor;
 	};
-	
+
 	// Simple easing function used to help interpolate params
-	class FLinearEase 
+	class FLinearEase
 	{
 	public:
 		FLinearEase()
@@ -339,7 +339,7 @@ namespace Audio
 			++CurrentTick;
 			return CurrentValue;
 		}
-		 
+
 		// Updates the target value without changing the duration or tick data.
 		// Sets the state as if the new value was the target value all along
 		void SetValueInterrupt(const float InValue)
@@ -372,7 +372,7 @@ namespace Audio
 
 			if (DurationTicks == 0)
 			{
-				CurrentValue = InValue;			
+				CurrentValue = InValue;
 			}
 			else
 			{
@@ -473,7 +473,7 @@ namespace Audio
 			{
 				// Copy to the end of Internal Buffer and then the beginning once we wrap around.
 				FMemory::Memcpy(DestBuffer + WriteIndex, InBuffer, DestBufferRemainder * sizeof(SampleType));
-					
+
 				const uint32 InBufferRemainder = NumSamplesToCopy - DestBufferRemainder;
 				FMemory::Memcpy(DestBuffer, InBuffer + DestBufferRemainder, InBufferRemainder * sizeof(SampleType));
 				WriteCounter.Set(InBufferRemainder);
@@ -506,7 +506,7 @@ namespace Audio
 			{
 				// Copy from the end of Internal Buffer and then the beginning once we wrap around.
 				FMemory::Memcpy(OutBuffer, SrcBuffer + ReadIndex, SrcBufferRemainder * sizeof(SampleType));
-				
+
 				const uint32 OutBufferRemainder = NumSamplesToCopy - SrcBufferRemainder;
 				FMemory::Memcpy(OutBuffer + SrcBufferRemainder, SrcBuffer, OutBufferRemainder * sizeof(SampleType));
 
diff --git a/Engine/Source/Runtime/Engine/Private/SoundConcurrency.cpp b/Engine/Source/Runtime/Engine/Private/SoundConcurrency.cpp
--- a/Engine/Source/Runtime/Engine/Private/SoundConcurrency.cpp
+++ b/Engine/Source/Runtime/Engine/Private/SoundConcurrency.cpp
@@ -115,7 +115,7 @@ void FConcurrencySoundData::SetTarget(float InTargetVolume, float InLerpTime)
 	LerpTime = FMath::Max(InLerpTime, 0.0f);
 	Elapsed = 0.0f;
 
-	DbTargetVolume = Audio::ConvertToDecibels(InTargetVolume);
+	DbTargetVolume = Audio::ConvertToDecibels(InTargetVolume, KINDA_SMALL_NUMBER);
 }
 
 /************************************************************************/
@@ -812,7 +812,7 @@ void FSoundConcurrencyManager::RemoveActiveSound(FActiveSound& ActiveSound)
 					}
 				}
 			}
-			
+
 			delete ConcurrencyGroup;
 		}
 	}
