Changed default mapping of speaker angles  to SideLeft=250 and SideRight=100 to be better inline with Dolby spec when these angles are used for both 7.1 and 5.1 configurations.
Set UseListenerOverrideForSpreadCVar = 0 by default so that panning logic uses the correct (non attenuation override) distance.

diff --git a/Engine/Config/BaseEngine.ini b/Engine/Config/BaseEngine.ini
--- a/Engine/Config/BaseEngine.ini
+++ b/Engine/Config/BaseEngine.ini
@@ -821,8 +821,8 @@ BackRight=150
 FrontLeftOfCenter=15
 FrontRightOfCenter=345
 BackCenter=180
-SideLeft=270
-SideRight=90
+SideLeft=250
+SideRight=100
 
 ; These are the default channel orderings to use when using pro-audio interfaces that don't report channel order
 [AudioDefaultChannelOrder]
diff --git a/Engine/Source/Runtime/AudioMixer/Private/AudioMixerChannelMaps.cpp b/Engine/Source/Runtime/AudioMixer/Private/AudioMixerChannelMaps.cpp
--- a/Engine/Source/Runtime/AudioMixer/Private/AudioMixerChannelMaps.cpp
+++ b/Engine/Source/Runtime/AudioMixer/Private/AudioMixerChannelMaps.cpp
@@ -394,8 +394,8 @@ namespace Audio
 		DefaultChannelAzimuthPosition[EAudioMixerChannel::FrontLeftOfCenter] = { EAudioMixerChannel::FrontLeftOfCenter, 15 };
 		DefaultChannelAzimuthPosition[EAudioMixerChannel::FrontRightOfCenter] = { EAudioMixerChannel::FrontRightOfCenter, 345 };
 		DefaultChannelAzimuthPosition[EAudioMixerChannel::BackCenter] = { EAudioMixerChannel::BackCenter, 180 };
-		DefaultChannelAzimuthPosition[EAudioMixerChannel::SideLeft] = { EAudioMixerChannel::SideLeft, 270 };
-		DefaultChannelAzimuthPosition[EAudioMixerChannel::SideRight] = { EAudioMixerChannel::SideRight, 90 };
+		DefaultChannelAzimuthPosition[EAudioMixerChannel::SideLeft] = { EAudioMixerChannel::SideLeft, 250 };
+		DefaultChannelAzimuthPosition[EAudioMixerChannel::SideRight] = { EAudioMixerChannel::SideRight, 110 };
 
 		// Check any engine ini overrides for these default positions
 		if (NumChannels != 2)
