#rnx

diff --git a/Engine/Source/Runtime/AudioMixer/Private/AudioMixerDevice.cpp b/Engine/Source/Runtime/AudioMixer/Private/AudioMixerDevice.cpp
--- a/Engine/Source/Runtime/AudioMixer/Private/AudioMixerDevice.cpp
+++ b/Engine/Source/Runtime/AudioMixer/Private/AudioMixerDevice.cpp
@@ -919,11 +919,11 @@ namespace Audio
 		}
 	}
 
-	void FMixerDevice::FlushAudioRenderingCommands()
+	void FMixerDevice::FlushAudioRenderingCommands(bool bPumpSynchronously)
 	{
 		if (IsInitialized() && (FPlatformProcess::SupportsMultithreading() && !AudioMixerPlatform->IsNonRealtime()))
 		{
-			SourceManager.FlushCommandQueue();
+			SourceManager.FlushCommandQueue(bPumpSynchronously);
 		}
 		else if (AudioMixerPlatform->IsNonRealtime())
 		{
diff --git a/Engine/Source/Runtime/AudioMixer/Private/AudioMixerSourceManager.cpp b/Engine/Source/Runtime/AudioMixer/Private/AudioMixerSourceManager.cpp
--- a/Engine/Source/Runtime/AudioMixer/Private/AudioMixerSourceManager.cpp
+++ b/Engine/Source/Runtime/AudioMixer/Private/AudioMixerSourceManager.cpp
@@ -298,6 +298,9 @@ namespace Audio
 
 				// Make sure we've actually emptied the command queue from the render thread before writing to it
 				check(CommandBuffers[NextIndex].SourceCommandQueue.Num() == 0);
+
+				// Here we ensure that we block for any pending calls to AudioMixerThreadCommand.
+				FScopeLock ScopeLock(&CommandBufferIndexCriticalSection);
 				RenderThreadCommandBufferIndex.Set(CurrentGameIndex);
 
 				CommandsProcessedEvent->Reset();
@@ -2419,11 +2422,14 @@ namespace Audio
 		}
 	}
 
+
 	void FMixerSourceManager::AudioMixerThreadCommand(TFunction<void()> InFunction)
 	{
+		// Here, we make sure that we don't flip our command double buffer while we are executing this function.
+		FScopeLock ScopeLock(&CommandBufferIndexCriticalSection);
 		AUDIO_MIXER_CHECK_GAME_THREAD(MixerDevice);
 
-		// Add the function to the command queue
+		// Add the function to the command queue:
 		int32 AudioThreadCommandIndex = !RenderThreadCommandBufferIndex.GetValue();
 		CommandBuffers[AudioThreadCommandIndex].SourceCommandQueue.Add(MoveTemp(InFunction));
 		NumCommands.Increment();
diff --git a/Engine/Source/Runtime/AudioMixer/Private/AudioMixerSourceManager.h b/Engine/Source/Runtime/AudioMixer/Private/AudioMixerSourceManager.h
--- a/Engine/Source/Runtime/AudioMixer/Private/AudioMixerSourceManager.h
+++ b/Engine/Source/Runtime/AudioMixer/Private/AudioMixerSourceManager.h
@@ -317,6 +317,7 @@ namespace Audio
 		FThreadSafeCounter RenderThreadCommandBufferIndex;
 
 		FEvent* CommandsProcessedEvent;
+		FCriticalSection CommandBufferIndexCriticalSection;
 
 		TArray<int32> DebugSoloSources;
 
diff --git a/Engine/Source/Runtime/AudioMixer/Public/AudioMixerDevice.h b/Engine/Source/Runtime/AudioMixer/Public/AudioMixerDevice.h
--- a/Engine/Source/Runtime/AudioMixer/Public/AudioMixerDevice.h
+++ b/Engine/Source/Runtime/AudioMixer/Public/AudioMixerDevice.h
@@ -131,7 +131,7 @@ namespace Audio
 		virtual void RegisterSubmixBufferListener(ISubmixBufferListener* InSubmixBufferListener, USoundSubmix* InSubmix = nullptr) override;
 		virtual void UnregisterSubmixBufferListener(ISubmixBufferListener* InSubmixBufferListener, USoundSubmix* InSubmix = nullptr) override;
 
-		virtual void FlushAudioRenderingCommands() override;
+		virtual void FlushAudioRenderingCommands(bool bPumpSynchronously = false) override;
 		//~ End FAudioDevice
 
 		//~ Begin IAudioMixer
diff --git a/Engine/Source/Runtime/Engine/Public/AudioDevice.h b/Engine/Source/Runtime/Engine/Public/AudioDevice.h
--- a/Engine/Source/Runtime/Engine/Public/AudioDevice.h
+++ b/Engine/Source/Runtime/Engine/Public/AudioDevice.h
@@ -551,8 +551,11 @@ public:
 	 */
 	void Flush(UWorld* WorldToFlush, bool bClearActivatedReverb = true);
 
-	/** Allows audio rendering command queue to flush during audio device flush. */
-	virtual void FlushAudioRenderingCommands() {}
+	/** 
+	 * Allows audio rendering command queue to flush during audio device flush. 
+	 * @param bPumpSynchronously must be called in situations where the audio render thread is not being called.
+	 */
+	virtual void FlushAudioRenderingCommands(bool bPumpSynchronously = false) {}
 
 	/**
 	 * Stop any playing sounds that are using a particular SoundWave
