Fixing assert when unpausing the game.

diff --git a/Engine/Source/Runtime/AudioMixer/Private/AudioMixerSource.cpp b/Engine/Source/Runtime/AudioMixer/Private/AudioMixerSource.cpp
--- a/Engine/Source/Runtime/AudioMixer/Private/AudioMixerSource.cpp
+++ b/Engine/Source/Runtime/AudioMixer/Private/AudioMixerSource.cpp
@@ -553,6 +553,12 @@ namespace Audio
 			return;
 		}
 
+		if (bIsStopping)
+		{
+			UE_LOG(LogAudioMixer, Warning, TEXT("Restarting a source which was stopping. Stopping now."));
+			return;
+		}
+
 		// It's possible if Pause and Play are called while a sound is async initializing. In this case
 		// we'll just not actually play the source here. Instead we'll call play when the sound finishes loading.
 		if (MixerSourceVoice && InitializationState == EMixerSourceInitializationState::Initialized)
