Utilizing audio legacy reverb setting to control which reverb submix effect is used.

diff --git a/Engine/Source/Runtime/AudioMixer/Private/AudioMixerDevice.cpp b/Engine/Source/Runtime/AudioMixer/Private/AudioMixerDevice.cpp
--- a/Engine/Source/Runtime/AudioMixer/Private/AudioMixerDevice.cpp
+++ b/Engine/Source/Runtime/AudioMixer/Private/AudioMixerDevice.cpp
@@ -8,6 +8,7 @@
 #include "UObject/UObjectHash.h"
 #include "AudioMixerEffectsManager.h"
 #include "SubmixEffects/AudioMixerSubmixEffectReverb.h"
+#include "SubmixEffects/AudioMixerSubmixEffectReverbFast.h"
 #include "SubmixEffects/AudioMixerSubmixEffectEQ.h"
 #include "SubmixEffects/AudioMixerSubmixEffectDynamicsProcessor.h"
 #include "DSP/Noise.h"
@@ -666,9 +667,18 @@ namespace Audio
 			else if (MasterSubmixInstances[EMasterSubmixType::Reverb].IsValid())
 			{
 				// Setup the master reverb only if we don't have a reverb plugin
+				USoundEffectSubmixPreset* ReverbPreset = nullptr;
+				USoundSubmix* MasterReverbSoundSubmix = FMixerDevice::MasterSubmixes[EMasterSubmixType::Reverb];
 
-				USoundSubmix* MasterReverbSubix = FMixerDevice::MasterSubmixes[EMasterSubmixType::Reverb];
-				USubmixEffectReverbPreset* ReverbPreset = NewObject<USubmixEffectReverbPreset>(MasterReverbSubix, TEXT("Master Reverb Effect Preset"));
+				if (GetDefault<UAudioSettings>()->bEnableLegacyReverb)
+				{
+					ReverbPreset = NewObject<USubmixEffectReverbPreset>(MasterReverbSoundSubmix, TEXT("Master Reverb Effect Preset"));
+				}
+				else
+				{
+					ReverbPreset = NewObject<USubmixEffectReverbFastPreset>(MasterReverbSoundSubmix, TEXT("Master Reverb Effect Preset"));
+				}
+				
 				ReverbPreset->AddToRoot();
 
 				FSoundEffectSubmix* ReverbEffectSubmix = static_cast<FSoundEffectSubmix*>(ReverbPreset->CreateNewEffect());
diff --git a/Engine/Source/Runtime/AudioMixer/Private/AudioMixerEffectsManager.cpp b/Engine/Source/Runtime/AudioMixer/Private/AudioMixerEffectsManager.cpp
--- a/Engine/Source/Runtime/AudioMixer/Private/AudioMixerEffectsManager.cpp
+++ b/Engine/Source/Runtime/AudioMixer/Private/AudioMixerEffectsManager.cpp
@@ -4,6 +4,7 @@
 #include "AudioMixerDevice.h"
 #include "AudioMixerEffectsManager.h"
 #include "SubmixEffects/AudioMixerSubmixEffectReverb.h"
+#include "SubmixEffects/AudioMixerSubmixEffectReverbFast.h"
 #include "SubmixEffects/AudioMixerSubmixEffectEQ.h"
 
 namespace Audio
@@ -113,8 +114,17 @@ namespace Audio
 			FSoundEffectSubmix* SoundEffectSubmix = MasterReverbSubmixPtr->GetSubmixEffect(0);
 			if (SoundEffectSubmix)
 			{
-				FSubmixEffectReverb* SoundEffectReverb = static_cast<FSubmixEffectReverb*>(SoundEffectSubmix);
-				SoundEffectReverb->SetEffectParameters(ReverbEffectParameters);
+				// Choose correct reverb based upon ini settings.
+				if (GetDefault<UAudioSettings>()->bEnableLegacyReverb)
+				{
+					FSubmixEffectReverb* SoundEffectReverb = static_cast<FSubmixEffectReverb*>(SoundEffectSubmix);
+					SoundEffectReverb->SetEffectParameters(ReverbEffectParameters);
+				}
+				else
+				{
+					FSubmixEffectReverbFast* SoundEffectReverb = static_cast<FSubmixEffectReverbFast*>(SoundEffectSubmix);
+					SoundEffectReverb->SetEffectParameters(ReverbEffectParameters);
+				}
 				PrintReverbSettings(ReverbEffectParameters);
 			}
 		}
