Bypassing assert when timing out a flush on shutdown

diff --git a/Engine/Source/Runtime/AudioMixer/Private/AudioMixerSourceManager.cpp b/Engine/Source/Runtime/AudioMixer/Private/AudioMixerSourceManager.cpp
--- a/Engine/Source/Runtime/AudioMixer/Private/AudioMixerSourceManager.cpp
+++ b/Engine/Source/Runtime/AudioMixer/Private/AudioMixerSourceManager.cpp
@@ -270,7 +270,7 @@ namespace Audio
 		bPumpQueue = false;
 	}
 
-	void FMixerSourceManager::Update()
+	void FMixerSourceManager::Update(bool bTimedOut)
 	{
 		AUDIO_MIXER_CHECK_GAME_THREAD(MixerDevice);
 
@@ -297,7 +297,12 @@ namespace Audio
 				const int32 NextIndex = (CurrentGameIndex + 1) & 1;
 
 				// Make sure we've actually emptied the command queue from the render thread before writing to it
-				check(CommandBuffers[NextIndex].SourceCommandQueue.Num() == 0);
+#if !UE_BUILD_SHIPPING
+				if (bTimedOut)
+				{
+					check(CommandBuffers[NextIndex].SourceCommandQueue.Num() == 0);
+				}
+#endif
 
 				// Here we ensure that we block for any pending calls to AudioMixerThreadCommand.
 				FScopeLock ScopeLock(&CommandBufferIndexCriticalSection);
@@ -2484,9 +2489,11 @@ namespace Audio
 		}
 
 		// Make sure current current executing 
+		bool bTimedOut = false;
 		if (!CommandsProcessedEvent->Wait(1000))
 		{
 			CommandsProcessedEvent->Trigger();
+			bTimedOut = true;
 			UE_LOG(LogAudioMixer, Warning, TEXT("Timed out waiting to flush the source manager command queue (1)."));
 		}
 		else
@@ -2495,7 +2502,7 @@ namespace Audio
 		}
 
 		// Call update to trigger a final pump of commands
-		Update();
+		Update(bTimedOut);
 
 		if (bPumpInCommand)
 		{
diff --git a/Engine/Source/Runtime/AudioMixer/Private/AudioMixerSourceManager.h b/Engine/Source/Runtime/AudioMixer/Private/AudioMixerSourceManager.h
--- a/Engine/Source/Runtime/AudioMixer/Private/AudioMixerSourceManager.h
+++ b/Engine/Source/Runtime/AudioMixer/Private/AudioMixerSourceManager.h
@@ -198,7 +198,7 @@ namespace Audio
 		~FMixerSourceManager();
 
 		void Init(const FSourceManagerInitParams& InitParams);
-		void Update();
+		void Update(bool bTimedOut = false);
 
 		bool GetFreeSourceId(int32& OutSourceId);
 		int32 GetNumActiveSources() const;
