Fix for non-external-send spatialization plugins not rendering audio output.

diff --git a/Engine/Source/Runtime/AudioMixer/Private/AudioMixerSourceManager.cpp b/Engine/Source/Runtime/AudioMixer/Private/AudioMixerSourceManager.cpp
--- a/Engine/Source/Runtime/AudioMixer/Private/AudioMixerSourceManager.cpp
+++ b/Engine/Source/Runtime/AudioMixer/Private/AudioMixerSourceManager.cpp
@@ -996,7 +996,7 @@ namespace Audio
 			FSourceInfo& SourceInfo = SourceInfos[SourceId];
 			FSourceDownmixData& DownmixData = DownmixDataArray[SourceId];
 
-			if (DownmixData.NumInputChannels != NumInputChannels)
+			if (DownmixData.NumInputChannels != NumInputChannels && !SourceInfo.bUseHRTFSpatializer)
 			{
 				// This means that this source has been reinitialized as a different source while this command was in flight,
 				// In which case it is of no use to us. Exit.
