fixed read access violation when linearinterpfractionaldelay set to maximum

diff --git a/Engine/Source/Runtime/AudioMixer/Private/DSP/LinearInterpFractionalDelay.cpp b/Engine/Source/Runtime/AudioMixer/Private/DSP/LinearInterpFractionalDelay.cpp
--- a/Engine/Source/Runtime/AudioMixer/Private/DSP/LinearInterpFractionalDelay.cpp
+++ b/Engine/Source/Runtime/AudioMixer/Private/DSP/LinearInterpFractionalDelay.cpp
@@ -28,8 +28,8 @@ FLinearInterpFractionalDelay::FLinearInterpFractionalDelay(int32 InMaxDelay, int
 	}
 
 	// Allocate and prepare delay line for maximum delay.
-	DelayLine = MakeUnique<FAlignedBlockBuffer>((2 * MaxDelay) + NumInternalBufferSamples, MaxDelay + NumInternalBufferSamples);
-	DelayLine->AddZeros(MaxDelay);
+	DelayLine = MakeUnique<FAlignedBlockBuffer>((2 * (MaxDelay + 1)) + NumInternalBufferSamples, MaxDelay + NumInternalBufferSamples + 1);
+	DelayLine->AddZeros(MaxDelay + 1);
 	
 	IntegerDelayOffsets.Reset(NumInternalBufferSamples);
 	IntegerDelayOffsets.AddUninitialized(NumInternalBufferSamples);
@@ -56,7 +56,7 @@ FLinearInterpFractionalDelay::~FLinearInterpFractionalDelay()
 void FLinearInterpFractionalDelay::Reset() 
 {
 	DelayLine->ClearSamples();
-	DelayLine->AddZeros(MaxDelay);
+	DelayLine->AddZeros(MaxDelay + 1);
 }
 
 
@@ -107,7 +107,7 @@ void FLinearInterpFractionalDelay::ProcessAudioBlock(const float* InSamples, con
 	// Update delay line.
 	DelayLine->AddSamples(InSamples, InNum);
 
-	const float* DelayData = DelayLine->InspectSamples(InNum + MaxDelay);
+	const float* DelayData = DelayLine->InspectSamples(InNum + MaxDelay + 1);
 	const int32* IntegerDelayOffsetData = IntegerDelayOffsets.GetData();
 
 	const VectorRegister VMaxDelay = MakeVectorRegister((float)MaxDelay, (float)MaxDelay, (float)MaxDelay, (float)MaxDelay);
@@ -126,14 +126,15 @@ void FLinearInterpFractionalDelay::ProcessAudioBlock(const float* InSamples, con
 		VectorRegister VLowerCoefficients = VectorSubtract(GlobalVectorConstants::FloatOne, VUpperCoefficients);
 
 
-		// Make integer locations relative to bblock
+		// Make integer locations relative to block
 		VectorRegisterInt VIntegerDelays = VectorFloatToInt(VFloorDelays);
 		VectorRegisterInt VIntegerDelayOffset = VectorIntLoadAligned(&IntegerDelayOffsetData[i]);
 		VIntegerDelays = VectorIntSubtract(VIntegerDelayOffset, VIntegerDelays);
 
 		// Lookup samples for interpolation
-		VectorIntStoreAligned(VIntegerDelays, LowerDelayPos);
-		VectorIntStoreAligned(VectorIntSubtract(VIntegerDelays, GlobalVectorConstants::IntOne), UpperDelayPos);
+		VectorIntStoreAligned(VIntegerDelays, UpperDelayPos);
+		VectorIntStoreAligned(VectorIntAdd(VIntegerDelays, GlobalVectorConstants::IntOne), LowerDelayPos);
+		
 		VectorRegister VLowerSamples = MakeVectorRegister(
 			DelayData[LowerDelayPos[0]],
 			DelayData[LowerDelayPos[1]],
