Submitting for Maxwell.Hayes
Change to bypass High-pass filters for frequencies <= 20hz
Fix for per-source LP/HP filters failing to be gracefully bypassed

https://p4-swarm.epicgames.net/reviews/6541561
https://p4-swarm.epicgames.net/reviews/6523646

diff --git a/Engine/Source/Runtime/AudioMixer/Private/AudioMixerSourceManager.cpp b/Engine/Source/Runtime/AudioMixer/Private/AudioMixerSourceManager.cpp
--- a/Engine/Source/Runtime/AudioMixer/Private/AudioMixerSourceManager.cpp
+++ b/Engine/Source/Runtime/AudioMixer/Private/AudioMixerSourceManager.cpp
@@ -2145,13 +2145,14 @@ namespace Audio
 				continue;
 			}
 
-			if (SourceInfo.bIs3D)
+			if (SourceInfo.bIs3D && !DownmixData.bIsInitialDownmix)
 			{
 				ComputeDownmix3D(DownmixData);
 			}
 			else
 			{
 				ComputeDownmix2D(DownmixData);
+				DownmixData.bIsInitialDownmix = false;
 			}
 		}
 	}
diff --git a/Engine/Source/Runtime/AudioMixer/Private/AudioMixerSourceManager.h b/Engine/Source/Runtime/AudioMixer/Private/AudioMixerSourceManager.h
--- a/Engine/Source/Runtime/AudioMixer/Private/AudioMixerSourceManager.h
+++ b/Engine/Source/Runtime/AudioMixer/Private/AudioMixerSourceManager.h
@@ -146,6 +146,7 @@ namespace Audio
 		// This is the number of bytes the gain array is using:
 		// (Number of input channels * number of output channels) * sizeof float.
 		int32 CopySize;
+		bool bIsInit;
 
 		FSourceChannelMap(int32 InNumInChannels, int32 InNumOutChannels) 
 			: CopySize(InNumInChannels * InNumOutChannels * sizeof(float))
@@ -163,6 +164,7 @@ namespace Audio
 			CopySize = InNumInChannels * InNumOutChannels * sizeof(float);
 			FMemory::Memzero(ChannelStartGains, CopySize);
 			FMemory::Memzero(ChannelDestinationGains, CopySize);
+			bIsInit = true;
 		}
 
 		FORCEINLINE void CopyDestinationToStart()
@@ -173,11 +175,16 @@ namespace Audio
 		FORCEINLINE void SetChannelMap(const float* RESTRICT InChannelGains)
 		{
 			FMemory::Memcpy(ChannelDestinationGains, InChannelGains, CopySize);
+			if (bIsInit)
+			{
+				FMemory::Memcpy(ChannelStartGains, InChannelGains, CopySize);
+			}
 		}
 
 	private:
 		FSourceChannelMap()
 			: CopySize(0)
+			, bIsInit(true)
 		{
 		}
 	};
@@ -361,6 +368,7 @@ namespace Audio
 			uint32 NumInputChannels;
 			const uint32 NumFrames;
 			uint32 NumDeviceChannels;
+			uint8 bIsInitialDownmix : 1;
 
 			FSourceDownmixData(uint32 SourceNumChannels, uint32 NumDeviceOutputChannels, uint32 InNumFrames)
 				: PostEffectBuffers(nullptr)
@@ -373,6 +381,7 @@ namespace Audio
 				, NumInputChannels(SourceNumChannels)
 				, NumFrames(InNumFrames)
 				, NumDeviceChannels(NumDeviceOutputChannels)
+				, bIsInitialDownmix(true)
 			{
 			}
 
@@ -394,6 +403,7 @@ namespace Audio
 				FiveOneSubmixInfo.Reset(NumInputChannels, 6, NumFrames);
 				SevenOneSubmixInfo.Reset(NumInputChannels, 8, NumFrames);
 				AmbisonicsSubmixInfo.Reset(NumInputChannels, 4, NumFrames);
+				bIsInitialDownmix = true;
 			}
 		};
 
diff --git a/Engine/Source/Runtime/AudioMixer/Public/DSP/OnePole.h b/Engine/Source/Runtime/AudioMixer/Public/DSP/OnePole.h
--- a/Engine/Source/Runtime/AudioMixer/Public/DSP/OnePole.h
+++ b/Engine/Source/Runtime/AudioMixer/Public/DSP/OnePole.h
@@ -117,7 +117,7 @@ namespace Audio
 			if (!FMath::IsNearlyEqual(InFrequency, CutoffFrequency))
 			{
 				CutoffFrequency = InFrequency;
-				float NormalizedFreq = FMath::Clamp(0.5f * InFrequency / SampleRate, 0.0f, 1.0f);
+				float NormalizedFreq = FMath::Clamp(2.0f * InFrequency / SampleRate, 0.0f, 1.0f);
 				B1 = FMath::Exp(-PI * NormalizedFreq);
 				A0 = 1.0f - B1;
 			}
