[NO NON-PII COMMIT MESSAGE]

diff --git a/Engine/Source/Runtime/AudioMixer/Private/AudioMixerSourceManager.cpp b/Engine/Source/Runtime/AudioMixer/Private/AudioMixerSourceManager.cpp
--- a/Engine/Source/Runtime/AudioMixer/Private/AudioMixerSourceManager.cpp
+++ b/Engine/Source/Runtime/AudioMixer/Private/AudioMixerSourceManager.cpp
@@ -66,6 +66,14 @@ FAutoConsoleVariableRef CVarBypassAudioPlugins(
 	TEXT("0: Not Disabled, 1: Disabled"),
 	ECVF_Default);
 
+static int32 FlushCommandBufferOnTimeoutCvar = 0;
+FAutoConsoleVariableRef CVarFlushCommandBufferOnTimeout(
+	TEXT("au.FlushCommandBufferOnTimeout"),
+	FlushCommandBufferOnTimeoutCvar,
+	TEXT("When set to 1, flushes audio render thread synchronously when our fence has timed out.\n")
+	TEXT("0: Not Disabled, 1: Disabled"),
+	ECVF_Default);
+
 #define ONEOVERSHORTMAX (3.0517578125e-5f) // 1/32768
 #define ENVELOPE_TAIL_THRESHOLD (1.58489e-5f) // -96 dB
 
@@ -296,13 +304,23 @@ namespace Audio
 				// And will allow the audio thread to write to a new command slot
 				const int32 NextIndex = (CurrentGameIndex + 1) & 1;
 
+				FCommands& NextCommandBuffer = CommandBuffers[NextIndex];
+
 				// Make sure we've actually emptied the command queue from the render thread before writing to it
-#if !UE_BUILD_SHIPPING
-				if (!bTimedOut)
+				if (FlushCommandBufferOnTimeoutCvar && NextCommandBuffer.SourceCommandQueue.Num() != 0)
 				{
-					check(CommandBuffers[NextIndex].SourceCommandQueue.Num() == 0);
+					UE_LOG(LogAudioMixer, Warning, TEXT("Audio render callback stopped. Flushing %d commands."), NextCommandBuffer.SourceCommandQueue.Num());
+
+					// Pop and execute all the commands that came since last update tick
+					for (int32 Id = 0; Id < NextCommandBuffer.SourceCommandQueue.Num(); ++Id)
+					{
+						TFunction<void()>& CommandFunction = NextCommandBuffer.SourceCommandQueue[Id];
+						CommandFunction();
+						NumCommands.Decrement();
+					}
+
+					NextCommandBuffer.SourceCommandQueue.Reset();
 				}
-#endif
 
 				// Here we ensure that we block for any pending calls to AudioMixerThreadCommand.
 				FScopeLock ScopeLock(&CommandBufferIndexCriticalSection);
