Separating memory alignment specifiers between allocation and pointer alignment. Fixes bug when platforms allocation alignment differs from SIMD alignment.

diff --git a/Engine/Source/Runtime/AudioMixer/Private/DSP/AlignedBlockBuffer.cpp b/Engine/Source/Runtime/AudioMixer/Private/DSP/AlignedBlockBuffer.cpp
--- a/Engine/Source/Runtime/AudioMixer/Private/DSP/AlignedBlockBuffer.cpp
+++ b/Engine/Source/Runtime/AudioMixer/Private/DSP/AlignedBlockBuffer.cpp
@@ -6,18 +6,19 @@
 
 using namespace Audio;
 
-FAlignedBlockBuffer::FAlignedBlockBuffer(int32 InSampleCapacity, int32 InMaxNumInspectSamples, uint32 InByteAlignment)
-:	ByteAlignment(InByteAlignment),
-	FloatAlignment(0),
-	InternalBuffer(nullptr),
-	RolloverBuffer(nullptr),
-	ReadPointer(nullptr),
-	WritePointer(nullptr),
-	SampleCapacity(InSampleCapacity),
-	MaxNumInspectSamples(InMaxNumInspectSamples),
-	WriteCount(0),
-	WritePos(0),
-	ReadPos(0)
+FAlignedBlockBuffer::FAlignedBlockBuffer(int32 InSampleCapacity, int32 InMaxNumInspectSamples, uint32 InByteAlignment, uint32 InAllocByteAlignment)
+	:	AllocByteAlignment(InAllocByteAlignment)
+	,	ByteAlignment(InByteAlignment)
+	,	FloatAlignment(0)
+	,	InternalBuffer(nullptr)
+	,	RolloverBuffer(nullptr)
+	,	ReadPointer(nullptr)
+	,	WritePointer(nullptr)
+	,	SampleCapacity(InSampleCapacity)
+	,	MaxNumInspectSamples(InMaxNumInspectSamples)
+	,	WriteCount(0)
+	,	WritePos(0)
+	,	ReadPos(0)
 {
 	FloatAlignment = ByteAlignment / sizeof(float);
 
@@ -297,7 +298,7 @@ void FAlignedBlockBuffer::SetMaxNumInspectSamples(int32 InMax)
 float* FAlignedBlockBuffer::AllocateAlignedFloatArray(int32 InNum) const 
 {
 	int32 Size = InNum * sizeof(float);
-	return (float*)FMemory::Malloc(Size, ByteAlignment);
+	return (float*)FMemory::Malloc(Size, AllocByteAlignment);
 }
 
 void FAlignedBlockBuffer::FreeFloatArray(float*& Array) const
diff --git a/Engine/Source/Runtime/AudioMixer/Public/DSP/AlignedBlockBuffer.h b/Engine/Source/Runtime/AudioMixer/Public/DSP/AlignedBlockBuffer.h
--- a/Engine/Source/Runtime/AudioMixer/Public/DSP/AlignedBlockBuffer.h
+++ b/Engine/Source/Runtime/AudioMixer/Public/DSP/AlignedBlockBuffer.h
@@ -16,7 +16,7 @@ namespace Audio
 	public:
 		
 		// Constructor
-		FAlignedBlockBuffer(int32 InSampleCapacity=8192, int32 InMaxNumInspectSamples=1024, uint32 InByteAlignment=AUDIO_BUFFER_ALIGNMENT);
+		FAlignedBlockBuffer(int32 InSampleCapacity=8192, int32 InMaxNumInspectSamples=1024, uint32 InByteAlignment=AUDIO_SIMD_BYTE_ALIGNMENT, uint32 InAllocByteAlignment=AUDIO_BUFFER_ALIGNMENT);
 
 		// Destructor
 		~FAlignedBlockBuffer() throw();
@@ -66,6 +66,7 @@ namespace Audio
 		FAlignedBlockBuffer(const FAlignedBlockBuffer& CopyBuff);
 		
 		// Alignment definition
+		uint32 AllocByteAlignment;
 		uint32 ByteAlignment;
 		uint32 FloatAlignment;
 
