Merge refactor of CircularAudioBuffer from Dev-Online. #jira none #rb none

diff --git a/Engine/Source/Runtime/AudioMixer/Public/DSP/Dsp.h b/Engine/Source/Runtime/AudioMixer/Public/DSP/Dsp.h
--- a/Engine/Source/Runtime/AudioMixer/Public/DSP/Dsp.h
+++ b/Engine/Source/Runtime/AudioMixer/Public/DSP/Dsp.h
@@ -435,6 +435,13 @@ namespace Audio
 	template <class SampleType>
 	class TCircularAudioBuffer
 	{
+	private:
+
+		TArray<SampleType> InternalBuffer;
+		uint32 Capacity;
+		FThreadSafeCounter ReadCounter;
+		FThreadSafeCounter WriteCounter;
+
 	public:
 		TCircularAudioBuffer()
 		{
@@ -449,7 +456,7 @@ namespace Audio
 		void SetCapacity(uint32 InCapacity)
 		{
 			checkf(InCapacity < (uint32)TNumericLimits<int32>::Max(), TEXT("Max capacity for this buffer is 2,147,483,647 samples. Otherwise our index arithmetic will not work."));
-			Capacity = InCapacity;
+			Capacity = InCapacity + 1;
 			ReadCounter.Set(0);
 			WriteCounter.Set(0);
 			InternalBuffer.Reset();
@@ -457,101 +464,51 @@ namespace Audio
 		}
 
 		// Pushes some amount of samples into this circular buffer.
-		// Returns the amount of samples remaining if positive,
-		// and the number of samples short we were if negative.
+		// Returns the amount of samples read
 		int32 Push(const SampleType* InBuffer, uint32 NumSamples)
 		{
 			SampleType* DestBuffer = InternalBuffer.GetData();
 			const uint32 ReadIndex = ReadCounter.GetValue();
 			const uint32 WriteIndex = WriteCounter.GetValue();
 
-			// Determine the actual capacity we have to ensure we don't write past the read index.
-			const uint32 CurrentSlack = (ReadIndex <= WriteIndex) ? ((ReadIndex + Capacity) - WriteIndex) : (ReadIndex - WriteIndex);
-			const uint32 NumSamplesToCopy = FMath::Min(CurrentSlack - 1, NumSamples);
-			const uint32 DestBufferRemainder = Capacity - WriteIndex;
-			if (NumSamplesToCopy >= DestBufferRemainder)
-			{
-				// Copy to the end of Internal Buffer and then the beginning once we wrap around.
-				FMemory::Memcpy(DestBuffer + WriteIndex, InBuffer, DestBufferRemainder * sizeof(SampleType));
+			int32 NumToCopy = FMath::Min<int32>(NumSamples, Remainder());
+			const int32 NumToWrite = FMath::Min<int32>(NumToCopy, Capacity - WriteIndex);
+			FMemory::Memcpy(&DestBuffer[WriteIndex], InBuffer, NumToWrite * sizeof(SampleType));
+			FMemory::Memcpy(&DestBuffer[0], &InBuffer[NumToWrite], (NumToCopy - NumToWrite) * sizeof(SampleType));
 
-				const uint32 InBufferRemainder = NumSamplesToCopy - DestBufferRemainder;
-				FMemory::Memcpy(DestBuffer, InBuffer + DestBufferRemainder, InBufferRemainder * sizeof(SampleType));
-				WriteCounter.Set(InBufferRemainder);
-			}
-			else
-			{
-				// We have enough space to copy the full buffer.
-				FMemory::Memcpy(DestBuffer + WriteIndex, InBuffer, NumSamplesToCopy * sizeof(SampleType));
-				WriteCounter.Add(NumSamplesToCopy);
-			}
+			WriteCounter.Set((WriteIndex + NumToCopy) % Capacity);
 
-			check(NumSamples < ((uint32)TNumericLimits<int32>::Max()));
-			return ((int32) CurrentSlack) - NumSamples;
+			return NumToCopy;
 		}
 
-		// Pops some amount of samples into this circular buffer.
-		// Returns the amount of samples written to OutBuffer if positive,
-		// and the number of samples short we were if negative.
-		int32 Pop(SampleType* OutBuffer, uint32 NumSamples)
+		// Same as Pop(), but does not increment the read counter.
+		int32 Peek(SampleType* OutBuffer, uint32 NumSamples)
 		{
 			SampleType* SrcBuffer = InternalBuffer.GetData();
 			const uint32 ReadIndex = ReadCounter.GetValue();
 			const uint32 WriteIndex = WriteCounter.GetValue();
 
-			// Determine the actual capacity we have to ensure we don't read past the write index.
-			const uint32 CurrentSlack = (WriteIndex < ReadIndex) ? ((WriteIndex + Capacity) - ReadIndex) : (WriteIndex - ReadIndex);
-			const uint32 NumSamplesToCopy = FMath::Min(CurrentSlack, NumSamples);
-			const uint32 SrcBufferRemainder = Capacity - ReadIndex;
-			if (NumSamplesToCopy >= SrcBufferRemainder)
-			{
-				// Copy from the end of Internal Buffer and then the beginning once we wrap around.
-				FMemory::Memcpy(OutBuffer, SrcBuffer + ReadIndex, SrcBufferRemainder * sizeof(SampleType));
+			int32 NumToCopy = FMath::Min<int32>(NumSamples, Num());
 
-				const uint32 OutBufferRemainder = NumSamplesToCopy - SrcBufferRemainder;
-				FMemory::Memcpy(OutBuffer + SrcBufferRemainder, SrcBuffer, OutBufferRemainder * sizeof(SampleType));
-
-				ReadCounter.Set(SrcBufferRemainder);
-			}
-			else
-			{
-				// we have enough space to copy the full buffer.
-				FMemory::Memcpy(OutBuffer, SrcBuffer + ReadIndex, NumSamplesToCopy * sizeof(SampleType));
-				ReadCounter.Add(NumSamplesToCopy);
-			}
+			const int32 NumRead = FMath::Min<int32>(NumToCopy, Capacity - ReadIndex);
+			FMemory::Memcpy(OutBuffer, &SrcBuffer[ReadIndex], NumRead * sizeof(SampleType));
+			FMemory::Memcpy(&OutBuffer[NumRead], &SrcBuffer[0], (NumSamples - NumRead) * sizeof(SampleType));
 
-			check(NumSamples < ((uint32) TNumericLimits<int32>::Max()));
+			check(NumSamples < ((uint32)TNumericLimits<int32>::Max()));
 
-			return ((int32)CurrentSlack) - NumSamples;
+			return NumToCopy;
 		}
 
-		// Same as Pop(), but does not increment the read counter.
-		int32 Peek(SampleType* OutBuffer, uint32 NumSamples)
+		// Pops some amount of samples into this circular buffer.
+		// Returns the amount of samples read.
+		int32 Pop(SampleType* OutBuffer, uint32 NumSamples)
 		{
-			SampleType* SrcBuffer = InternalBuffer.GetData();
-			const uint32 ReadIndex = ReadCounter.GetValue();
-			const uint32 WriteIndex = WriteCounter.GetValue();
-
-			// Determine the actual capacity we have to ensure we don't read past the write index.
-			const uint32 CurrentSlack = (WriteIndex < ReadIndex) ? (WriteIndex + Capacity - ReadIndex) : (WriteIndex - ReadIndex);
-			const uint32 NumSamplesToCopy = FMath::Min(CurrentSlack, NumSamples);
-			const uint32 SrcBufferRemainder = Capacity - ReadIndex;
-			if (NumSamplesToCopy > SrcBufferRemainder)
-			{
-				// Copy from the end of Internal Buffer and then the beginning once we wrap around.
-				FMemory::Memcpy(OutBuffer, SrcBuffer + ReadIndex, SrcBufferRemainder * sizeof(SampleType));
-
-				const uint32 OutBufferRemainder = NumSamplesToCopy - SrcBufferRemainder;
-				FMemory::Memcpy(OutBuffer + SrcBufferRemainder, SrcBuffer, OutBufferRemainder * sizeof(SampleType));
-			}
-			else
-			{
-				// we have enough space to copy the full buffer.
-				FMemory::Memcpy(OutBuffer, SrcBuffer + ReadIndex, NumSamplesToCopy * sizeof(SampleType));
-			}
-
+			int32 NumSamplesRead = Peek(OutBuffer, NumSamples);
 			check(NumSamples < ((uint32)TNumericLimits<int32>::Max()));
 
-			return ((int32)CurrentSlack) - NumSamples;
+			ReadCounter.Set((ReadCounter.GetValue() + NumSamplesRead) % Capacity);
+
+			return NumSamplesRead;
 		}
 
 		// When called, seeks the read or write cursor to only retain either the NumSamples latest data
@@ -580,9 +537,17 @@ namespace Audio
 		// Get number of samples that can be popped off of the buffer.
 		uint32 Num()
 		{
-			const uint32 ReadIndex = ReadCounter.GetValue();
-			const uint32 WriteIndex = WriteCounter.GetValue();
-			return (WriteIndex < ReadIndex) ? ((WriteIndex + Capacity) - ReadIndex) : (WriteIndex - ReadIndex);
+			const int32 ReadIndex = ReadCounter.GetValue();
+			const int32 WriteIndex = WriteCounter.GetValue();
+
+			if (WriteIndex >= ReadIndex)
+			{
+				return WriteIndex - ReadIndex;
+			}
+			else
+			{
+				return Capacity - ReadIndex + WriteIndex;
+			}
 		}
 
 		// Get number of samples that can be pushed onto the buffer before it is full.
@@ -590,15 +555,9 @@ namespace Audio
 		{
 			const uint32 ReadIndex = ReadCounter.GetValue();
 			const uint32 WriteIndex = WriteCounter.GetValue();
-			return (ReadIndex <= WriteIndex) ? ((ReadIndex + Capacity) - WriteIndex) : (ReadIndex - WriteIndex);
-		}
 
-	private:
-
-		TArray<SampleType> InternalBuffer;
-		uint32 Capacity;
-		FThreadSafeCounter ReadCounter;
-		FThreadSafeCounter WriteCounter;
+			return (Capacity - 1 - WriteIndex + ReadIndex) % Capacity;
+		}
 	};
 }
 
