Added interpolated one-pole pass filters and made correspoding changes to Audio Mixer Source Mananger

diff --git a/Engine/Source/Runtime/AudioMixer/Private/AudioMixerSourceManager.cpp b/Engine/Source/Runtime/AudioMixer/Private/AudioMixerSourceManager.cpp
--- a/Engine/Source/Runtime/AudioMixer/Private/AudioMixerSourceManager.cpp
+++ b/Engine/Source/Runtime/AudioMixer/Private/AudioMixerSourceManager.cpp
@@ -2083,36 +2083,40 @@ namespace Audio
 				SourceInfo.SourceListener->OnEffectTailsDone();
 			}
 
-			// Only scale with distance attenuation and send to source audio to plugins if we're not in output-to-bus only mode
-			if (!SourceInfo.bOutputToBusOnly && !DisableFilteringCvar)
+			if (!SourceInfo.bOutputToBusOnly)
 			{
+				// Only scale with distance attenuation and send to source audio to plugins if we're not in output-to-bus only mode
+				const int32 NumOutputSamplesThisSource = NumOutputFrames * SourceInfo.NumInputChannels;
+
+				const bool BypassLPF = DisableFilteringCvar || (SourceInfo.LowPassFilter.GetCutoffFrequency() >= (MAX_FILTER_FREQUENCY - KINDA_SMALL_NUMBER));
+				const bool BypassHPF = DisableFilteringCvar || DisableHPFilteringCvar || (SourceInfo.HighPassFilter.GetCutoffFrequency() <= (MIN_FILTER_FREQUENCY + KINDA_SMALL_NUMBER));
+
 				float* PostDistanceAttenBufferPtr = SourceInfo.SourceBuffer.GetData();
+				float* HpfInputBuffer = PreDistanceAttenBufferPtr; // assume bypassing LPF (HPF uses input buffer as input)
 
-				// Process the filters after the source effects (Using optimized filters)
-				for (int32 SampleIndex = 0; SampleIndex < NumOutputFrames * SourceInfo.NumInputChannels; SampleIndex += SourceInfo.NumInputChannels)
+				if (!BypassLPF)
 				{
-					SourceInfo.LowPassFilter.ProcessAudioFrame(&PreDistanceAttenBufferPtr[SampleIndex], &PostDistanceAttenBufferPtr[SampleIndex]);
+					// Not bypassing LPF, so tell HPF to use LPF output buffer as input
+					HpfInputBuffer = PostDistanceAttenBufferPtr;
 
-					if (!DisableHPFilteringCvar)
-					{
-						SourceInfo.HighPassFilter.ProcessAudioFrame(&PostDistanceAttenBufferPtr[SampleIndex], &PostDistanceAttenBufferPtr[SampleIndex]);
-					}
+					// process LPF audio block
+					SourceInfo.LowPassFilter.ProcessAudioBuffer(PreDistanceAttenBufferPtr, PostDistanceAttenBufferPtr, NumOutputSamplesThisSource);
+				}
+
+				if(!BypassHPF)
+				{
+					// process HPF audio block
+					SourceInfo.HighPassFilter.ProcessAudioBuffer(HpfInputBuffer, PostDistanceAttenBufferPtr, NumOutputSamplesThisSource);
 				}
 
 				// We manually reset interpolation to avoid branches in filter code
 				SourceInfo.LowPassFilter.StopFrequencyInterpolation();
 				SourceInfo.HighPassFilter.StopFrequencyInterpolation();
 
-				// Apply distance attenuation
-				ApplyDistanceAttenuation(SourceInfo, NumSamples);
-
-				// Send source audio to plugins
-				ComputePluginAudio(SourceInfo, DownmixData, SourceId, NumSamples);
-			}
-			else if (DisableFilteringCvar)
-			{
-				float* PostDistanceAttenBufferPtr = SourceInfo.SourceBuffer.GetData();
-				FMemory::Memcpy(PostDistanceAttenBufferPtr, PreDistanceAttenBufferPtr, NumSamples * sizeof(float));
+				if (BypassLPF && BypassHPF)
+				{
+					FMemory::Memcpy(PostDistanceAttenBufferPtr, PreDistanceAttenBufferPtr, NumSamples * sizeof(float));
+				}
 
 				// Apply distance attenuation
 				ApplyDistanceAttenuation(SourceInfo, NumSamples);
diff --git a/Engine/Source/Runtime/AudioMixer/Private/DSP/InterpolatedOnePole.cpp b/Engine/Source/Runtime/AudioMixer/Private/DSP/InterpolatedOnePole.cpp
--- a/Engine/Source/Runtime/AudioMixer/Private/DSP/InterpolatedOnePole.cpp
+++ b/Engine/Source/Runtime/AudioMixer/Private/DSP/InterpolatedOnePole.cpp
@@ -19,6 +19,7 @@ namespace Audio
 		NumChannels = InNumChannels;
 		CutoffFrequency = -1.0f;
 		Z1.Init(0.0f, NumChannels);
+		Z1Data = Z1.GetData();
 		Reset();
 	}
 
@@ -30,10 +31,15 @@ namespace Audio
 		if (!FMath::IsNearlyEqual(InTargetFrequency, CutoffFrequency))
 		{
 			CutoffFrequency = InTargetFrequency;
-			float NormalizedFreq = FMath::Clamp(0.5f * InTargetFrequency / SampleRate, 0.0f, 1.0f);
+			float NormalizedFreq = FMath::Clamp(2.0f * InTargetFrequency / SampleRate, 0.0f, 1.0f);
 			B1Target = FMath::Exp(-PI * NormalizedFreq);
 			B1Delta = (B1Target - B1Curr) / static_cast<float>(InterpLength);
 		}
+
+		if (InterpLength <= 1)
+		{
+			StopFrequencyInterpolation();
+		}
 	}
 
 	void FInterpolatedLPF::ProcessAudioFrame(float* RESTRICT InputFrame, float* RESTRICT OutputFrame)
@@ -55,7 +61,6 @@ namespace Audio
 
 		for (int32 i = 0; i < NumChannels; ++i)
 		{
-			float* Z1Data = Z1.GetData();
 			const float InputSample = InputFrame[i];
 			float Yn = InputSample + B1Curr * (Z1Data[i] - InputSample); // LPF
 			Yn = UnderflowClamp(Yn);
@@ -64,6 +69,25 @@ namespace Audio
 		}
 	}
 
+	void FInterpolatedLPF::ProcessAudioBuffer(float *RESTRICT InputBuffer, float *RESTRICT OutputBuffer, const int32 NumSamples)
+	{
+		for (int SampleIndex = 0; SampleIndex < NumSamples; ++SampleIndex)
+		{
+			// cache which delay term we should be using
+			const int32 ChannelIndex = SampleIndex % NumChannels;
+
+			// step forward coefficient
+			B1Curr += B1Delta;
+			++CurrInterpCounter;
+
+			const float InputSample = InputBuffer[SampleIndex];
+			float Yn = InputSample + B1Curr * (Z1Data[ChannelIndex] - InputSample); // LPF
+			Yn = UnderflowClamp(Yn);
+			Z1Data[ChannelIndex] = Yn;
+			OutputBuffer[SampleIndex] = Yn;
+		}
+	}
+
 	void FInterpolatedLPF::Reset()
 	{
 		B1Curr = 0.0f;
@@ -72,6 +96,7 @@ namespace Audio
 		CurrInterpLength = 0;
 		CurrInterpCounter = 0;
 		ClearMemory();
+		Z1Data = Z1.GetData();
 	}
 
 	void FInterpolatedLPF::ClearMemory()
@@ -95,6 +120,7 @@ namespace Audio
 		NumChannels = InNumChannels;
 		CutoffFrequency = -1.0f;
 		Z1.Init(0.0f, NumChannels);
+		Z1Data = Z1.GetData();
 		Reset();
 	}
 
@@ -112,6 +138,11 @@ namespace Audio
 
 			A0Delta = (A0Target - A0Curr) / static_cast<float>(InterpLength);
 		}
+
+		if (InterpLength <= 1)
+		{
+			StopFrequencyInterpolation();
+		}
 	}
 
 	void FInterpolatedHPF::ProcessAudioFrame(float* RESTRICT InputFrame, float* RESTRICT OutputFrame)
@@ -121,7 +152,6 @@ namespace Audio
 
 		for (int32 i = 0; i < NumChannels; ++i)
 		{
-			float* Z1Data = Z1.GetData();
 			const float InputSample = InputFrame[i];
 			const float Vn = (InputSample - Z1Data[i]) * A0Curr;
 			const float LPF = Vn + Z1Data[i];
@@ -131,6 +161,27 @@ namespace Audio
 		}
 	}
 
+	void FInterpolatedHPF::ProcessAudioBuffer(float *RESTRICT InputBuffer, float *RESTRICT OutputBuffer, const int32 NumSamples)
+	{
+		for (int SampleIndex = 0; SampleIndex < NumSamples; ++SampleIndex)
+		{
+			// cache which delay term we should be using
+			const int32 ChannelIndex = SampleIndex % NumChannels;
+
+			// step forward coefficient
+			A0Curr += A0Delta;
+			++CurrInterpCounter;
+
+			const float InputSample = InputBuffer[SampleIndex];
+			const float Vn = (InputSample - Z1Data[ChannelIndex]) * A0Curr;
+			const float LPF = Vn + Z1Data[ChannelIndex];
+			Z1Data[ChannelIndex] = Vn + LPF;
+
+			OutputBuffer[SampleIndex] = InputSample - LPF;
+		}
+	}
+
+
 	void FInterpolatedHPF::Reset()
 	{
 		A0Curr = 0.0f;
@@ -138,6 +189,7 @@ namespace Audio
 		CurrInterpLength = 0;
 		CurrInterpCounter = 0;
 		ClearMemory();
+		Z1Data = Z1.GetData();
 	}
 
 	void FInterpolatedHPF::ClearMemory()
diff --git a/Engine/Source/Runtime/AudioMixer/Public/DSP/InterpolatedOnePole.h b/Engine/Source/Runtime/AudioMixer/Public/DSP/InterpolatedOnePole.h
--- a/Engine/Source/Runtime/AudioMixer/Public/DSP/InterpolatedOnePole.h
+++ b/Engine/Source/Runtime/AudioMixer/Public/DSP/InterpolatedOnePole.h
@@ -34,6 +34,9 @@ namespace Audio
 		// interpolates coefficient and processes a sample
 		void ProcessAudioFrame(float* RESTRICT InputFrame, float* RESTRICT OutputFrame);
 
+		// interpolates coefficient and processes a buffer
+		void ProcessAudioBuffer(float* RESTRICT InputBuffer, float* RESTRICT OutputBuffer, const int32 NumSamples);
+
 		/*
 			StopFrequencyInterpolation() needs to be called manually when the interpolation should be done.
 			Snaps the coefficient to the target value.
@@ -58,6 +61,7 @@ namespace Audio
 		float B1Delta{ 0.0f }; // coefficient step size 
 		float B1Target{ 0.0f };
 		TArray<float> Z1; // multi-channel delay terms
+		float* Z1Data{ nullptr };
 		int32 CurrInterpLength{ 0 };
 		int32 CurrInterpCounter{ 0 };
 		int32 NumInterpSteps;
@@ -95,6 +99,9 @@ namespace Audio
 		// interpolates coefficient and processes a sample
 		void ProcessAudioFrame(float* RESTRICT InputFrame, float* RESTRICT OutputFrame);
 
+		// interpolates coefficient and processes a buffer
+		void ProcessAudioBuffer(float* RESTRICT InputBuffer, float* RESTRICT OutputBuffer, const int32 NumSamples);
+
 		/*
 			StopFrequencyInterpolation() needs to be called manually when the interpolation should be done.
 			Snaps the coefficient to the target value.
@@ -127,6 +134,7 @@ namespace Audio
 		float A0Delta{ 0.0f }; // coefficient step size
 		float A0Target{ 0.0f };
 		TArray<float> Z1; // multi-channel delay terms
+		float* Z1Data{ nullptr };
 		int32 CurrInterpLength{ 0 };
 		int32 CurrInterpCounter{ 0 };
 		int32 NumInterpSteps;
