Fix for submix effects not registering until restart.

diff --git a/Engine/Source/Runtime/AudioMixer/Private/AudioMixerDevice.cpp b/Engine/Source/Runtime/AudioMixer/Private/AudioMixerDevice.cpp
--- a/Engine/Source/Runtime/AudioMixer/Private/AudioMixerDevice.cpp
+++ b/Engine/Source/Runtime/AudioMixer/Private/AudioMixerDevice.cpp
@@ -1000,32 +1000,42 @@ namespace Audio
 
 			if (!IsMasterSubmixType(InSoundSubmix))
 			{
+				FMixerSubmixPtr MixerSubmix; 
+
 				// If the sound submix wasn't already registered get it into the system.
 				if (!Submixes.Contains(InSoundSubmix))
 				{
-					FMixerSubmixPtr MixerSubmix = FMixerSubmixPtr(new FMixerSubmix(this));
+					MixerSubmix = FMixerSubmixPtr(new FMixerSubmix(this));
 					Submixes.Add(InSoundSubmix, MixerSubmix);
+				}
+				else
+				{
+					FMixerSubmixPtr* ExistingMixerSubmix = Submixes.Find(InSoundSubmix);
+					check(ExistingMixerSubmix);
+					MixerSubmix = *ExistingMixerSubmix;
+				}
+
+				check(MixerSubmix.IsValid());
 
-					if (bInit)
+				if (bInit)
+				{
+					// Setup the parent-child relationship
+					FMixerSubmixWeakPtr ParentSubmixInstance;
+					if (InSoundSubmix->ParentSubmix)
 					{
-						// Setup the parent-child relationship
-						FMixerSubmixWeakPtr ParentSubmixInstance;
-						if (InSoundSubmix->ParentSubmix)
-						{
-							ParentSubmixInstance = GetSubmixInstance(InSoundSubmix->ParentSubmix);
-						}
-						else
-						{
-							ParentSubmixInstance = GetMasterSubmix();
-						}
+						ParentSubmixInstance = GetSubmixInstance(InSoundSubmix->ParentSubmix);
+					}
+					else
+					{
+						ParentSubmixInstance = GetMasterSubmix();
+					}
 
-						FMixerSubmixPtr ParentSubmixInstancePtr = ParentSubmixInstance.Pin();
-						if (ParentSubmixInstancePtr.IsValid())
-						{
-							ParentSubmixInstancePtr->AddChildSubmix(MixerSubmix);
-							MixerSubmix->SetParentSubmix(ParentSubmixInstance);
-							MixerSubmix->Init(InSoundSubmix);
-						}
+					FMixerSubmixPtr ParentSubmixInstancePtr = ParentSubmixInstance.Pin();
+					if (ParentSubmixInstancePtr.IsValid())
+					{
+						ParentSubmixInstancePtr->AddChildSubmix(MixerSubmix);
+						MixerSubmix->SetParentSubmix(ParentSubmixInstance);
+						MixerSubmix->Init(InSoundSubmix);
 					}
 				}
 			}
