Crash fix for listener transform update.

diff --git a/Engine/Source/Runtime/AudioMixer/Private/AudioMixerDevice.cpp b/Engine/Source/Runtime/AudioMixer/Private/AudioMixerDevice.cpp
--- a/Engine/Source/Runtime/AudioMixer/Private/AudioMixerDevice.cpp
+++ b/Engine/Source/Runtime/AudioMixer/Private/AudioMixerDevice.cpp
@@ -354,17 +354,6 @@ namespace Audio
 			AudioMixerPlatform->ResumePlaybackOnNewDevice();
 		}
 
-#if 0 // Disable touching the listener transforms
-		ListenerTransforms.Reset();
-		for (FListener& Listener : Listeners)
-		{
-			ListenerTransforms.Add(Listener.Transform);
-		}
-
-		// Update listener transforms, some effects use the listener transform data
-		SourceManager.SetListenerTransforms(ListenerTransforms);
-#endif
-
 		// Loop through any envelope-following submixes and perform any broadcasting of envelope data if needed
 		TArray<float> SubmixEnvelopeData;
 		for (USoundSubmix* SoundSubmix : EnvelopeFollowingSubmixes)
