Fixing one-pole filter pop-in when new (de-virtualized) sources play.

diff --git a/Engine/Source/Runtime/AudioMixer/Private/DSP/InterpolatedOnePole.cpp b/Engine/Source/Runtime/AudioMixer/Private/DSP/InterpolatedOnePole.cpp
--- a/Engine/Source/Runtime/AudioMixer/Private/DSP/InterpolatedOnePole.cpp
+++ b/Engine/Source/Runtime/AudioMixer/Private/DSP/InterpolatedOnePole.cpp
@@ -26,7 +26,6 @@ namespace Audio
 	void FInterpolatedLPF::StartFrequencyInterpolation(const float InTargetFrequency, const int32 InterpLength)
 	{
 		CurrInterpLength = InterpLength;
-		CurrInterpCounter = 0;
 
 		if (!FMath::IsNearlyEqual(InTargetFrequency, CutoffFrequency))
 		{
@@ -45,7 +44,6 @@ namespace Audio
 	void FInterpolatedLPF::ProcessAudioFrame(float* RESTRICT InputFrame, float* RESTRICT OutputFrame)
 	{
 		B1Curr += B1Delta; // step forward coefficient
-		++CurrInterpCounter;
 
 		/*
 			[absorbing A0 coefficient]
@@ -77,8 +75,8 @@ namespace Audio
 			const int32 ChannelIndex = SampleIndex % NumChannels;
 
 			// step forward coefficient
-			B1Curr += B1Delta;
-			++CurrInterpCounter;
+			// Multiply delta by !ChannelIndex so the coefficient only accumulates at the beginning of each frame (on channel 0)
+			B1Curr += B1Delta * !ChannelIndex;
 
 			const float InputSample = InputBuffer[SampleIndex];
 			float Yn = InputSample + B1Curr * (Z1Data[ChannelIndex] - InputSample); // LPF
@@ -94,7 +92,6 @@ namespace Audio
 		B1Delta = 0.0f;
 		B1Target = B1Curr;
 		CurrInterpLength = 0;
-		CurrInterpCounter = 0;
 		ClearMemory();
 		Z1Data = Z1.GetData();
 	}
@@ -127,7 +124,6 @@ namespace Audio
 	void FInterpolatedHPF::StartFrequencyInterpolation(const float InTargetFrequency, const int32 InterpLength)
 	{
 		CurrInterpLength = InterpLength;
-		CurrInterpCounter = 0;
 
 		if (!FMath::IsNearlyEqual(InTargetFrequency, CutoffFrequency))
 		{
@@ -148,7 +144,6 @@ namespace Audio
 	void FInterpolatedHPF::ProcessAudioFrame(float* RESTRICT InputFrame, float* RESTRICT OutputFrame)
 	{
 		A0Curr += A0Delta; // step forward coefficient
-		++CurrInterpCounter;
 
 		for (int32 i = 0; i < NumChannels; ++i)
 		{
@@ -169,8 +164,8 @@ namespace Audio
 			const int32 ChannelIndex = SampleIndex % NumChannels;
 
 			// step forward coefficient
-			A0Curr += A0Delta;
-			++CurrInterpCounter;
+			// Multiply delta by !ChannelIndex so the coefficient only accumulates at the beginning of each frame (on channel 0)
+			A0Curr += A0Delta * !ChannelIndex;
 
 			const float InputSample = InputBuffer[SampleIndex];
 			const float Vn = (InputSample - Z1Data[ChannelIndex]) * A0Curr;
@@ -187,7 +182,6 @@ namespace Audio
 		A0Curr = 0.0f;
 		A0Delta = 0.0f;
 		CurrInterpLength = 0;
-		CurrInterpCounter = 0;
 		ClearMemory();
 		Z1Data = Z1.GetData();
 	}
diff --git a/Engine/Source/Runtime/AudioMixer/Public/DSP/InterpolatedOnePole.h b/Engine/Source/Runtime/AudioMixer/Public/DSP/InterpolatedOnePole.h
--- a/Engine/Source/Runtime/AudioMixer/Public/DSP/InterpolatedOnePole.h
+++ b/Engine/Source/Runtime/AudioMixer/Public/DSP/InterpolatedOnePole.h
@@ -63,7 +63,6 @@ namespace Audio
 		TArray<float> Z1; // multi-channel delay terms
 		float* Z1Data{ nullptr };
 		int32 CurrInterpLength{ 0 };
-		int32 CurrInterpCounter{ 0 };
 		int32 NumInterpSteps;
 		float SampleRate{ 0 };
 		int32 NumChannels{ 1 };
@@ -136,7 +135,6 @@ namespace Audio
 		TArray<float> Z1; // multi-channel delay terms
 		float* Z1Data{ nullptr };
 		int32 CurrInterpLength{ 0 };
-		int32 CurrInterpCounter{ 0 };
 		int32 NumInterpSteps;
 		float SampleRate{ 0 };
 		float NyquistLimit{ 0.0f };
