Fix for xma2 popping issue.

diff --git a/Engine/Source/Runtime/AudioMixer/Private/AudioMixerSourceBuffer.cpp b/Engine/Source/Runtime/AudioMixer/Private/AudioMixerSourceBuffer.cpp
--- a/Engine/Source/Runtime/AudioMixer/Private/AudioMixerSourceBuffer.cpp
+++ b/Engine/Source/Runtime/AudioMixer/Private/AudioMixerSourceBuffer.cpp
@@ -238,7 +238,7 @@ namespace Audio
 
 		// Prepare the buffer for the PCM submission
 		SourceVoiceBuffers[0]->AudioData.Reset(NumSamplesPerBuffer);
-		SourceVoiceBuffers[0]->AudioData.AddUninitialized(NumSamplesPerBuffer);
+		SourceVoiceBuffers[0]->AudioData.AddZeroed(NumSamplesPerBuffer);
 
 		RawPCMDataBuffer.GetNextBuffer(SourceVoiceBuffers[0].Get(), NumSamplesPerBuffer);
 
@@ -268,7 +268,7 @@ namespace Audio
 				for (int32 i = 0; i < 2; ++i)
 				{
 					SourceVoiceBuffers[i]->AudioData.Reset();
-					SourceVoiceBuffers[i]->AudioData.AddUninitialized(NumSamples);
+					SourceVoiceBuffers[i]->AudioData.AddZeroed(NumSamples);
 				}
 
 				int16* CachedBufferPtr0 = (int16*)CachedRealtimeFirstBuffer.GetData();
@@ -294,7 +294,7 @@ namespace Audio
 #elif (PLATFORM_NUM_AUDIODECOMPRESSION_PRECACHE_BUFFERS == 1)
 			{
 				SourceVoiceBuffers[0]->AudioData.Reset();
-				SourceVoiceBuffers[0]->AudioData.AddUninitialized(NumSamples);
+				SourceVoiceBuffers[0]->AudioData.AddZeroed(NumSamples);
 
 				int16* CachedBufferPtr0 = (int16*)CachedRealtimeFirstBuffer.GetData();
 
@@ -321,7 +321,7 @@ namespace Audio
 	{
 		const int32 MaxSamples = MONO_PCM_BUFFER_SAMPLES * NumChannels;
 		SourceVoiceBuffers[BufferIndex]->AudioData.Reset();
-		SourceVoiceBuffers[BufferIndex]->AudioData.AddUninitialized(MaxSamples);
+		SourceVoiceBuffers[BufferIndex]->AudioData.AddZeroed(MaxSamples);
 
 		if (bProcedural)
 		{
diff --git a/Engine/Source/Runtime/Engine/Public/AudioDecompress.h b/Engine/Source/Runtime/Engine/Public/AudioDecompress.h
--- a/Engine/Source/Runtime/Engine/Public/AudioDecompress.h
+++ b/Engine/Source/Runtime/Engine/Public/AudioDecompress.h
@@ -26,6 +26,10 @@ struct FSoundQualityInfo;
 class ICompressedAudioInfo
 {
 public:
+	ICompressedAudioInfo()
+		: StreamingSoundWave(nullptr)
+	{}
+
 	/**
 	* Virtual destructor.
 	*/
