Fixed reverb with recent change for google resonance not doing sends.

diff --git a/Engine/Source/Runtime/AudioMixer/Private/AudioMixerSource.cpp b/Engine/Source/Runtime/AudioMixer/Private/AudioMixerSource.cpp
--- a/Engine/Source/Runtime/AudioMixer/Private/AudioMixerSource.cpp
+++ b/Engine/Source/Runtime/AudioMixer/Private/AudioMixerSource.cpp
@@ -929,16 +929,19 @@ namespace Audio
 			}
 
 			// Send the source audio to the reverb plugin if enabled
+			bool bPluginBypassedMasterReverb = false;
 			if (UseReverbPlugin())
 			{
 				if (MixerDevice->GetMasterReverbPluginSubmix().IsValid())
 				{
 					MixerSourceVoice->SetSubmixSendInfo(MixerDevice->GetMasterReverbPluginSubmix(), ReverbSendLevel);
 				}
+
+				bPluginBypassedMasterReverb = AudioDevice->IsReverbPluginBypassingMasterReverb();
 			}
 
 			// Send the source audio to the master reverb
-			if (!AudioDevice->IsReverbPluginBypassingMasterReverb() && MixerDevice->GetMasterReverbSubmix().IsValid())
+			if (!bPluginBypassedMasterReverb && MixerDevice->GetMasterReverbSubmix().IsValid())
 			{
 				MixerSourceVoice->SetSubmixSendInfo(MixerDevice->GetMasterReverbSubmix(), ReverbSendLevel);
 			}
diff --git a/Engine/Source/Runtime/Engine/Private/AudioDevice.cpp b/Engine/Source/Runtime/Engine/Private/AudioDevice.cpp
--- a/Engine/Source/Runtime/Engine/Private/AudioDevice.cpp
+++ b/Engine/Source/Runtime/Engine/Private/AudioDevice.cpp
@@ -245,7 +245,7 @@ FAudioDevice::FAudioDevice()
 	, bSpatializationInterfaceEnabled(false)
 	, bOcclusionInterfaceEnabled(false)
 	, bReverbInterfaceEnabled(false)
-	, bReverbPluginBypassesMasterReverb(true)
+	, bReverbPluginBypassesMasterReverb(false)
 	, bModulationInterfaceEnabled(false)
 	, bPluginListenersInitialized(false)
 	, bHRTFEnabledForAll(false)
diff --git a/Engine/Source/Runtime/Engine/Public/AudioDevice.h b/Engine/Source/Runtime/Engine/Public/AudioDevice.h
--- a/Engine/Source/Runtime/Engine/Public/AudioDevice.h
+++ b/Engine/Source/Runtime/Engine/Public/AudioDevice.h
@@ -1100,7 +1100,7 @@ public:
 	/** Whether or not the reverb plugin is bypassing the master reverb. */
 	bool IsReverbPluginBypassingMasterReverb() const
 	{
-		return bReverbPluginBypassesMasterReverb;
+		return IsReverbPluginEnabled() && bReverbPluginBypassesMasterReverb;
 	}
 
 	static bool IsReverbPluginLoaded()
