Fix for input gain not actually being applied to anything but peak detector, thus causing the output of the dynamics processor to be quieter when input gain was increased.

diff --git a/Engine/Source/Runtime/AudioMixer/Private/DSP/DynamicsProcessor.cpp b/Engine/Source/Runtime/AudioMixer/Private/DSP/DynamicsProcessor.cpp
--- a/Engine/Source/Runtime/AudioMixer/Private/DSP/DynamicsProcessor.cpp
+++ b/Engine/Source/Runtime/AudioMixer/Private/DSP/DynamicsProcessor.cpp
@@ -174,7 +174,7 @@ namespace Audio
 			float LookaheadOutput = LookaheadDelay[Channel].ProcessAudioSample(InFrame[Channel]);
 
 			// Write into the output with the computed gain value
-			OutFrame[Channel] = Gain[Channel] * LookaheadOutput * OutputGain;
+			OutFrame[Channel] = Gain[Channel] * LookaheadOutput * OutputGain * InputGain;
 		}
 	}
 
