Added logic to one pole filters to avoid interpolation on first render call.

diff --git a/Engine/Source/Runtime/AudioMixer/Private/DSP/InterpolatedOnePole.cpp b/Engine/Source/Runtime/AudioMixer/Private/DSP/InterpolatedOnePole.cpp
--- a/Engine/Source/Runtime/AudioMixer/Private/DSP/InterpolatedOnePole.cpp
+++ b/Engine/Source/Runtime/AudioMixer/Private/DSP/InterpolatedOnePole.cpp
@@ -27,15 +27,21 @@ namespace Audio
 	{
 		CurrInterpLength = InterpLength;
 
+		if(isFirstFrequencyChange)
+		{
+			CurrInterpLength = 0;
+			isFirstFrequencyChange = false;
+		}
+
 		if (!FMath::IsNearlyEqual(InTargetFrequency, CutoffFrequency))
 		{
 			CutoffFrequency = InTargetFrequency;
 			float NormalizedFreq = FMath::Clamp(2.0f * InTargetFrequency / SampleRate, 0.0f, 1.0f);
 			B1Target = FMath::Exp(-PI * NormalizedFreq);
-			B1Delta = (B1Target - B1Curr) / static_cast<float>(InterpLength);
+			B1Delta = (B1Target - B1Curr) / static_cast<float>(CurrInterpLength);
 		}
 
-		if (InterpLength <= 1)
+		if (CurrInterpLength <= 1)
 		{
 			StopFrequencyInterpolation();
 		}
@@ -94,6 +100,7 @@ namespace Audio
 		CurrInterpLength = 0;
 		ClearMemory();
 		Z1Data = Z1.GetData();
+		isFirstFrequencyChange = true;
 	}
 
 	void FInterpolatedLPF::ClearMemory()
@@ -125,6 +132,12 @@ namespace Audio
 	{
 		CurrInterpLength = InterpLength;
 
+		if (isFirstFrequencyChange)
+		{
+			CurrInterpLength = 0;
+			isFirstFrequencyChange = false;
+		}
+
 		if (!FMath::IsNearlyEqual(InTargetFrequency, CutoffFrequency))
 		{
 			CutoffFrequency = FMath::Min(InTargetFrequency, NyquistLimit);
@@ -132,10 +145,10 @@ namespace Audio
 			const float G = GetGCoefficient();
 			A0Target = G / (1.0f + G);
 
-			A0Delta = (A0Target - A0Curr) / static_cast<float>(InterpLength);
+			A0Delta = (A0Target - A0Curr) / static_cast<float>(CurrInterpLength);
 		}
 
-		if (InterpLength <= 1)
+		if (CurrInterpLength <= 1)
 		{
 			StopFrequencyInterpolation();
 		}
@@ -184,6 +197,7 @@ namespace Audio
 		CurrInterpLength = 0;
 		ClearMemory();
 		Z1Data = Z1.GetData();
+		isFirstFrequencyChange = true;
 	}
 
 	void FInterpolatedHPF::ClearMemory()
diff --git a/Engine/Source/Runtime/AudioMixer/Public/DSP/InterpolatedOnePole.h b/Engine/Source/Runtime/AudioMixer/Public/DSP/InterpolatedOnePole.h
--- a/Engine/Source/Runtime/AudioMixer/Public/DSP/InterpolatedOnePole.h
+++ b/Engine/Source/Runtime/AudioMixer/Public/DSP/InterpolatedOnePole.h
@@ -66,6 +66,7 @@ namespace Audio
 		int32 NumInterpSteps;
 		float SampleRate{ 0 };
 		int32 NumChannels{ 1 };
+		int32 isFirstFrequencyChange : 1;
 
 	}; // class Interpolated Low pass filter
 
@@ -139,5 +140,7 @@ namespace Audio
 		float SampleRate{ 0 };
 		float NyquistLimit{ 0.0f };
 		int32 NumChannels{ 1 };
+		int32 isFirstFrequencyChange : 1;
+	
 	}; // class Interpolated High pass filter
 } // namespace Audio
