Branch code that modifies USoundWave causing sound to not reload, as this is not desired for seekable streams, as its possible to miss the first chunk and that's ok, as it should then attempt again later.

diff --git a/Engine/Source/Runtime/AudioMixer/Private/AudioMixerBuffer.cpp b/Engine/Source/Runtime/AudioMixer/Private/AudioMixerBuffer.cpp
--- a/Engine/Source/Runtime/AudioMixer/Private/AudioMixerBuffer.cpp
+++ b/Engine/Source/Runtime/AudioMixer/Private/AudioMixerBuffer.cpp
@@ -94,7 +94,7 @@ namespace Audio
 
 			case EBufferType::PCMRealTime:
 				return (DecompressionState ? DecompressionState->GetSourceBufferSize() : 0) + (MONO_PCM_BUFFER_SIZE * NumChannels);
-			
+
 			case EBufferType::Streaming:
 				return MONO_PCM_BUFFER_SIZE * NumChannels;
 
@@ -305,7 +305,7 @@ namespace Audio
 	FMixerBuffer* FMixerBuffer::CreateStreamingBuffer(FAudioDevice* AudioDevice, USoundWave* InWave)
 	{
 		FMixerBuffer* Buffer = new FMixerBuffer(AudioDevice, InWave, EBufferType::Streaming);
-		
+
 		FSoundQualityInfo QualityInfo = { 0 };
 
 		Buffer->DecompressionState = AudioDevice->CreateCompressedAudioInfo(InWave);
@@ -327,10 +327,18 @@ namespace Audio
 		}
 		else
 		{
-			InWave->DecompressionType = DTYPE_Invalid;
-			InWave->NumChannels = 0;
+			// When set to seekable streaming, missing the first chunk is possible and
+			// does not signify any issue with the asset itself, so don't mark it as invalid.
+			if (InWave && !InWave->IsSeekableStreaming())
+			{
+				UE_LOG(LogAudioMixer, Warning ,
+					TEXT("FMixerBuffer::CreateStreamingBuffer failed to StreamCompressedInfo on SoundWave '%s'.  Invalidating wave resource data (asset now requires re-cook)."),
+					*InWave->GetName());
 
-			InWave->RemoveAudioResource(); 
+				InWave->DecompressionType = DTYPE_Invalid;
+				InWave->NumChannels = 0;
+				InWave->RemoveAudioResource();
+			}
 
 			delete Buffer;
 			Buffer = nullptr;
