Linux: fixed editor launch with no audio device

diff --git a/Engine/Source/Runtime/AudioMixer/Private/AudioMixerDevice.cpp b/Engine/Source/Runtime/AudioMixer/Private/AudioMixerDevice.cpp
--- a/Engine/Source/Runtime/AudioMixer/Private/AudioMixerDevice.cpp
+++ b/Engine/Source/Runtime/AudioMixer/Private/AudioMixerDevice.cpp
@@ -911,7 +911,21 @@ namespace Audio
 
 	void FMixerDevice::FlushAudioRenderingCommands()
 	{
-		SourceManager.FlushCommandQueue();
+		if (AudioMixerPlatform && AudioMixerPlatform->IsInitialized())
+		{
+			SourceManager.FlushCommandQueue();
+		}
+		else
+		{
+			// Pump the audio device's command queue
+			PumpCommandQueue();
+
+			// And also directly pump the source manager command queue
+			SourceManager.PumpCommandQueue();
+			SourceManager.PumpCommandQueue();
+
+			SourceManager.UpdatePendingReleaseData(true);
+		}
 	}
 
 	bool FMixerDevice::IsMasterSubmixType(USoundSubmix* InSubmix) const
