Bug fix pausability

- Pausing sound while stopping can cause a crash when updating the sound

diff --git a/Engine/Source/Runtime/AudioMixer/Private/AudioMixerSource.cpp b/Engine/Source/Runtime/AudioMixer/Private/AudioMixerSource.cpp
--- a/Engine/Source/Runtime/AudioMixer/Private/AudioMixerSource.cpp
+++ b/Engine/Source/Runtime/AudioMixer/Private/AudioMixerSource.cpp
@@ -460,6 +460,13 @@ namespace Audio
 			return;
 		}
 
+		// Don't restart the sound if it was stopping when we paused, just stop it.
+		if (Paused && (bIsStopping || bIsDone))
+		{
+			StopNow();
+			return;
+		}
+
 		// It's possible if Pause and Play are called while a sound is async initializing. In this case
 		// we'll just not actually play the source here. Instead we'll call play when the sound finishes loading.
 		if (MixerSourceVoice && InitializationState == EMixerSourceInitializationState::Initialized)
