Fix for audio device shutdown

-- Actually, need to make sure the flush code doesn't execute on platforms without threads. This will still fail on HTML5

diff --git a/Engine/Source/Runtime/AudioMixer/Private/AudioMixerDevice.cpp b/Engine/Source/Runtime/AudioMixer/Private/AudioMixerDevice.cpp
--- a/Engine/Source/Runtime/AudioMixer/Private/AudioMixerDevice.cpp
+++ b/Engine/Source/Runtime/AudioMixer/Private/AudioMixerDevice.cpp
@@ -911,7 +911,7 @@ namespace Audio
 
 	void FMixerDevice::FlushAudioRenderingCommands()
 	{
-		if (AudioMixerPlatform && AudioMixerPlatform->IsInitialized())
+		if (AudioMixerPlatform && AudioMixerPlatform->IsInitialized() && FPlatformProcess::SupportsMultithreading())
 		{
 			SourceManager.FlushCommandQueue();
 		}
