Fix ClangEditor CIS warning.

Change MSVC_PRAGMA to not be defined when compiling using Clang on Windows. Also add a link to bad codegen report on Microsoft's developer forum for an instance of optimization being disabled.

diff --git a/Engine/Source/Runtime/AudioMixer/Private/AudioMixerSubmix.cpp b/Engine/Source/Runtime/AudioMixer/Private/AudioMixerSubmix.cpp
--- a/Engine/Source/Runtime/AudioMixer/Private/AudioMixerSubmix.cpp
+++ b/Engine/Source/Runtime/AudioMixer/Private/AudioMixerSubmix.cpp
@@ -137,7 +137,7 @@ namespace Audio
 		});
 	}
 
-// MSVC 2017 15.8.4 is generating bad code where setting submix will crash
+// MSVC 2017 15.8.4 is generating bad code where setting submix will crash. See https://developercommunity.visualstudio.com/content/problem/345511/bad-code-generation-in-vs-2017-v1585.html.
 MSVC_PRAGMA(optimize("", off))
 	void FMixerSubmix::AddChildSubmix(TWeakPtr<FMixerSubmix, ESPMode::ThreadSafe> SubmixWeakPtr)
 	{
diff --git a/Engine/Source/Runtime/Core/Public/Windows/WIndowsPlatform.h b/Engine/Source/Runtime/Core/Public/Windows/WIndowsPlatform.h
--- a/Engine/Source/Runtime/Core/Public/Windows/WIndowsPlatform.h
+++ b/Engine/Source/Runtime/Core/Public/Windows/WIndowsPlatform.h
@@ -132,7 +132,9 @@ typedef FWindowsPlatformTypes FPlatformTypes;
 #endif
 
 // Pragmas
-#define MSVC_PRAGMA(Pragma) __pragma(Pragma)
+#if !defined(__clang__)
+	#define MSVC_PRAGMA(Pragma) __pragma(Pragma)
+#endif
 
 // Prefetch
 #define PLATFORM_CACHE_LINE_SIZE	128
