Add additional info to comment around optimization being disabled due to codegen bug.

diff --git a/Engine/Source/Runtime/AudioMixer/Private/AudioMixerSourceVoice.cpp b/Engine/Source/Runtime/AudioMixer/Private/AudioMixerSourceVoice.cpp
--- a/Engine/Source/Runtime/AudioMixer/Private/AudioMixerSourceVoice.cpp
+++ b/Engine/Source/Runtime/AudioMixer/Private/AudioMixerSourceVoice.cpp
@@ -246,7 +246,7 @@ namespace Audio
 		return SourceManager->MixOutputBuffers(SourceId, InSubmixChannelType, SendLevel, OutWetBuffer);
 	}
 
-// MSVC 2017 15.8.4 is generating bad code where setting submix will crash
+// MSVC 2017 15.8.4 is generating bad code where setting submix will crash. See https://developercommunity.visualstudio.com/content/problem/345511/bad-code-generation-in-vs-2017-v1585.html.
 MSVC_PRAGMA(optimize("", off))
 	void FMixerSourceVoice::SetSubmixSendInfo(FMixerSubmixWeakPtr Submix, const float SendLevel)
 	{
