Another fix for lock up on starting without audio devices (UE-64506).

diff --git a/Engine/Source/Runtime/AudioMixer/Private/AudioMixerDevice.cpp b/Engine/Source/Runtime/AudioMixer/Private/AudioMixerDevice.cpp
--- a/Engine/Source/Runtime/AudioMixer/Private/AudioMixerDevice.cpp
+++ b/Engine/Source/Runtime/AudioMixer/Private/AudioMixerDevice.cpp
@@ -911,7 +911,7 @@ namespace Audio
 
 	void FMixerDevice::FlushAudioRenderingCommands()
 	{
-		if (AudioMixerPlatform && AudioMixerPlatform->IsInitialized() && FPlatformProcess::SupportsMultithreading())
+		if (IsInitialized() && FPlatformProcess::SupportsMultithreading())
 		{
 			SourceManager.FlushCommandQueue();
 		}
diff --git a/Engine/Source/Runtime/Engine/Public/AudioDevice.h b/Engine/Source/Runtime/Engine/Public/AudioDevice.h
--- a/Engine/Source/Runtime/Engine/Public/AudioDevice.h
+++ b/Engine/Source/Runtime/Engine/Public/AudioDevice.h
@@ -1699,6 +1699,9 @@ protected:
 
 	float DeviceDeltaTime;
 
+	/** Whether the device was initialized. */
+	FORCEINLINE bool IsInitialized() const { return bIsInitialized; }
+
 private:
 
 	/** Whether the value in HighestPriorityActivatedReverb should be used - Audio Thread owned */
