Merge fix for #jira UE-63510 to 4.21. #rb none.

diff --git a/Engine/Source/Runtime/AudioMixer/Private/Effects/AudioMixerSubmixEffectReverb.cpp b/Engine/Source/Runtime/AudioMixer/Private/Effects/AudioMixerSubmixEffectReverb.cpp
--- a/Engine/Source/Runtime/AudioMixer/Private/Effects/AudioMixerSubmixEffectReverb.cpp
+++ b/Engine/Source/Runtime/AudioMixer/Private/Effects/AudioMixerSubmixEffectReverb.cpp
@@ -98,6 +98,7 @@ void FSubmixEffectReverb::OnProcessAudio(const FSoundEffectSubmixInputData& InDa
 			PlateReverb.ProcessAudioFrame(&AudioData[SampleIndex], InData.NumChannels, &OutAudioData[SampleIndex], OutData.NumChannels);
 
 			OutAudioData[SampleIndex] += DryLevel * AudioData[SampleIndex];
+			OutAudioData[SampleIndex + 1] += DryLevel * AudioData[SampleIndex + 1];
 		}
 	}
 	// 5.1 or higher surround sound. Map stereo output to quad output
@@ -107,14 +108,15 @@ void FSubmixEffectReverb::OnProcessAudio(const FSoundEffectSubmixInputData& InDa
 		{
 			// Processed downmixed audio frame
 			PlateReverb.ProcessAudioFrame(&AudioData[InSampleIndex], InData.NumChannels, &OutAudioData[OutSampleIndex], InData.NumChannels);
-
-			OutAudioData[OutSampleIndex] += DryLevel * AudioData[InSampleIndex];
-
 			// Now do a cross-over to the back-left/back-right speakers from the front-left and front-right
 			
 			// Using standard speaker map order map the right output to the BackLeft channel
 			OutAudioData[OutSampleIndex + EAudioMixerChannel::BackRight] = OutAudioData[OutSampleIndex + EAudioMixerChannel::FrontLeft];
 			OutAudioData[OutSampleIndex + EAudioMixerChannel::BackLeft] = OutAudioData[OutSampleIndex + EAudioMixerChannel::FrontRight];
+
+			// Copy dry output to output data to stereo fronts
+			OutAudioData[OutSampleIndex] += DryLevel * AudioData[InSampleIndex];
+			OutAudioData[OutSampleIndex + 1] += DryLevel * AudioData[InSampleIndex + 1];
 		}
 	}
 }
