[NO NON-PII COMMIT MESSAGE]

diff --git a/Engine/Source/Runtime/AudioMixer/Private/AudioMixerSource.cpp b/Engine/Source/Runtime/AudioMixer/Private/AudioMixerSource.cpp
--- a/Engine/Source/Runtime/AudioMixer/Private/AudioMixerSource.cpp
+++ b/Engine/Source/Runtime/AudioMixer/Private/AudioMixerSource.cpp
@@ -170,7 +170,7 @@ namespace Audio
 			{
 				// If we're spatializing using HRTF and its an external send, don't need to setup a default/base submix send to master or EQ submix
 				// We'll only be using non-default submix sends (e.g. reverb).
-				if (!(WaveInstance->SpatializationMethod == ESoundSpatializationAlgorithm::SPATIALIZATION_HRTF && MixerDevice->bSpatializationIsExternalSend))
+				if (!(InitParams.bUseHRTFSpatialization && MixerDevice->bSpatializationIsExternalSend))
 				{
 					// If this sound is an ambisonics file, we preempt the normal base submix routing and only send to master ambisonics submix
 					if (WaveInstance->bIsAmbisonics)
