[NO NON-PII COMMIT MESSAGE]

diff --git a/Engine/Source/Runtime/AudioMixer/Private/AudioMixerSubmix.cpp b/Engine/Source/Runtime/AudioMixer/Private/AudioMixerSubmix.cpp
--- a/Engine/Source/Runtime/AudioMixer/Private/AudioMixerSubmix.cpp
+++ b/Engine/Source/Runtime/AudioMixer/Private/AudioMixerSubmix.cpp
@@ -622,6 +622,12 @@ MSVC_PRAGMA(optimize("", on))
 			{
 				FSoundEffectSubmix* SubmixEffect = SubmixEffectInfo.EffectInstance;
 
+				// SubmixEffectInfo.EffectInstance will be null if FMixerSubmix::RemoveSoundEffectSubmix was called earlier.
+				if (!SubmixEffect)
+				{
+					continue;
+				}
+
 				// Reset the output scratch buffer
 				ScratchBuffer.Reset(NumSamples);
 				ScratchBuffer.AddZeroed(NumSamples);
