Fix for missing audio on level reload.

diff --git a/Engine/Source/Runtime/AudioMixer/Private/AudioMixerSourceManager.cpp b/Engine/Source/Runtime/AudioMixer/Private/AudioMixerSourceManager.cpp
--- a/Engine/Source/Runtime/AudioMixer/Private/AudioMixerSourceManager.cpp
+++ b/Engine/Source/Runtime/AudioMixer/Private/AudioMixerSourceManager.cpp
@@ -107,6 +107,8 @@ namespace Audio
 		, bInitialized(false)
 		, bUsingSpatializationPlugin(false)
 	{
+		CommandsProcessedEvent = FPlatformProcess::GetSynchEventFromPool();
+		check(CommandsProcessedEvent != nullptr);
 	}
 
 	FMixerSourceManager::~FMixerSourceManager()
@@ -121,6 +123,8 @@ namespace Audio
 
 			SourceWorkers.Reset();
 		}
+
+		FPlatformProcess::ReturnSynchEventToPool(CommandsProcessedEvent);
 	}
 
 	void FMixerSourceManager::Init(const FSourceManagerInitParams& InitParams)
@@ -2134,16 +2138,23 @@ namespace Audio
 		Commands.SourceCommandQueue.Reset();
 
 		RenderThreadCommandBufferIndex.Set(!CurrentRenderThreadIndex);
+
+		check(CommandsProcessedEvent != nullptr);
+		CommandsProcessedEvent->Trigger();
 	}
 
 	void FMixerSourceManager::FlushCommandQueue()
 	{
-		bPumpQueue = true;
-		while (bPumpQueue)
-		{
-			FPlatformProcess::Sleep(0);
-		}
-		bPumpQueue = true;
+		check(CommandsProcessedEvent != nullptr);
+
+		// Make sure current current executing 
+		CommandsProcessedEvent->Wait();
+
+		// Call update to trigger a final pump of commands
+		Update();
+
+		// Wait one more time for the double pump
+		CommandsProcessedEvent->Wait();
 	}
 
 	void FMixerSourceManager::UpdatePendingReleaseData(bool bForceWait)
diff --git a/Engine/Source/Runtime/AudioMixer/Private/AudioMixerSourceManager.h b/Engine/Source/Runtime/AudioMixer/Private/AudioMixerSourceManager.h
--- a/Engine/Source/Runtime/AudioMixer/Private/AudioMixerSourceManager.h
+++ b/Engine/Source/Runtime/AudioMixer/Private/AudioMixerSourceManager.h
@@ -342,6 +342,8 @@ namespace Audio
 		FThreadSafeCounter AudioThreadCommandBufferIndex;
 		FThreadSafeCounter RenderThreadCommandBufferIndex;
 
+		FEvent* CommandsProcessedEvent;
+
 		TArray<int32> DebugSoloSources;
 
 		struct FSubmixChannelTypeInfo
