[NO NON-PII COMMIT MESSAGE]

diff --git a/Engine/Source/Runtime/AudioMixer/Private/AudioMixerBlueprintLibrary.cpp b/Engine/Source/Runtime/AudioMixer/Private/AudioMixerBlueprintLibrary.cpp
--- a/Engine/Source/Runtime/AudioMixer/Private/AudioMixerBlueprintLibrary.cpp
+++ b/Engine/Source/Runtime/AudioMixer/Private/AudioMixerBlueprintLibrary.cpp
@@ -42,6 +42,7 @@ void UAudioMixerBlueprintLibrary::AddMasterSubmixEffect(const UObject* WorldCont
 {
 	if (!SubmixEffectPreset)
 	{
+		UE_LOG(LogAudioMixer, Warning, TEXT("AddMasterSubmixEffect was passed invalid submix effect preset"));
 		return;
 	}
 
@@ -72,6 +73,7 @@ void UAudioMixerBlueprintLibrary::RemoveMasterSubmixEffect(const UObject* WorldC
 {
 	if (!SubmixEffectPreset)
 	{
+		UE_LOG(LogAudioMixer, Warning, TEXT("RemoveMasterSubmixEffect was passed invalid submix effect preset"));
 		return;
 	}
 
@@ -92,7 +94,6 @@ void UAudioMixerBlueprintLibrary::ClearMasterSubmixEffects(const UObject* WorldC
 	}
 }
 
-
 void UAudioMixerBlueprintLibrary::StartRecordingOutput(const UObject* WorldContextObject, float ExpectedDuration, USoundSubmix* SubmixToRecord)
 {
 	if (Audio::FMixerDevice* MixerDevice = GetAudioMixerDeviceFromWorldContext(WorldContextObject))
@@ -177,11 +178,11 @@ void UAudioMixerBlueprintLibrary::PauseRecordingOutput(const UObject* WorldConte
 	}
 }
 
-void UAudioMixerBlueprintLibrary::ResumeRecordingOutput(const UObject* WorldContextObject, USoundSubmix* SubmixToPause /* = nullptr */)
+void UAudioMixerBlueprintLibrary::ResumeRecordingOutput(const UObject* WorldContextObject, USoundSubmix* SubmixToResume /* = nullptr */)
 {
 	if (Audio::FMixerDevice* MixerDevice = GetAudioMixerDeviceFromWorldContext(WorldContextObject))
 	{
-		MixerDevice->ResumeRecording(SubmixToPause);
+		MixerDevice->ResumeRecording(SubmixToResume);
 	}
 	else
 	{
@@ -191,6 +192,12 @@ void UAudioMixerBlueprintLibrary::ResumeRecordingOutput(const UObject* WorldCont
 
 void UAudioMixerBlueprintLibrary::AddSourceEffectToPresetChain(const UObject* WorldContextObject, USoundEffectSourcePresetChain* PresetChain, FSourceEffectChainEntry Entry)
 {
+	if (!PresetChain)
+	{
+		UE_LOG(LogAudioMixer, Warning, TEXT("AddSourceEffectToPresetChain was passed invalid preset chain"));
+		return;
+	}
+
 	if (Audio::FMixerDevice* MixerDevice = GetAudioMixerDeviceFromWorldContext(WorldContextObject))
 	{
 		TArray<FSourceEffectChainEntry> Chain;
@@ -209,6 +216,12 @@ void UAudioMixerBlueprintLibrary::AddSourceEffectToPresetChain(const UObject* Wo
 
 void UAudioMixerBlueprintLibrary::RemoveSourceEffectFromPresetChain(const UObject* WorldContextObject, USoundEffectSourcePresetChain* PresetChain, int32 EntryIndex)
 {
+	if (!PresetChain)
+	{
+		UE_LOG(LogAudioMixer, Warning, TEXT("RemoveSourceEffectFromPresetChain was passed invalid preset chain"));
+		return;
+	}
+
 	if (Audio::FMixerDevice* MixerDevice = GetAudioMixerDeviceFromWorldContext(WorldContextObject))
 	{
 		TArray<FSourceEffectChainEntry> Chain;
@@ -232,6 +245,12 @@ void UAudioMixerBlueprintLibrary::RemoveSourceEffectFromPresetChain(const UObjec
 
 void UAudioMixerBlueprintLibrary::SetBypassSourceEffectChainEntry(const UObject* WorldContextObject, USoundEffectSourcePresetChain* PresetChain, int32 EntryIndex, bool bBypassed)
 {
+	if (!PresetChain)
+	{
+		UE_LOG(LogAudioMixer, Warning, TEXT("SetBypassSourceEffectChainEntry was passed invalid preset chain"));
+		return;
+	}
+
 	if (Audio::FMixerDevice* MixerDevice = GetAudioMixerDeviceFromWorldContext(WorldContextObject))
 	{
 		TArray<FSourceEffectChainEntry> Chain;
@@ -254,6 +273,12 @@ void UAudioMixerBlueprintLibrary::SetBypassSourceEffectChainEntry(const UObject*
 
 int32 UAudioMixerBlueprintLibrary::GetNumberOfEntriesInSourceEffectChain(const UObject* WorldContextObject, USoundEffectSourcePresetChain* PresetChain)
 {
+	if (!PresetChain)
+	{
+		UE_LOG(LogAudioMixer, Warning, TEXT("GetNumberOfEntriesInSourceEffectChain was passed invalid preset chain"));
+		return 0;
+	}
+
 	if (Audio::FMixerDevice* MixerDevice = GetAudioMixerDeviceFromWorldContext(WorldContextObject))
 	{
 		TArray<FSourceEffectChainEntry> Chain;
