Fixing spatialization issue of synth components.

diff --git a/Engine/Source/Runtime/AudioMixer/Private/Components/SynthComponent.cpp b/Engine/Source/Runtime/AudioMixer/Private/Components/SynthComponent.cpp
--- a/Engine/Source/Runtime/AudioMixer/Private/Components/SynthComponent.cpp
+++ b/Engine/Source/Runtime/AudioMixer/Private/Components/SynthComponent.cpp
@@ -225,33 +225,36 @@ void USynthComponent::CreateAudioComponent()
 		// Create the audio component which will be used to play the procedural sound wave
 		AudioComponent = NewObject<UAudioComponent>(this);
 
-		if (AudioComponent)
-		{
-			AudioComponent->bAutoActivate = false;
-			AudioComponent->bStopWhenOwnerDestroyed = true;
-			AudioComponent->bShouldRemainActiveIfDropped = true;
-			AudioComponent->Mobility = EComponentMobility::Movable;
+	}
+
+	if (AudioComponent)
+	{
+		AudioComponent->bAutoActivate = false;
+		AudioComponent->bStopWhenOwnerDestroyed = true;
+		AudioComponent->bShouldRemainActiveIfDropped = true;
+		AudioComponent->Mobility = EComponentMobility::Movable;
 
 #if WITH_EDITORONLY_DATA
-			AudioComponent->bVisualizeComponent = false;
+		AudioComponent->bVisualizeComponent = false;
 #endif
-			if (AudioComponent->GetAttachParent() == nullptr && !AudioComponent->IsAttachedTo(this))
-			{
-				AudioComponent->SetupAttachment(this);
-			}
+		if (AudioComponent->GetAttachParent() == nullptr && !AudioComponent->IsAttachedTo(this))
+		{
+			AudioComponent->SetupAttachment(this);
+		}
 
-			AudioComponent->OnAudioSingleEnvelopeValueNative.AddUObject(this, &USynthComponent::OnAudioComponentEnvelopeValue);
+		AudioComponent->OnAudioSingleEnvelopeValueNative.AddUObject(this, &USynthComponent::OnAudioComponentEnvelopeValue);
 
-			// Set defaults to be the same as audio component defaults
-			AudioComponent->EnvelopeFollowerAttackTime = EnvelopeFollowerAttackTime;
-			AudioComponent->EnvelopeFollowerReleaseTime = EnvelopeFollowerReleaseTime;
-		}
+		// Set defaults to be the same as audio component defaults
+		AudioComponent->EnvelopeFollowerAttackTime = EnvelopeFollowerAttackTime;
+		AudioComponent->EnvelopeFollowerReleaseTime = EnvelopeFollowerReleaseTime;
 	}
 }
 
 
 void USynthComponent::OnRegister()
 {
+	CreateAudioComponent();
+
 	Super::OnRegister();
 }
 
