[NO NON-PII COMMIT MESSAGE]

diff --git a/Engine/Source/Runtime/AudioMixer/Private/AudioMixerDevice.cpp b/Engine/Source/Runtime/AudioMixer/Private/AudioMixerDevice.cpp
--- a/Engine/Source/Runtime/AudioMixer/Private/AudioMixerDevice.cpp
+++ b/Engine/Source/Runtime/AudioMixer/Private/AudioMixerDevice.cpp
@@ -201,7 +201,7 @@ namespace Audio
 				AudioClockDelta = (double)OpenStreamParams.NumFrames / OpenStreamParams.SampleRate;
 
 				FAudioPluginInitializationParams PluginInitializationParams;
-				PluginInitializationParams.NumSources = MaxChannels;
+				PluginInitializationParams.NumSources = SourceManagerInitParams.NumSources;
 				PluginInitializationParams.SampleRate = SampleRate;
 				PluginInitializationParams.BufferLength = OpenStreamParams.NumFrames;
 				PluginInitializationParams.AudioDevicePtr = this;
