Adding ability to customize the precache buffer size in engine ini files.

- Adding some new LLM categories to help track decoded precache buffer mem usage and fully decoded buffers.

More info:

Precache buffer sizes, on platforms that utilize the feature, can be a huge memory hog. By reducing their sizes, we can recover quite a bit of memory.
Tricky part of the change is accounting for the variable precache buffer size while still retaining the MONO_PCM_BUFFER_SIZE size for regular real-time decoding cases.
We do a "skip first buffer(s)" step when starting a real-time decoded sound so the smaller precache buffer size needs to be shared to the decoders.

diff --git a/Engine/Source/Editor/MovieSceneTools/Private/TrackEditors/AudioTrackEditor.cpp b/Engine/Source/Editor/MovieSceneTools/Private/TrackEditors/AudioTrackEditor.cpp
--- a/Engine/Source/Editor/MovieSceneTools/Private/TrackEditors/AudioTrackEditor.cpp
+++ b/Engine/Source/Editor/MovieSceneTools/Private/TrackEditors/AudioTrackEditor.cpp
@@ -306,7 +306,7 @@ void FAudioThumbnail::GenerateWaveformPreview(TArray<uint8>& OutData, TRange<flo
 
 			if ( SoundWave->InitAudioResource( AudioDevice->GetRuntimeFormat( SoundWave ) ) && (SoundWave->DecompressionType != DTYPE_RealTime || SoundWave->CachedRealtimeFirstBuffer == nullptr ) )
 			{
-				FAsyncAudioDecompress TempDecompress(SoundWave);
+				FAsyncAudioDecompress TempDecompress(SoundWave, 0);
 				TempDecompress.StartSynchronousTask();
 			}
 
diff --git a/Engine/Source/Runtime/Android/AndroidAudio/Private/AndroidAudioBuffer.cpp b/Engine/Source/Runtime/Android/AndroidAudio/Private/AndroidAudioBuffer.cpp
--- a/Engine/Source/Runtime/Android/AndroidAudio/Private/AndroidAudioBuffer.cpp
+++ b/Engine/Source/Runtime/Android/AndroidAudio/Private/AndroidAudioBuffer.cpp
@@ -292,10 +292,10 @@ FSLESSoundBuffer* FSLESSoundBuffer::Init(  FSLESAudioDevice* AudioDevice ,USound
  * @param bLooping			Whether to loop the sound seamlessly, or pad with zeroes
  * @return					Whether the sound looped or not
  */
-bool FSLESSoundBuffer::ReadCompressedData( uint8* Destination, bool bLooping )
+bool FSLESSoundBuffer::ReadCompressedData( uint8* Destination, int32 NumFramesToDecode, bool bLooping )
 {
-	ensure( DecompressionState);
-
+	ensure(DecompressionState);
+	
 	if(Format == SoundFormat_Streaming)
 	{
 		return(DecompressionState->StreamCompressedData(Destination, bLooping, DecompressionState->GetStreamBufferSize() * NumChannels));
diff --git a/Engine/Source/Runtime/Android/AndroidAudio/Private/AndroidAudioSource.cpp b/Engine/Source/Runtime/Android/AndroidAudio/Private/AndroidAudioSource.cpp
--- a/Engine/Source/Runtime/Android/AndroidAudio/Private/AndroidAudioSource.cpp
+++ b/Engine/Source/Runtime/Android/AndroidAudio/Private/AndroidAudioSource.cpp
@@ -240,11 +240,11 @@ bool FSLESSoundSource::ReadMorePCMData(const int32 BufferIndex, EDataReadMode Da
 	{
 		if (DataReadMode == EDataReadMode::Synchronous)
 		{
-			return SLESBuffer->ReadCompressedData(AudioBuffers[BufferIndex].AudioData, WaveInstance->LoopingMode != LOOP_Never);
+			return SLESBuffer->ReadCompressedData(AudioBuffers[BufferIndex].AudioData, MONO_PCM_BUFFER_SAMPLES, WaveInstance->LoopingMode != LOOP_Never);
 		}
 		else
 		{
-			RealtimeAsyncTask = new FAsyncRealtimeAudioTask(SLESBuffer, AudioBuffers[BufferIndex].AudioData, WaveInstance->LoopingMode != LOOP_Never, DataReadMode == EDataReadMode::AsynchronousSkipFirstFrame);
+			RealtimeAsyncTask = new FAsyncRealtimeAudioTask(SLESBuffer, AudioBuffers[BufferIndex].AudioData, WaveInstance->WaveData->NumPrecacheFrames, WaveInstance->LoopingMode != LOOP_Never, DataReadMode == EDataReadMode::AsynchronousSkipFirstFrame);
 			RealtimeAsyncTask->StartBackgroundTask();
 			return false;
 		}
diff --git a/Engine/Source/Runtime/Android/AndroidAudio/Public/AndroidAudioDevice.h b/Engine/Source/Runtime/Android/AndroidAudio/Public/AndroidAudioDevice.h
--- a/Engine/Source/Runtime/Android/AndroidAudio/Public/AndroidAudioDevice.h
+++ b/Engine/Source/Runtime/Android/AndroidAudio/Public/AndroidAudioDevice.h
@@ -109,7 +109,7 @@ public:
 	 * @param bLooping			Whether to loop the sound seamlessly, or pad with zeroes
 	 * @return					Whether the sound looped or not
 	 */
-	bool ReadCompressedData( uint8* Destination, bool bLooping );
+	virtual bool ReadCompressedData(uint8* Destination, int32 NumFramesToDecode, bool bLooping) override;
 
 	/**
 	 * Sets the point in time within the buffer to the specified time
diff --git a/Engine/Source/Runtime/AudioMixer/Private/AudioMixerBuffer.cpp b/Engine/Source/Runtime/AudioMixer/Private/AudioMixerBuffer.cpp
--- a/Engine/Source/Runtime/AudioMixer/Private/AudioMixerBuffer.cpp
+++ b/Engine/Source/Runtime/AudioMixer/Private/AudioMixerBuffer.cpp
@@ -153,11 +153,6 @@ namespace Audio
 		return false;
 	}
 
-	bool FMixerBuffer::ReadCompressedData(uint8* Destination, bool bLooping)
-	{
-		return ReadCompressedData(Destination, MONO_PCM_BUFFER_SAMPLES, bLooping);
-	}
-
 	bool FMixerBuffer::ReadCompressedData(uint8* Destination, int32 InNumFrames, bool bLooping)
 	{
 		if (!DecompressionState)
diff --git a/Engine/Source/Runtime/AudioMixer/Private/AudioMixerBuffer.h b/Engine/Source/Runtime/AudioMixer/Private/AudioMixerBuffer.h
--- a/Engine/Source/Runtime/AudioMixer/Private/AudioMixerBuffer.h
+++ b/Engine/Source/Runtime/AudioMixer/Private/AudioMixerBuffer.h
@@ -35,13 +35,10 @@ namespace Audio
 		int32 GetCurrentChunkOffset() const override;
 		bool IsRealTimeSourceReady() override;
 		bool ReadCompressedInfo(USoundWave* SoundWave) override;
-		bool ReadCompressedData(uint8* Destination, bool bLooping) override;
+		bool ReadCompressedData(uint8* Destination, int32 NumFrames, bool bLooping) override;
 		void Seek(const float SeekTime) override;
 		//~ End FSoundBuffer Interface
 
-		// Allows specifying the number of frames you want to read
-		bool ReadCompressedData(uint8* Destination, int32 NumFrames, bool bLooping);
-
 		static FMixerBuffer* Init(FAudioDevice* AudioDevice, USoundWave* InWave, bool bForceRealtime);
 		static FMixerBuffer* CreatePreviewBuffer(FAudioDevice* AudioDevice, USoundWave* InWave);
 		static FMixerBuffer* CreateProceduralBuffer(FAudioDevice* AudioDevice, USoundWave* InWave);
diff --git a/Engine/Source/Runtime/AudioMixer/Private/AudioMixerSourceBuffer.cpp b/Engine/Source/Runtime/AudioMixer/Private/AudioMixerSourceBuffer.cpp
--- a/Engine/Source/Runtime/AudioMixer/Private/AudioMixerSourceBuffer.cpp
+++ b/Engine/Source/Runtime/AudioMixer/Private/AudioMixerSourceBuffer.cpp
@@ -94,8 +94,8 @@ namespace Audio
 		{
 			SourceVoiceBuffers.Add(TSharedPtr<FMixerSourceVoiceBuffer>(new FMixerSourceVoiceBuffer()));
 
-			SourceVoiceBuffers[BufferIndex]->AudioData.AddZeroed(TotalSamples);
-			SourceVoiceBuffers[BufferIndex]->Samples = TotalSamples;
+			// Prepare the memory to fit the max number of samples
+			SourceVoiceBuffers[BufferIndex]->AudioData.Reset(TotalSamples);
 			SourceVoiceBuffers[BufferIndex]->bRealTimeBuffer = true;
 			SourceVoiceBuffers[BufferIndex]->LoopCount = 0;
 		}
@@ -177,6 +177,9 @@ namespace Audio
 		const uint32 NumSamplesPerBuffer = MONO_PCM_BUFFER_SAMPLES * MixerBuffer->NumChannels;
 		int16* RawPCMBufferDataPtr = (int16*)RawPCMDataBuffer.Data;
 
+		// Prepare the buffer for the PCM submission
+		SourceVoiceBuffers[0]->AudioData.Reset(NumSamplesPerBuffer);
+
 		RawPCMDataBuffer.GetNextBuffer(SourceVoiceBuffers[0].Get(), NumSamplesPerBuffer);
 
 		SubmitBuffer(SourceVoiceBuffers[0]);
@@ -195,12 +198,20 @@ namespace Audio
 		{
 			bPlayedCachedBuffer = true;
 
-			const uint32 NumSamples = MONO_PCM_BUFFER_SAMPLES * MixerBuffer->NumChannels;
-			const uint32 BufferSize = MONO_PCM_BUFFER_SIZE * MixerBuffer->NumChannels;
+			const int32 NumPrecacheFrames = SoundWave->NumPrecacheFrames;
+			const uint32 NumSamples = NumPrecacheFrames * MixerBuffer->NumChannels;
+			const uint32 BufferSize = NumSamples * sizeof(int16);
 
 			// Format convert the first cached buffers
 #if (PLATFORM_NUM_AUDIODECOMPRESSION_PRECACHE_BUFFERS == 2)
 			{
+				// Prepare the precache buffer memory
+				for (int32 i = 0; i < 2; ++i)
+				{
+					SourceVoiceBuffers[i]->AudioData.Reset();
+					SourceVoiceBuffers[i]->AudioData.AddUninitialized(NumSamples);
+				}
+
 				int16* CachedBufferPtr0 = (int16*)SoundWave->CachedRealtimeFirstBuffer;
 				int16* CachedBufferPtr1 = (int16*)(SoundWave->CachedRealtimeFirstBuffer + BufferSize);
 				float* AudioData0 = SourceVoiceBuffers[0]->AudioData.GetData();
@@ -223,6 +234,9 @@ namespace Audio
 			}
 #elif (PLATFORM_NUM_AUDIODECOMPRESSION_PRECACHE_BUFFERS == 1)
 			{
+				SourceVoiceBuffers[0]->AudioData.Reset();
+				SourceVoiceBuffers[0]->AudioData.AddUninitialized(NumSamples);
+
 				int16* CachedBufferPtr0 = (int16*)SoundWave->CachedRealtimeFirstBuffer;
 
 				float* AudioData0 = SourceVoiceBuffers[0]->AudioData.GetData();
@@ -250,10 +264,12 @@ namespace Audio
 
 	bool FMixerSourceBuffer::ReadMoreRealtimeData(const int32 BufferIndex, EBufferReadMode BufferReadMode)
 	{
+		const int32 MaxSamples = MONO_PCM_BUFFER_SAMPLES * MixerBuffer->NumChannels;
+		SourceVoiceBuffers[BufferIndex]->AudioData.Reset();
+		SourceVoiceBuffers[BufferIndex]->AudioData.AddUninitialized(MaxSamples);
+
 		if (SoundWave && SoundWave->bProcedural)
 		{
-			const int32 MaxSamples = MONO_PCM_BUFFER_SAMPLES * MixerBuffer->NumChannels;
-
 			FProceduralAudioTaskData NewTaskData;
 			NewTaskData.ProceduralSoundWave = SoundWave;
 			NewTaskData.AudioData = SourceVoiceBuffers[BufferIndex]->AudioData.GetData();
@@ -270,8 +286,7 @@ namespace Audio
 			check(RawPCMDataBuffer.Data != nullptr);
 
 			// Read the next raw PCM buffer into the source buffer index. This converts raw PCM to float.
-			const uint32 NumSamplesPerBuffer = MONO_PCM_BUFFER_SAMPLES * MixerBuffer->NumChannels;
-			return RawPCMDataBuffer.GetNextBuffer(SourceVoiceBuffers[BufferIndex].Get(), NumSamplesPerBuffer);
+			return RawPCMDataBuffer.GetNextBuffer(SourceVoiceBuffers[BufferIndex].Get(), MaxSamples);
 		}
 
 		FDecodeAudioTaskData NewTaskData;
@@ -280,6 +295,7 @@ namespace Audio
 		NewTaskData.bLoopingMode = LoopingMode != LOOP_Never;
 		NewTaskData.bSkipFirstBuffer = (BufferReadMode == EBufferReadMode::AsynchronousSkipFirstFrame);
 		NewTaskData.NumFramesToDecode = MONO_PCM_BUFFER_SAMPLES;
+		NewTaskData.NumPrecacheFrames = SoundWave->NumPrecacheFrames;
 
 		check(!AsyncRealtimeAudioTask);
 		AsyncRealtimeAudioTask = CreateAudioTask(NewTaskData);
@@ -332,7 +348,6 @@ namespace Audio
 				FDecodeAudioTaskResults TaskResult;
 				AsyncRealtimeAudioTask->GetResult(TaskResult);
 
-				SourceVoiceBuffers[CurrentBuffer]->Samples = MONO_PCM_BUFFER_SAMPLES * MixerBuffer->NumChannels;
 				bLooped = TaskResult.bLooped;
 			}
 			break;
@@ -342,7 +357,7 @@ namespace Audio
 				FProceduralAudioTaskResults TaskResult;
 				AsyncRealtimeAudioTask->GetResult(TaskResult);
 
-				SourceVoiceBuffers[CurrentBuffer]->Samples = TaskResult.NumSamplesWritten;
+				SourceVoiceBuffers[CurrentBuffer]->AudioData.SetNum(TaskResult.NumSamplesWritten);
 			}
 			break;
 			}
diff --git a/Engine/Source/Runtime/AudioMixer/Private/AudioMixerSourceDecode.cpp b/Engine/Source/Runtime/AudioMixer/Private/AudioMixerSourceDecode.cpp
--- a/Engine/Source/Runtime/AudioMixer/Private/AudioMixerSourceDecode.cpp
+++ b/Engine/Source/Runtime/AudioMixer/Private/AudioMixerSourceDecode.cpp
@@ -105,7 +105,7 @@ public:
 				{
 					for (int32 NumberOfBuffersToSkip = 0; NumberOfBuffersToSkip < PLATFORM_NUM_AUDIODECOMPRESSION_PRECACHE_BUFFERS; NumberOfBuffersToSkip++)
 					{
-						DecodeTaskData.MixerBuffer->ReadCompressedData(DecodeBuffer.GetData(), DecodeTaskData.NumFramesToDecode, DecodeTaskData.bLoopingMode);
+						DecodeTaskData.MixerBuffer->ReadCompressedData(DecodeBuffer.GetData(), DecodeTaskData.NumPrecacheFrames, DecodeTaskData.bLoopingMode);
 					}
 				}
 
diff --git a/Engine/Source/Runtime/AudioMixer/Private/AudioMixerSourceDecode.h b/Engine/Source/Runtime/AudioMixer/Private/AudioMixerSourceDecode.h
--- a/Engine/Source/Runtime/AudioMixer/Private/AudioMixerSourceDecode.h
+++ b/Engine/Source/Runtime/AudioMixer/Private/AudioMixerSourceDecode.h
@@ -42,6 +42,9 @@ namespace Audio
 		// A pointer to a buffer of audio which will be decoded to
 		float* AudioData;
 
+		// The number of frames which are precached
+		int32 NumPrecacheFrames;
+
 		// The number of frames to decode
 		int32 NumFramesToDecode;
 
@@ -54,6 +57,7 @@ namespace Audio
 		FDecodeAudioTaskData()
 			: MixerBuffer(nullptr)
 			, AudioData(nullptr)
+			, NumPrecacheFrames(0)
 			, NumFramesToDecode(0)
 			, bLoopingMode(false)
 			, bSkipFirstBuffer(false)
diff --git a/Engine/Source/Runtime/AudioMixer/Private/AudioMixerSourceManager.cpp b/Engine/Source/Runtime/AudioMixer/Private/AudioMixerSourceManager.cpp
--- a/Engine/Source/Runtime/AudioMixer/Private/AudioMixerSourceManager.cpp
+++ b/Engine/Source/Runtime/AudioMixer/Private/AudioMixerSourceManager.cpp
@@ -1153,7 +1153,7 @@ namespace Audio
 			if (SourceInfo.MixerSourceBuffer->GetNumBuffersQueued() > 0)
 			{
 				SourceInfo.CurrentPCMBuffer = SourceInfo.MixerSourceBuffer->GetNextBuffer();
-				SourceInfo.CurrentAudioChunkNumFrames = SourceInfo.CurrentPCMBuffer->Samples / NumChannels;
+				SourceInfo.CurrentAudioChunkNumFrames = SourceInfo.CurrentPCMBuffer->AudioData.Num() / NumChannels;
 
 				// Subtract the number of frames in the current buffer from our frame index.
 				// Note: if this is the first time we're playing, CurrentFrameIndex will be 0
@@ -2177,7 +2177,7 @@ namespace Audio
 					}
 
 					// Clear out the mixer source buffer
-					SourceInfo.MixerSourceBuffer = nullptr;
+					SourceInfo.MixerSourceBuffer.Reset();
 
 					// Set the sound to be done playing
 					// This will flag the sound to be released
diff --git a/Engine/Source/Runtime/AudioMixer/Private/AudioMixerSourceManager.h b/Engine/Source/Runtime/AudioMixer/Private/AudioMixerSourceManager.h
--- a/Engine/Source/Runtime/AudioMixer/Private/AudioMixerSourceManager.h
+++ b/Engine/Source/Runtime/AudioMixer/Private/AudioMixerSourceManager.h
@@ -33,9 +33,6 @@ namespace Audio
 		/** PCM float data. */
 		AlignedFloatBuffer AudioData;
 
-		/** The amount of samples of audio data in the float buffer array. */
-		uint32 Samples;
-
 		/** How many times this buffer will loop. */
 		int32 LoopCount;
 
diff --git a/Engine/Source/Runtime/AudioMixer/Private/AudioMixerSourceVoice.cpp b/Engine/Source/Runtime/AudioMixer/Private/AudioMixerSourceVoice.cpp
--- a/Engine/Source/Runtime/AudioMixer/Private/AudioMixerSourceVoice.cpp
+++ b/Engine/Source/Runtime/AudioMixer/Private/AudioMixerSourceVoice.cpp
@@ -1,4 +1,4 @@
-﻿// Copyright 1998-2019 Epic Games, Inc. All Rights Reserved.
+// Copyright 1998-2019 Epic Games, Inc. All Rights Reserved.
 
 #include "AudioMixerSourceVoice.h"
 #include "AudioMixerSource.h"
@@ -282,7 +282,9 @@ MSVC_PRAGMA(optimize("", on))
 	{
 		AUDIO_MIXER_CHECK_AUDIO_PLAT_THREAD(MixerDevice);
 
-		for (uint32 i = 0; i < OutMixerSourceBuffer->Samples; ++i)
+		check(OutMixerSourceBuffer->AudioData.Num() > 0);
+
+		for (int32 i = 0; i < OutMixerSourceBuffer->AudioData.Num(); ++i)
 		{
 			OutMixerSourceBuffer->AudioData[i] = 0.0f;
 		}
diff --git a/Engine/Source/Runtime/AudioMixer/Private/DSP/SoundWaveDecoder.cpp b/Engine/Source/Runtime/AudioMixer/Private/DSP/SoundWaveDecoder.cpp
--- a/Engine/Source/Runtime/AudioMixer/Private/DSP/SoundWaveDecoder.cpp
+++ b/Engine/Source/Runtime/AudioMixer/Private/DSP/SoundWaveDecoder.cpp
@@ -146,8 +146,9 @@ namespace Audio
 
 			if (MixerSourceBuffer.GetNumBuffersQueued() > 0)
 			{
+				check(SourceInfo.CurrentPCMBuffer->AudioData.Num() > 0);
 				SourceInfo.CurrentPCMBuffer = MixerSourceBuffer.GetNextBuffer();
-				SourceInfo.CurrentAudioChunkNumFrames = SourceInfo.CurrentPCMBuffer->Samples / SourceInfo.NumSourceChannels;
+				SourceInfo.CurrentAudioChunkNumFrames = SourceInfo.CurrentPCMBuffer->AudioData.Num() / SourceInfo.NumSourceChannels;
 
 				if (bReadCurrentFrame)
 				{
diff --git a/Engine/Source/Runtime/Core/Private/HAL/LowLevelMemTracker.cpp b/Engine/Source/Runtime/Core/Private/HAL/LowLevelMemTracker.cpp
--- a/Engine/Source/Runtime/Core/Private/HAL/LowLevelMemTracker.cpp
+++ b/Engine/Source/Runtime/Core/Private/HAL/LowLevelMemTracker.cpp
@@ -67,6 +67,10 @@ DECLARE_LLM_MEMORY_STAT(TEXT("Engine Misc"), STAT_EngineMiscLLM, STATGROUP_LLMFU
 DECLARE_LLM_MEMORY_STAT(TEXT("TaskGraph Misc Tasks"), STAT_TaskGraphTasksMiscLLM, STATGROUP_LLMFULL);
 DECLARE_LLM_MEMORY_STAT(TEXT("Audio"), STAT_AudioLLM, STATGROUP_LLMFULL);
 DECLARE_LLM_MEMORY_STAT(TEXT("AudioMixer"), STAT_AudioMixerLLM, STATGROUP_LLMFULL);
+DECLARE_LLM_MEMORY_STAT(TEXT("AudioPrecache"), STAT_AudioPrecacheLLM, STATGROUP_LLMFULL);
+DECLARE_LLM_MEMORY_STAT(TEXT("AudioDecompress"), STAT_AudioDecompressLLM, STATGROUP_LLMFULL);
+DECLARE_LLM_MEMORY_STAT(TEXT("AudioRealtimePrecache"), STAT_AudioRealtimePrecacheLLM, STATGROUP_LLMFULL);
+DECLARE_LLM_MEMORY_STAT(TEXT("AudioFullDecompress"), STAT_AudioFullDecompressLLM, STATGROUP_LLMFULL);
 DECLARE_LLM_MEMORY_STAT(TEXT("FName"), STAT_FNameLLM, STATGROUP_LLMFULL);
 DECLARE_LLM_MEMORY_STAT(TEXT("Networking"), STAT_NetworkingLLM, STATGROUP_LLMFULL);
 DECLARE_LLM_MEMORY_STAT(TEXT("Meshes"), STAT_MeshesLLM, STATGROUP_LLMFULL);
diff --git a/Engine/Source/Runtime/Core/Public/HAL/LowLevelMemTracker.h b/Engine/Source/Runtime/Core/Public/HAL/LowLevelMemTracker.h
--- a/Engine/Source/Runtime/Core/Public/HAL/LowLevelMemTracker.h
+++ b/Engine/Source/Runtime/Core/Public/HAL/LowLevelMemTracker.h
@@ -179,7 +179,11 @@ enum class ELLMTagSet : uint8
 	macro(EngineMisc,							"EngineMisc",					GET_STATFNAME(STAT_EngineMiscLLM),							GET_STATFNAME(STAT_EngineSummaryLLM))		\
 	macro(TaskGraphTasksMisc,					"TaskGraphMiscTasks",			GET_STATFNAME(STAT_TaskGraphTasksMiscLLM),					GET_STATFNAME(STAT_EngineSummaryLLM))		\
 	macro(Audio,								"Audio",						GET_STATFNAME(STAT_AudioLLM),								GET_STATFNAME(STAT_AudioSummaryLLM))		\
-	macro(AudioMixer,							"AudioMixer",					GET_STATFNAME(STAT_AudioMixerLLM),							GET_STATFNAME(STAT_AudioSummaryLLM))	\
+	macro(AudioMixer,							"AudioMixer",					GET_STATFNAME(STAT_AudioMixerLLM),							GET_STATFNAME(STAT_AudioSummaryLLM))		\
+	macro(AudioPrecache,						"AudioPrecache",				GET_STATFNAME(STAT_AudioPrecacheLLM),						GET_STATFNAME(STAT_AudioSummaryLLM))		\
+	macro(AudioDecompress,						"AudioDecompress",				GET_STATFNAME(STAT_AudioDecompressLLM),						GET_STATFNAME(STAT_AudioSummaryLLM))		\
+	macro(AudioRealtimePrecache,				"AudioRealtimePrecache",		GET_STATFNAME(STAT_AudioRealtimePrecacheLLM),				GET_STATFNAME(STAT_AudioSummaryLLM))		\
+	macro(AudioFullDecompress,					"AudioFullDecompress",			GET_STATFNAME(STAT_AudioFullDecompressLLM),					GET_STATFNAME(STAT_AudioSummaryLLM))		\
 	macro(FName,								"FName",						GET_STATFNAME(STAT_FNameLLM),								GET_STATFNAME(STAT_EngineSummaryLLM))		\
 	macro(Networking,							"Networking",					GET_STATFNAME(STAT_NetworkingLLM),							GET_STATFNAME(STAT_EngineSummaryLLM))		\
 	macro(Meshes,								"Meshes",						GET_STATFNAME(STAT_MeshesLLM),								GET_STATFNAME(STAT_MeshesSummaryLLM))		\
diff --git a/Engine/Source/Runtime/Engine/Classes/Sound/SoundWave.h b/Engine/Source/Runtime/Engine/Classes/Sound/SoundWave.h
--- a/Engine/Source/Runtime/Engine/Classes/Sound/SoundWave.h
+++ b/Engine/Source/Runtime/Engine/Classes/Sound/SoundWave.h
@@ -317,27 +317,30 @@ private:
 public:	
 	/** Async worker that decompresses the audio data on a different thread */
 	typedef FAsyncTask< class FAsyncAudioDecompressWorker > FAsyncAudioDecompress;	// Forward declare typedef
-	FAsyncAudioDecompress*		AudioDecompressor;
+	FAsyncAudioDecompress* AudioDecompressor;
 
 	/** Pointer to 16 bit PCM data - used to avoid synchronous operation to obtain first block of the realtime decompressed buffer */
-	uint8*						CachedRealtimeFirstBuffer;
+	uint8* CachedRealtimeFirstBuffer;
+	
+	/** The number of frames which have been precached for this sound wave. */
+	int32 NumPrecacheFrames;
 
 	/** Pointer to 16 bit PCM data - used to decompress data to and preview sounds */
-	uint8*						RawPCMData;
+	uint8* RawPCMData;
 
 	/** Size of RawPCMData, or what RawPCMData would be if the sound was fully decompressed */
-	int32						RawPCMDataSize;
+	int32 RawPCMDataSize;
 
 	/** Memory containing the data copied from the compressed bulk data */
-	uint8*						ResourceData;
+	uint8* ResourceData;
 
 	/** Uncompressed wav data 16 bit in mono or stereo - stereo not allowed for multichannel data */
-	FByteBulkData				RawData;
+	FByteBulkData RawData;
 
 	/** GUID used to uniquely identify this node so it can be found in the DDC */
-	FGuid						CompressedDataGuid;
+	FGuid CompressedDataGuid;
 
-	FFormatContainer			CompressedFormatData;
+	FFormatContainer CompressedFormatData;
 
 #if WITH_EDITORONLY_DATA
 	TMap<FName, uint32> AsyncLoadingDataFormats;
diff --git a/Engine/Source/Runtime/Engine/Private/AudioDecompress.cpp b/Engine/Source/Runtime/Engine/Private/AudioDecompress.cpp
--- a/Engine/Source/Runtime/Engine/Private/AudioDecompress.cpp
+++ b/Engine/Source/Runtime/Engine/Private/AudioDecompress.cpp
@@ -609,10 +609,12 @@ namespace ADPCM
 /**
  * Worker for decompression on a separate thread
  */
-FAsyncAudioDecompressWorker::FAsyncAudioDecompressWorker(USoundWave* InWave)
+FAsyncAudioDecompressWorker::FAsyncAudioDecompressWorker(USoundWave* InWave, int32 InPrecacheBufferNumFrames)
 	: Wave(InWave)
-	, AudioInfo(NULL)
+	, AudioInfo(nullptr)
+	, NumPrecacheFrames(InPrecacheBufferNumFrames)
 {
+	check(NumPrecacheFrames > 0);
 	if (GEngine && GEngine->GetMainAudioDevice())
 	{
 		AudioInfo = GEngine->GetMainAudioDevice()->CreateCompressedAudioInfo(Wave);
@@ -622,6 +624,7 @@ FAsyncAudioDecompressWorker::FAsyncAudioDecompressWorker(USoundWave* InWave)
 void FAsyncAudioDecompressWorker::DoWork()
 {
 	LLM_SCOPE(ELLMTag::Audio);
+	LLM_SCOPE(ELLMTag::AudioDecompress);
 
 	if (AudioInfo)
 	{
@@ -655,8 +658,10 @@ void FAsyncAudioDecompressWorker::DoWork()
 
 			if (Wave->DecompressionType == DTYPE_RealTime)
 			{
+				LLM_SCOPE(ELLMTag::AudioRealtimePrecache);
 #if PLATFORM_NUM_AUDIODECOMPRESSION_PRECACHE_BUFFERS > 0
-				const uint32 PCMBufferSize = MONO_PCM_BUFFER_SIZE * Wave->NumChannels * PLATFORM_NUM_AUDIODECOMPRESSION_PRECACHE_BUFFERS;
+				const uint32 PCMBufferSize = NumPrecacheFrames * sizeof(int16) * Wave->NumChannels * PLATFORM_NUM_AUDIODECOMPRESSION_PRECACHE_BUFFERS;
+				Wave->NumPrecacheFrames = NumPrecacheFrames;
 				if (Wave->CachedRealtimeFirstBuffer == nullptr)
 				{
 					Wave->CachedRealtimeFirstBuffer = (uint8*)FMemory::Malloc(PCMBufferSize);
@@ -670,6 +675,7 @@ void FAsyncAudioDecompressWorker::DoWork()
 			}
 			else
 			{
+				LLM_SCOPE(ELLMTag::AudioFullDecompress);
 				check(Wave->DecompressionType == DTYPE_Native || Wave->DecompressionType == DTYPE_Procedural);
 
 				Wave->RawPCMDataSize = QualityInfo.SampleDataSize;
diff --git a/Engine/Source/Runtime/Engine/Private/AudioDevice.cpp b/Engine/Source/Runtime/Engine/Private/AudioDevice.cpp
--- a/Engine/Source/Runtime/Engine/Private/AudioDevice.cpp
+++ b/Engine/Source/Runtime/Engine/Private/AudioDevice.cpp
@@ -89,6 +89,15 @@ FAutoConsoleVariableRef CVarWaitForSoundWaveToLoad(
 	TEXT("0: Attempt to play back, 1: Wait for load."),
 	ECVF_Default);
 
+static int32 NumPrecacheFramesCvar = 0;
+FAutoConsoleVariableRef CVarNumPrecacheFrames(
+	TEXT("au.NumPrecacheFrames"),
+	NumPrecacheFramesCvar,
+	TEXT("When set to > 0, will use that value as the number of frames to precache audio buffers with.\n")
+	TEXT("0: Use default value for precache frames, >0: Number of frames to precache."),
+	ECVF_Default);
+
+
 /*-----------------------------------------------------------------------------
 FDynamicParameter implementation.
 -----------------------------------------------------------------------------*/
@@ -153,6 +162,7 @@ FAudioDevice::FAudioDevice()
 	, NumStoppingVoices(32)
 	, MaxWaveInstances(0)
 	, SampleRate(0)
+	, NumPrecacheFrames(MONO_PCM_BUFFER_SAMPLES)
 	, CommonAudioPoolSize(0)
 	, CommonAudioPool(nullptr)
 	, CommonAudioPoolFreeBytes(0)
@@ -261,6 +271,12 @@ bool FAudioDevice::Init(int32 InMaxChannels)
 		PlatformAudioHeadroom = FMath::Pow(10.0f, Headroom / 20.0f);
 	}
 
+	int32 NumPrecacheFramesSettings = 0; 
+	if (GConfig->GetInt(TEXT("Audio"), TEXT("NumPrecacheFrames"), NumPrecacheFramesSettings, GEngineIni))
+	{
+		NumPrecacheFrames = FMath::Min(128, NumPrecacheFramesSettings);
+	}
+
 	bIsStoppingVoicesEnabled = !DisableStoppingVoicesCvar;
 
 	const UAudioSettings* AudioSettings = GetDefault<UAudioSettings>();
@@ -4927,6 +4943,7 @@ void FAudioDevice::Flush(UWorld* WorldToFlush, bool bClearActivatedReverb)
 void FAudioDevice::Precache(USoundWave* SoundWave, bool bSynchronous, bool bTrackMemory, bool bForceFullDecompression)
 {
 	LLM_SCOPE(ELLMTag::Audio);
+	LLM_SCOPE(ELLMTag::AudioPrecache);
 
 	if (SoundWave == nullptr)
 	{
@@ -5063,14 +5080,14 @@ void FAudioDevice::Precache(USoundWave* SoundWave, bool bSynchronous, bool bTrac
 			if (bSynchronous)
 			{
 				// Create a worker to decompress the vorbis data
-				FAsyncAudioDecompress TempDecompress(SoundWave);
+				FAsyncAudioDecompress TempDecompress(SoundWave, GetNumPrecacheFrames());
 				TempDecompress.StartSynchronousTask();
 			}
 			else
 			{
 				// This is the one case where precaching will not be done when this function exits
 				check(SoundWave->GetPrecacheState() == ESoundWavePrecacheState::InProgress);
-				SoundWave->AudioDecompressor = new FAsyncAudioDecompress(SoundWave);
+				SoundWave->AudioDecompressor = new FAsyncAudioDecompress(SoundWave, GetNumPrecacheFrames());
 				SoundWave->AudioDecompressor->StartBackgroundTask();
 				PrecachingSoundWaves.Add(SoundWave);
 			}
@@ -5534,4 +5551,15 @@ bool FAudioDevice::ShouldUseAttenuation(const UWorld* World) const
 	// - we are forcing the use of attenuation (e.g. for some editors)
 	const bool bIsInGameWorld = World ? World->IsGameWorld() : true;
 	return (bIsInGameWorld || bUseAttenuationForNonGameWorlds);
+}
+
+int32 FAudioDevice::GetNumPrecacheFrames() const
+{
+	// Check the cvar and use that if it's been set.
+	if (NumPrecacheFramesCvar > 0)
+	{
+		return NumPrecacheFramesCvar;
+	}
+	// Otherwise, use the default value or value set in ini file
+	return NumPrecacheFrames;
 }
\ No newline at end of file
diff --git a/Engine/Source/Runtime/Engine/Public/Audio.h b/Engine/Source/Runtime/Engine/Public/Audio.h
--- a/Engine/Source/Runtime/Engine/Public/Audio.h
+++ b/Engine/Source/Runtime/Engine/Public/Audio.h
@@ -561,7 +561,7 @@ public:
 	virtual bool ReadCompressedInfo(USoundWave* SoundWave) { return true; }
 
 	/** Reads the next compressed data chunk */
-	virtual bool ReadCompressedData(uint8* Destination, bool bLooping) { return true; }
+	virtual bool ReadCompressedData(uint8* Destination, int32 NumFramesToDecode, bool bLooping) { return true; }
 	
 	/** Seeks the buffer to the given seek time */
 	virtual void Seek(const float SeekTime) {}
diff --git a/Engine/Source/Runtime/Engine/Public/AudioDecompress.h b/Engine/Source/Runtime/Engine/Public/AudioDecompress.h
--- a/Engine/Source/Runtime/Engine/Public/AudioDecompress.h
+++ b/Engine/Source/Runtime/Engine/Public/AudioDecompress.h
@@ -271,9 +271,9 @@ protected:
 class FAsyncAudioDecompressWorker : public FNonAbandonableTask
 {
 protected:
-	class USoundWave*		Wave;
-
-	ICompressedAudioInfo*	AudioInfo;
+	USoundWave* Wave;
+	ICompressedAudioInfo* AudioInfo;
+	int32 NumPrecacheFrames;
 
 public:
 	/**
@@ -281,7 +281,7 @@ public:
 	 *
 	 * @param	InWave		Wave data to decompress
 	 */
-	ENGINE_API FAsyncAudioDecompressWorker(USoundWave* InWave);
+	ENGINE_API FAsyncAudioDecompressWorker(USoundWave* InWave, int32 InNumPrecacheFrames);
 
 	/**
 	 * Performs the async audio decompression
@@ -315,6 +315,7 @@ protected:
 	T* AudioBuffer;
 	USoundWave* WaveData;
 	uint8* AudioData;
+	int32 NumPrecacheFrames;
 	int32 MaxSamples;
 	int32 BytesWritten;
 	ERealtimeAudioTaskType TaskType;
@@ -327,6 +328,7 @@ public:
 		: AudioBuffer(InAudioBuffer)
 		, WaveData(InWaveData)
 		, AudioData(nullptr)
+		, NumPrecacheFrames(0)
 		, MaxSamples(0)
 		, BytesWritten(0)
 		, TaskType(ERealtimeAudioTaskType::CompressedInfo)
@@ -338,9 +340,10 @@ public:
 		check(WaveData);
 	}
 
-	FAsyncRealtimeAudioTaskWorker(T* InAudioBuffer, uint8* InAudioData, bool bInLoopingMode, bool bInSkipFirstBuffer)
+	FAsyncRealtimeAudioTaskWorker(T* InAudioBuffer, uint8* InAudioData, int32 InNumPrecacheFrames, bool bInLoopingMode, bool bInSkipFirstBuffer)
 		: AudioBuffer(InAudioBuffer)
 		, AudioData(InAudioData)
+		, NumPrecacheFrames(InNumPrecacheFrames)
 		, TaskType(ERealtimeAudioTaskType::Decompress)
 		, bSkipFirstBuffer(bInSkipFirstBuffer)
 		, bLoopingMode(bInLoopingMode)
@@ -353,6 +356,7 @@ public:
 	FAsyncRealtimeAudioTaskWorker(USoundWave* InWaveData, uint8* InAudioData, int32 InMaxSamples)
 		: WaveData(InWaveData)
 		, AudioData(InAudioData)
+		, NumPrecacheFrames(0)
 		, MaxSamples(InMaxSamples)
 		, BytesWritten(0)
 		, TaskType(ERealtimeAudioTaskType::Procedural)
@@ -376,18 +380,18 @@ public:
 			{
 #if PLATFORM_ANDROID
 				// Only skip one buffer on Android
-				AudioBuffer->ReadCompressedData( ( uint8* )AudioData, bLoopingMode );
+				AudioBuffer->ReadCompressedData((uint8*)AudioData, NumPrecacheFrames, bLoopingMode );
 #else
 				// If we're using cached data we need to skip the first two reads from the data
-				AudioBuffer->ReadCompressedData( ( uint8* )AudioData, bLoopingMode );
-				AudioBuffer->ReadCompressedData( ( uint8* )AudioData, bLoopingMode );
+				AudioBuffer->ReadCompressedData((uint8*)AudioData, NumPrecacheFrames, bLoopingMode);
+				AudioBuffer->ReadCompressedData((uint8*)AudioData, NumPrecacheFrames, bLoopingMode);
 #endif
 			}
-			bLooped = AudioBuffer->ReadCompressedData( ( uint8* )AudioData, bLoopingMode );
+			bLooped = AudioBuffer->ReadCompressedData((uint8*)AudioData, MONO_PCM_BUFFER_SAMPLES, bLoopingMode);
 			break;
 
 		case ERealtimeAudioTaskType::Procedural:
-			BytesWritten = WaveData->GeneratePCMData( (uint8*)AudioData, MaxSamples );
+			BytesWritten = WaveData->GeneratePCMData((uint8*)AudioData, MaxSamples);
 			break;
 
 		default:
@@ -436,9 +440,9 @@ public:
 		Task = new FAsyncTask<FAsyncRealtimeAudioTaskWorker<T>>(InAudioBuffer, InWaveData);
 	}
 
-	FAsyncRealtimeAudioTaskProxy(T* InAudioBuffer, uint8* InAudioData, bool bInLoopingMode, bool bInSkipFirstBuffer)
+	FAsyncRealtimeAudioTaskProxy(T* InAudioBuffer, uint8* InAudioData, int32 InNumFramesToDecode, bool bInLoopingMode, bool bInSkipFirstBuffer)
 	{
-		Task = new FAsyncTask<FAsyncRealtimeAudioTaskWorker<T>>(InAudioBuffer, InAudioData, bInLoopingMode, bInSkipFirstBuffer);
+		Task = new FAsyncTask<FAsyncRealtimeAudioTaskWorker<T>>(InAudioBuffer, InAudioData, InNumFramesToDecode, bInLoopingMode, bInSkipFirstBuffer);
 	}
 	FAsyncRealtimeAudioTaskProxy(USoundWave* InWaveData, uint8* InAudioData, int32 InMaxSamples)
 	{
diff --git a/Engine/Source/Runtime/Engine/Public/AudioDevice.h b/Engine/Source/Runtime/Engine/Public/AudioDevice.h
--- a/Engine/Source/Runtime/Engine/Public/AudioDevice.h
+++ b/Engine/Source/Runtime/Engine/Public/AudioDevice.h
@@ -1454,15 +1454,12 @@ private:
 	/** Stops oldest sound source. */
 	void StopOldestStoppingSource();
 
-	/** Processes any pending active sounds. */
-	void ProcessPendingNewActiveSounds();
-
-	/** Adds new active sound on the audio thread */
-	void AddNewActiveSoundInternal(FActiveSound* NewActiveSound);
-
 	/** Check whether we should use attenuation settings */
 	bool ShouldUseAttenuation(const UWorld* World) const;
 
+	/** Returns the number of frames to use per precache buffer. */
+	int32 GetNumPrecacheFrames() const;
+
 public:
 
 	/** Query if the editor is in VR Preview for the current play world. Returns false for non-editor builds */
@@ -1545,8 +1542,9 @@ public:
 	/** The platform specific audio settings. */
 	FAudioPlatformSettings PlatformSettings;
 
-	/** The length of output callback buffer */
-
+	/** The number of frames to precache. */
+	int32 NumPrecacheFrames;
+	
 	/** The amount of memory to reserve for always resident sounds */
 	int32 CommonAudioPoolSize;
 
diff --git a/Engine/Source/Runtime/IOS/IOSAudio/Private/IOSAudioBuffer.cpp b/Engine/Source/Runtime/IOS/IOSAudio/Private/IOSAudioBuffer.cpp
--- a/Engine/Source/Runtime/IOS/IOSAudio/Private/IOSAudioBuffer.cpp
+++ b/Engine/Source/Runtime/IOS/IOSAudio/Private/IOSAudioBuffer.cpp
@@ -123,7 +123,7 @@ bool FIOSAudioSoundBuffer::ReadCompressedInfo(USoundWave* InWave)
 	return DecompressionState->ReadCompressedInfo(InWave->ResourceData, InWave->ResourceSize, &QualityInfo);
 }
 
-bool FIOSAudioSoundBuffer::ReadCompressedData( uint8* Destination, bool bLooping )
+bool FIOSAudioSoundBuffer::ReadCompressedData( uint8* Destination, int32 NumFramesToDecode, bool bLooping )
 {
 	if(bStreaming)
 	{
diff --git a/Engine/Source/Runtime/IOS/IOSAudio/Private/IOSAudioSource.cpp b/Engine/Source/Runtime/IOS/IOSAudio/Private/IOSAudioSource.cpp
--- a/Engine/Source/Runtime/IOS/IOSAudio/Private/IOSAudioSource.cpp
+++ b/Engine/Source/Runtime/IOS/IOSAudio/Private/IOSAudioSource.cpp
@@ -477,6 +477,7 @@ OSStatus FIOSAudioSoundSource::IOSAudioRenderCallback(void* RefCon, AudioUnitRen
 		Source->bChannel0Finished =
 			Source->IOSBuffer->ReadCompressedData(
 				(uint8*)Source->IOSBuffer->SampleData,
+				MONO_PCM_BUFFER_SAMPLES,
 				Source->WaveInstance->LoopingMode == LOOP_WithNotification || Source->WaveInstance->LoopingMode == LOOP_Forever);
 	}
     
diff --git a/Engine/Source/Runtime/IOS/IOSAudio/Public/IOSAudioDevice.h b/Engine/Source/Runtime/IOS/IOSAudio/Public/IOSAudioDevice.h
--- a/Engine/Source/Runtime/IOS/IOSAudio/Public/IOSAudioDevice.h
+++ b/Engine/Source/Runtime/IOS/IOSAudio/Public/IOSAudioDevice.h
@@ -83,7 +83,7 @@ public:
 	 * @param bLooping			Whether to loop the sound seamlessly, or pad with zeroes
 	 * @return					Whether the sound looped or not
 	 */
-	virtual bool ReadCompressedData( uint8* Destination, bool bLooping ) override;
+	virtual bool ReadCompressedData( uint8* Destination, int32 NumFramesToDecode, bool bLooping ) override;
 	
 	
 	int32  RenderCallbackBufferSize;
diff --git a/Engine/Source/Runtime/Mac/CoreAudio/Private/CoreAudioBuffer.cpp b/Engine/Source/Runtime/Mac/CoreAudio/Private/CoreAudioBuffer.cpp
--- a/Engine/Source/Runtime/Mac/CoreAudio/Private/CoreAudioBuffer.cpp
+++ b/Engine/Source/Runtime/Mac/CoreAudio/Private/CoreAudioBuffer.cpp
@@ -115,14 +115,7 @@ void FCoreAudioSoundBuffer::InitAudioStreamBasicDescription( UInt32 FormatID, US
 	}
 }
 
-/**
- * Decompresses a chunk of compressed audio to the destination memory
- *
- * @param Destination		Memory to decompress to
- * @param bLooping			Whether to loop the sound seamlessly, or pad with zeroes
- * @return					Whether the sound looped or not
- */
-bool FCoreAudioSoundBuffer::ReadCompressedData( uint8* Destination, bool bLooping )
+bool FCoreAudioSoundBuffer::ReadCompressedData(uint8* Destination, int32 NumFramesToDecode, bool bLooping)
 {
 	const uint32 kPCMBufferSize = MONO_PCM_BUFFER_SIZE * NumChannels;
 	if (SoundFormat == SoundFormat_Streaming)
diff --git a/Engine/Source/Runtime/Mac/CoreAudio/Private/CoreAudioSource.cpp b/Engine/Source/Runtime/Mac/CoreAudio/Private/CoreAudioSource.cpp
--- a/Engine/Source/Runtime/Mac/CoreAudio/Private/CoreAudioSource.cpp
+++ b/Engine/Source/Runtime/Mac/CoreAudio/Private/CoreAudioSource.cpp
@@ -165,11 +165,11 @@ bool FCoreAudioSoundSource::ReadMorePCMData( const int32 BufferIndex, EDataReadM
 		if (DataReadMode == EDataReadMode::Synchronous)
 		{
 			++NumActiveBuffers;
-			return CoreAudioBuffer->ReadCompressedData( CoreAudioBuffers[BufferIndex].AudioData, WaveInstance->LoopingMode != LOOP_Never );
+			return CoreAudioBuffer->ReadCompressedData( CoreAudioBuffers[BufferIndex].AudioData, MONO_PCM_BUFFER_SAMPLES, WaveInstance->LoopingMode != LOOP_Never );
 		}
 		else
 		{
-			RealtimeAsyncTask = new FAsyncRealtimeAudioTask(CoreAudioBuffer, (uint8*)CoreAudioBuffers[BufferIndex].AudioData, WaveInstance->LoopingMode != LOOP_Never, DataReadMode == EDataReadMode::AsynchronousSkipFirstFrame);
+			RealtimeAsyncTask = new FAsyncRealtimeAudioTask(CoreAudioBuffer, (uint8*)CoreAudioBuffers[BufferIndex].AudioData, WaveInstance->WaveData->NumPrecacheFrames, WaveInstance->LoopingMode != LOOP_Never, DataReadMode == EDataReadMode::AsynchronousSkipFirstFrame);
 			RealtimeAsyncTask->StartBackgroundTask();
 			return false;
 		}
diff --git a/Engine/Source/Runtime/Mac/CoreAudio/Public/CoreAudioDevice.h b/Engine/Source/Runtime/Mac/CoreAudio/Public/CoreAudioDevice.h
--- a/Engine/Source/Runtime/Mac/CoreAudio/Public/CoreAudioDevice.h
+++ b/Engine/Source/Runtime/Mac/CoreAudio/Public/CoreAudioDevice.h
@@ -82,7 +82,7 @@ public:
 	 * @param bLooping			Whether to loop the sound seamlessly, or pad with zeroes
 	 * @return					Whether the sound looped or not
 	 */
-	bool ReadCompressedData( uint8* Destination, bool bLooping );
+	virtual bool ReadCompressedData( uint8* Destination, int32 NumFramesToDecode, bool bLooping ) override;
 
 	/**
 	 * Sets the point in time within the buffer to the specified time
diff --git a/Engine/Source/Runtime/Windows/XAudio2/Private/XAudio2Buffer.cpp b/Engine/Source/Runtime/Windows/XAudio2/Private/XAudio2Buffer.cpp
--- a/Engine/Source/Runtime/Windows/XAudio2/Private/XAudio2Buffer.cpp
+++ b/Engine/Source/Runtime/Windows/XAudio2/Private/XAudio2Buffer.cpp
@@ -329,7 +329,7 @@ bool FXAudio2SoundBuffer::ReadCompressedInfo(USoundWave* SoundWave)
 	return DecompressionState->ReadCompressedInfo(SoundWave->ResourceData, SoundWave->ResourceSize, nullptr);
 }
 
-bool FXAudio2SoundBuffer::ReadCompressedData( uint8* Destination, bool bLooping )
+bool FXAudio2SoundBuffer::ReadCompressedData( uint8* Destination, int32 NumFramesToDecode, bool bLooping )
 {
 	if (!DecompressionState)
 	{
@@ -337,14 +337,14 @@ bool FXAudio2SoundBuffer::ReadCompressedData( uint8* Destination, bool bLooping
 		return false;
 	}
 	
-	const uint32 kPCMBufferSize = MONO_PCM_BUFFER_SIZE * NumChannels;
+	int32 NumBytesToDecode = NumFramesToDecode * sizeof(int16) * NumChannels;
 	if (SoundFormat == SoundFormat_Streaming)
 	{
-		return DecompressionState->StreamCompressedData( Destination, bLooping, kPCMBufferSize );
+		return DecompressionState->StreamCompressedData(Destination, bLooping, NumBytesToDecode);
 	}
 	else
 	{
-		return( DecompressionState->ReadCompressedData( Destination, bLooping, kPCMBufferSize ) );
+		return DecompressionState->ReadCompressedData(Destination, bLooping, NumBytesToDecode);
 	}
 }
 
diff --git a/Engine/Source/Runtime/Windows/XAudio2/Private/XAudio2Source.cpp b/Engine/Source/Runtime/Windows/XAudio2/Private/XAudio2Source.cpp
--- a/Engine/Source/Runtime/Windows/XAudio2/Private/XAudio2Source.cpp
+++ b/Engine/Source/Runtime/Windows/XAudio2/Private/XAudio2Source.cpp
@@ -177,19 +177,24 @@ void FXAudio2SoundSource::SubmitPCMBuffers( void )
 bool FXAudio2SoundSource::ReadMorePCMData( const int32 BufferIndex, EDataReadMode DataReadMode )
 {
 	USoundWave* WaveData = WaveInstance->WaveData;
+
+	const int32 BufferBytes = MONO_PCM_BUFFER_SIZE * Buffer->NumChannels;
+	XAudio2Buffers[BufferIndex].pAudioData = GetRealtimeBufferData(BufferIndex, BufferBytes);
+	XAudio2Buffers[BufferIndex].AudioBytes = BufferBytes;
+
 	if( WaveData && WaveData->bProcedural )
 	{
-		const int32 MaxSamples = ( MONO_PCM_BUFFER_SIZE * Buffer->NumChannels ) / sizeof( int16 );
+		const int32 BufferSamples = MONO_PCM_BUFFER_SAMPLES * Buffer->NumChannels;
 
 		if (DataReadMode == EDataReadMode::Synchronous || WaveData->bCanProcessAsync == false)
 		{
-			const int32 BytesWritten = WaveData->GeneratePCMData( (uint8*)XAudio2Buffers[BufferIndex].pAudioData, MaxSamples );
+			const int32 BytesWritten = WaveData->GeneratePCMData( (uint8*)XAudio2Buffers[BufferIndex].pAudioData, BufferSamples);
 			XAudio2Buffers[BufferIndex].AudioBytes = BytesWritten;
 		}
 		else
 		{
 			check(!RealtimeAsyncTask);
-			RealtimeAsyncTask = new FAsyncRealtimeAudioTask(WaveData, (uint8*)XAudio2Buffers[BufferIndex].pAudioData, MaxSamples);
+			RealtimeAsyncTask = new FAsyncRealtimeAudioTask(WaveData, (uint8*)XAudio2Buffers[BufferIndex].pAudioData, BufferSamples);
 			RealtimeAsyncTask->StartBackgroundTask();
 		}
 
@@ -200,12 +205,11 @@ bool FXAudio2SoundSource::ReadMorePCMData( const int32 BufferIndex, EDataReadMod
 	{
 		if (DataReadMode == EDataReadMode::Synchronous)
 		{
-			return XAudio2Buffer->ReadCompressedData( ( uint8* )XAudio2Buffers[BufferIndex].pAudioData, WaveInstance->LoopingMode != LOOP_Never );
+			return XAudio2Buffer->ReadCompressedData( ( uint8* )XAudio2Buffers[BufferIndex].pAudioData, MONO_PCM_BUFFER_SAMPLES, WaveInstance->LoopingMode != LOOP_Never );
 		}
 		else
 		{
-			check(!RealtimeAsyncTask);
-			RealtimeAsyncTask = new FAsyncRealtimeAudioTask(XAudio2Buffer, (uint8*)XAudio2Buffers[BufferIndex].pAudioData, WaveInstance->LoopingMode != LOOP_Never, DataReadMode == EDataReadMode::AsynchronousSkipFirstFrame);
+			RealtimeAsyncTask = new FAsyncRealtimeAudioTask(XAudio2Buffer, (uint8*)XAudio2Buffers[BufferIndex].pAudioData, WaveInstance->WaveData->NumPrecacheFrames, WaveInstance->LoopingMode != LOOP_Never, DataReadMode == EDataReadMode::AsynchronousSkipFirstFrame);
 			RealtimeAsyncTask->StartBackgroundTask();
 			return false;
 		}
@@ -233,13 +237,8 @@ void FXAudio2SoundSource::SubmitPCMRTBuffers( void )
 	// Set the buffer to be in real time mode
 	CurrentBuffer = 0;
 
-	const uint32 BufferSize = MONO_PCM_BUFFER_SIZE * Buffer->NumChannels;
-
-	// Set up buffer areas to decompress to
 	for (int32 i = 0; i < 3; ++i)
 	{
-		XAudio2Buffers[i].pAudioData = GetRealtimeBufferData(i, BufferSize);
-		XAudio2Buffers[i].AudioBytes = BufferSize;
 		XAudio2Buffers[i].pContext = this;
 	}
 
@@ -249,8 +248,17 @@ void FXAudio2SoundSource::SubmitPCMRTBuffers( void )
 	if (WaveInstance->WaveData && WaveInstance->WaveData->CachedRealtimeFirstBuffer && !bIsSeeking)
 	{
 		bPlayedCachedBuffer = true;
-		FMemory::Memcpy((uint8*)XAudio2Buffers[0].pAudioData, WaveInstance->WaveData->CachedRealtimeFirstBuffer, BufferSize);
-		FMemory::Memcpy((uint8*)XAudio2Buffers[1].pAudioData, WaveInstance->WaveData->CachedRealtimeFirstBuffer + BufferSize, BufferSize);
+		const uint32 PrecacheBufferSize = WaveInstance->WaveData->NumPrecacheFrames * sizeof(int16) * Buffer->NumChannels;
+
+		for (int32 i = 0; i < 2; ++i)
+		{
+			XAudio2Buffers[i].pAudioData = GetRealtimeBufferData(i, PrecacheBufferSize);
+			XAudio2Buffers[i].AudioBytes = PrecacheBufferSize;
+			XAudio2Buffers[i].pContext = this;
+		}
+
+		FMemory::Memcpy((uint8*)XAudio2Buffers[0].pAudioData, WaveInstance->WaveData->CachedRealtimeFirstBuffer, PrecacheBufferSize);
+		FMemory::Memcpy((uint8*)XAudio2Buffers[1].pAudioData, WaveInstance->WaveData->CachedRealtimeFirstBuffer + PrecacheBufferSize, PrecacheBufferSize);
 	}
 	else
 	{
diff --git a/Engine/Source/Runtime/Windows/XAudio2/Private/XAudio2Support.h b/Engine/Source/Runtime/Windows/XAudio2/Private/XAudio2Support.h
--- a/Engine/Source/Runtime/Windows/XAudio2/Private/XAudio2Support.h
+++ b/Engine/Source/Runtime/Windows/XAudio2/Private/XAudio2Support.h
@@ -270,9 +270,10 @@ public:
 	 *
 	 * @param Destination		Memory to decompress to
 	 * @param bLooping			Whether to loop the sound seamlessly, or pad with zeroes
+	 * @param DecodeSizeBytes	Number of bytes to decode
 	 * @return					Whether the sound looped or not
 	 */
-	bool ReadCompressedData( uint8* Destination, bool bLooping ) override;
+	bool ReadCompressedData( uint8* Destination, int32 NumFramesToDecode, bool bLooping) override;
 
 	/**
 	 * Sets the point in time within the buffer to the specified time
