Fixing static analysis warning about null deref of SoundWave ptr.

diff --git a/Engine/Source/Runtime/AudioMixer/Private/AudioMixerSourceBuffer.cpp b/Engine/Source/Runtime/AudioMixer/Private/AudioMixerSourceBuffer.cpp
--- a/Engine/Source/Runtime/AudioMixer/Private/AudioMixerSourceBuffer.cpp
+++ b/Engine/Source/Runtime/AudioMixer/Private/AudioMixerSourceBuffer.cpp
@@ -295,7 +295,15 @@ namespace Audio
 		NewTaskData.bLoopingMode = LoopingMode != LOOP_Never;
 		NewTaskData.bSkipFirstBuffer = (BufferReadMode == EBufferReadMode::AsynchronousSkipFirstFrame);
 		NewTaskData.NumFramesToDecode = MONO_PCM_BUFFER_SAMPLES;
-		NewTaskData.NumPrecacheFrames = SoundWave->NumPrecacheFrames;
+
+		if (SoundWave)
+		{
+			NewTaskData.NumPrecacheFrames = SoundWave->NumPrecacheFrames;
+		}
+		else
+		{
+			NewTaskData.NumPrecacheFrames = MONO_PCM_BUFFER_SAMPLES;
+		}
 
 		check(!AsyncRealtimeAudioTask);
 		AsyncRealtimeAudioTask = CreateAudioTask(NewTaskData);
