Replace check with clamp, so that GetSpectralData won't crash when given a frequency greater than nyquist. #fyi shane.caudle #codereview aaron.mcleran

diff --git a/Engine/Source/Runtime/AudioMixer/Private/DSP/SpectrumAnalyzer.cpp b/Engine/Source/Runtime/AudioMixer/Private/DSP/SpectrumAnalyzer.cpp
--- a/Engine/Source/Runtime/AudioMixer/Private/DSP/SpectrumAnalyzer.cpp
+++ b/Engine/Source/Runtime/AudioMixer/Private/DSP/SpectrumAnalyzer.cpp
@@ -87,7 +87,9 @@ namespace Audio
 		// float Position = NyquistPosition + (InFreq / Nyquist);
 		const float NormalizedFreq = (InFreq / Nyquist);
 		float Position = InFreq >= 0 ? (NormalizedFreq * VectorLength / 2) : VectorLength - (NormalizedFreq * VectorLength / 2);
-		check(Position >= 0.0f && Position <= VectorLength - 1);
+		
+		// Position should be clamped between just above DC and just below nyquist to avoid rounding errors.
+		Position = FMath::Clamp<float>(Position, 0.01f, VectorLength - 1.01f);
 
 		switch (InMethod)
 		{
