[NO NON-PII COMMIT MESSAGE]

diff --git a/Engine/Source/Runtime/AudioMixer/Private/DSP/SpectrumAnalyzer.cpp b/Engine/Source/Runtime/AudioMixer/Private/DSP/SpectrumAnalyzer.cpp
--- a/Engine/Source/Runtime/AudioMixer/Private/DSP/SpectrumAnalyzer.cpp
+++ b/Engine/Source/Runtime/AudioMixer/Private/DSP/SpectrumAnalyzer.cpp
@@ -8,6 +8,7 @@ namespace Audio
 	FSpectrumAnalyzer::FSpectrumAnalyzer()
 		: CurrentSettings(SpectrumAnalyzerSettings::FSettings())
 		, bSettingsWereUpdated(false)
+		, bIsInitialized(false)
 		, SampleRate(0.0f)
 		, Window(CurrentSettings.WindowType, (int32)CurrentSettings.FFTSize, 1, false)
 		, InputQueue(FMath::Max((int32)CurrentSettings.FFTSize * 4, 4096))
@@ -19,6 +20,7 @@ namespace Audio
 	FSpectrumAnalyzer::FSpectrumAnalyzer(const SpectrumAnalyzerSettings::FSettings& InSettings, float InSampleRate)
 		: CurrentSettings(InSettings)
 		, bSettingsWereUpdated(false)
+		, bIsInitialized(true)
 		, SampleRate(InSampleRate)
 		, Window(InSettings.WindowType, (int32)InSettings.FFTSize, 1, false)
 		, InputQueue(FMath::Max((int32)CurrentSettings.FFTSize * 4, 4096))
@@ -31,6 +33,7 @@ namespace Audio
 	FSpectrumAnalyzer::FSpectrumAnalyzer(float InSampleRate)
 		: CurrentSettings(SpectrumAnalyzerSettings::FSettings())
 		, bSettingsWereUpdated(false)
+		, bIsInitialized(true)
 		, SampleRate(InSampleRate)
 		, Window(CurrentSettings.WindowType, (int32)CurrentSettings.FFTSize, 1, false)
 		, InputQueue(FMath::Max((int32)CurrentSettings.FFTSize * 4, 4096))
@@ -54,6 +57,8 @@ namespace Audio
 		InputQueue.SetCapacity(FMath::Max((int32)CurrentSettings.FFTSize * 4, 4096));
 		FrequencyBuffer.Reset(CurrentSettings);
 		ResetSettings();
+
+		bIsInitialized = true;
 	}
 
 	void FSpectrumAnalyzer::ResetSettings()
@@ -158,6 +163,11 @@ namespace Audio
 
 	float FSpectrumAnalyzer::GetMagnitudeForFrequency(float InFrequency)
 	{
+		if (!bIsInitialized)
+		{
+			return false;
+		}
+
 		const FSpectrumAnalyzerFrequencyVector* OutVector;
 		bool bShouldUnlockBuffer = true;
 
@@ -195,6 +205,11 @@ namespace Audio
 
 	float FSpectrumAnalyzer::GetPhaseForFrequency(float InFrequency)
 	{
+		if (!bIsInitialized)
+		{
+			return false;
+		}
+
 		const FSpectrumAnalyzerFrequencyVector* OutVector;
 		bool bShouldUnlockBuffer = true;
 
@@ -232,6 +247,11 @@ namespace Audio
 
 	void FSpectrumAnalyzer::LockOutputBuffer()
 	{
+		if (!bIsInitialized)
+		{
+			return;
+		}
+
 		if (LockedFrequencyVector != nullptr)
 		{
 			FrequencyBuffer.UnlockBuffer();
@@ -242,6 +262,11 @@ namespace Audio
 
 	void FSpectrumAnalyzer::UnlockOutputBuffer()
 	{
+		if (!bIsInitialized)
+		{
+			return;
+		}
+
 		if (LockedFrequencyVector != nullptr)
 		{
 			FrequencyBuffer.UnlockBuffer();
@@ -262,6 +287,11 @@ namespace Audio
 
 	bool FSpectrumAnalyzer::PerformAnalysisIfPossible()
 	{
+		if (!bIsInitialized)
+		{
+			return false;
+		}
+
 		// If settings were updated, perform resizing and parameter updates here:
 		if (bSettingsWereUpdated)
 		{
@@ -307,6 +337,11 @@ namespace Audio
 		}
 	}
 
+	bool FSpectrumAnalyzer::IsInitialized()
+	{
+		return bIsInitialized;
+	}
+
 	static const int32 SpectrumAnalyzerBufferSize = 3;
 
 	FSpectrumAnalyzerBuffer::FSpectrumAnalyzerBuffer()
diff --git a/Engine/Source/Runtime/AudioMixer/Public/DSP/SpectrumAnalyzer.h b/Engine/Source/Runtime/AudioMixer/Public/DSP/SpectrumAnalyzer.h
--- a/Engine/Source/Runtime/AudioMixer/Public/DSP/SpectrumAnalyzer.h
+++ b/Engine/Source/Runtime/AudioMixer/Public/DSP/SpectrumAnalyzer.h
@@ -115,8 +115,10 @@ namespace Audio
 	class AUDIOMIXER_API FSpectrumAnalyzer
 	{
 	public:
-		// Default constructor needs to call Init before using
+		// If an instance is created using the default constructor, Init() must be called before it is used.
 		FSpectrumAnalyzer();
+
+		// If an instance is created using either of these constructors, Init() is not neccessary.
 		FSpectrumAnalyzer(float InSampleRate);
 		FSpectrumAnalyzer(const SpectrumAnalyzerSettings::FSettings& InSettings, float InSampleRate);
 
@@ -150,6 +152,10 @@ namespace Audio
 		// Thread safe call to perform actual FFT. Returns true if it performed the FFT, false otherwise.
 		bool PerformAnalysisIfPossible();
 
+		// Returns false if this instance of FSpectrumAnalyzer was constructed with the default constructor 
+		// and Init() has not been called yet.
+		bool IsInitialized();
+
 	private:
 
 		// Called on analysis thread.
@@ -162,6 +168,8 @@ namespace Audio
 		SpectrumAnalyzerSettings::FSettings CurrentSettings;
 		volatile bool bSettingsWereUpdated;
 
+		volatile bool bIsInitialized;
+
 		float SampleRate;
 
 		// Cached window that is applied prior to running the FFT.
