[NO NON-PII COMMIT MESSAGE]

diff --git a/Engine/Plugins/Runtime/TimeSynth/Source/TimeSynth/Classes/TimeSynthComponent.h b/Engine/Plugins/Runtime/TimeSynth/Source/TimeSynth/Classes/TimeSynthComponent.h
--- a/Engine/Plugins/Runtime/TimeSynth/Source/TimeSynth/Classes/TimeSynthComponent.h
+++ b/Engine/Plugins/Runtime/TimeSynth/Source/TimeSynth/Classes/TimeSynthComponent.h
@@ -703,7 +703,7 @@ private:
 
 	// Spectum analyzer to allow BP delegates to visualize music
 	Audio::FSpectrumAnalyzer SpectrumAnalyzer;
-	Audio::SpectrumAnalyzerSettings::FSettings SpectrumAnalyzerSettings;
+	Audio::FSpectrumAnalyzerSettings SpectrumAnalyzerSettings;
 	FThreadSafeCounter SpectrumAnalysisCounter;
 
 	// Array of spectrum data, maps to FrequenciesToAnalyze UProperty
diff --git a/Engine/Source/Runtime/AudioMixer/Private/AudioMixerBlueprintLibrary.cpp b/Engine/Source/Runtime/AudioMixer/Private/AudioMixerBlueprintLibrary.cpp
--- a/Engine/Source/Runtime/AudioMixer/Private/AudioMixerBlueprintLibrary.cpp
+++ b/Engine/Source/Runtime/AudioMixer/Private/AudioMixerBlueprintLibrary.cpp
@@ -5,6 +5,7 @@
 #include "AudioDevice.h"
 #include "AudioMixerDevice.h"
 #include "CoreMinimal.h"
+#include "DSP/SpectrumAnalyzer.h"
 
 
 // This is our global recording task:
@@ -190,6 +191,56 @@ void UAudioMixerBlueprintLibrary::ResumeRecordingOutput(const UObject* WorldCont
 	}
 }
 
+void UAudioMixerBlueprintLibrary::StartAnalyzingOutput(const UObject* WorldContextObject, USoundSubmix* SubmixToAnalyze /*= nullptr*/, EFFTSize FFTSize /*= EFFTSize::Default*/, EFFTPeakInterpolationMethod InterpolationMethod /*= Linear*/, EFFTWindowType WindowType /*= EFFTWindowType::Hamming*/, float HopSize /*= 0*/)
+{
+	if (Audio::FMixerDevice* MixerDevice = GetAudioMixerDeviceFromWorldContext(WorldContextObject))
+	{
+		Audio::FSpectrumAnalyzerSettings Settings = Audio::FSpectrumAnalyzerSettings();
+		PopulateSpectrumAnalyzerSettings(FFTSize, InterpolationMethod, WindowType, HopSize, Settings);
+		MixerDevice->StartSpectrumAnalysis(SubmixToAnalyze, Settings);
+	}
+	else
+	{
+		UE_LOG(LogAudioMixer, Error, TEXT("Spectrum Analysis is an audio mixer only feature. Please run the game with -audiomixer to enable this feature."));
+	}
+}
+
+void UAudioMixerBlueprintLibrary::StopAnalyzingOutput(const UObject* WorldContextObject, USoundSubmix* SubmixToStopAnalyzing /*= nullptr*/)
+{
+	if (Audio::FMixerDevice* MixerDevice = GetAudioMixerDeviceFromWorldContext(WorldContextObject))
+	{
+		MixerDevice->StopSpectrumAnalysis(SubmixToStopAnalyzing);
+	}
+	else
+	{
+		UE_LOG(LogAudioMixer, Error, TEXT("Spectrum Analysis is an audio mixer only feature. Please run the game with -audiomixer to enable this feature."));
+	}
+}
+
+void UAudioMixerBlueprintLibrary::GetMagnitudeForFrequencies(const UObject* WorldContextObject, const TArray<float>& Frequencies, TArray<float>& Magnitudes, USoundSubmix* SubmixToAnalyze /*= nullptr*/)
+{
+	if (Audio::FMixerDevice* MixerDevice = GetAudioMixerDeviceFromWorldContext(WorldContextObject))
+	{
+		MixerDevice->GetMagnitudesForFrequencies(SubmixToAnalyze, Frequencies, Magnitudes);
+	}
+	else
+	{
+		UE_LOG(LogAudioMixer, Error, TEXT("Output recording is an audio mixer only feature. Please run the game with -audiomixer to enable this feature."));
+	}
+}
+
+void UAudioMixerBlueprintLibrary::GetPhaseForFrequencies(const UObject* WorldContextObject, const TArray<float>& Frequencies, TArray<float>& Phases, USoundSubmix* SubmixToAnalyze /*= nullptr*/)
+{
+	if (Audio::FMixerDevice* MixerDevice = GetAudioMixerDeviceFromWorldContext(WorldContextObject))
+	{
+		MixerDevice->GetPhasesForFrequencies(SubmixToAnalyze, Frequencies, Phases);
+	}
+	else
+	{
+		UE_LOG(LogAudioMixer, Error, TEXT("Output recording is an audio mixer only feature. Please run the game with -audiomixer to enable this feature."));
+	}
+}
+
 void UAudioMixerBlueprintLibrary::AddSourceEffectToPresetChain(const UObject* WorldContextObject, USoundEffectSourcePresetChain* PresetChain, FSourceEffectChainEntry Entry)
 {
 	if (!PresetChain)
@@ -295,3 +346,65 @@ int32 UAudioMixerBlueprintLibrary::GetNumberOfEntriesInSourceEffectChain(const U
 
 	return 0;
 }
+
+void UAudioMixerBlueprintLibrary::PopulateSpectrumAnalyzerSettings(EFFTSize FFTSize, EFFTPeakInterpolationMethod InterpolationMethod, EFFTWindowType WindowType, float HopSize, Audio::FSpectrumAnalyzerSettings &OutSettings)
+{
+	switch (FFTSize)
+	{
+	case EFFTSize::DefaultSize:
+		OutSettings.FFTSize = Audio::FSpectrumAnalyzerSettings::EFFTSize::Default;
+		break;
+	case EFFTSize::Min:
+		OutSettings.FFTSize = Audio::FSpectrumAnalyzerSettings::EFFTSize::Min_64;
+		break;
+	case EFFTSize::Small:
+		OutSettings.FFTSize = Audio::FSpectrumAnalyzerSettings::EFFTSize::Small_256;
+		break;
+	case EFFTSize::Medium:
+		OutSettings.FFTSize = Audio::FSpectrumAnalyzerSettings::EFFTSize::Medium_512;
+		break;
+	case EFFTSize::Large:
+		OutSettings.FFTSize = Audio::FSpectrumAnalyzerSettings::EFFTSize::Large_1024;
+		break;
+	case EFFTSize::Max:
+		OutSettings.FFTSize = Audio::FSpectrumAnalyzerSettings::EFFTSize::TestLarge_4096;
+		break;
+	default:
+		break;
+	}
+
+	switch (InterpolationMethod)
+	{
+	case EFFTPeakInterpolationMethod::NearestNeighbor:
+		OutSettings.InterpolationMethod = Audio::FSpectrumAnalyzerSettings::EPeakInterpolationMethod::NearestNeighbor;
+		break;
+	case EFFTPeakInterpolationMethod::Linear:
+		OutSettings.InterpolationMethod = Audio::FSpectrumAnalyzerSettings::EPeakInterpolationMethod::Linear;
+		break;
+	case EFFTPeakInterpolationMethod::Quadratic:
+		OutSettings.InterpolationMethod = Audio::FSpectrumAnalyzerSettings::EPeakInterpolationMethod::Quadratic;
+		break;
+	default:
+		break;
+	}
+
+	switch (WindowType)
+	{
+	case EFFTWindowType::None:
+		OutSettings.WindowType = Audio::EWindowType::None;
+		break;
+	case EFFTWindowType::Hamming:
+		OutSettings.WindowType = Audio::EWindowType::Hamming;
+		break;
+	case EFFTWindowType::Hann:
+		OutSettings.WindowType = Audio::EWindowType::Hann;
+		break;
+	case EFFTWindowType::Blackman:
+		OutSettings.WindowType = Audio::EWindowType::Blackman;
+		break;
+	default:
+		break;
+	}
+
+	OutSettings.HopSize = HopSize;
+}
diff --git a/Engine/Source/Runtime/AudioMixer/Private/AudioMixerDevice.cpp b/Engine/Source/Runtime/AudioMixer/Private/AudioMixerDevice.cpp
--- a/Engine/Source/Runtime/AudioMixer/Private/AudioMixerDevice.cpp
+++ b/Engine/Source/Runtime/AudioMixer/Private/AudioMixerDevice.cpp
@@ -1361,6 +1361,74 @@ namespace Audio
 	}
 
 
+	void FMixerDevice::StartSpectrumAnalysis(USoundSubmix* InSubmix, const Audio::FSpectrumAnalyzerSettings& InSettings)
+	{
+		Audio::FMixerSubmixPtr* FoundSubmix = Submixes.Find(InSubmix);
+		if (FoundSubmix)
+		{
+			(*FoundSubmix)->StartSpectrumAnalysis(InSettings);
+		}
+		else
+		{
+			FMixerSubmixWeakPtr MasterSubmix = GetMasterSubmix();
+			FMixerSubmixPtr MasterSubmixPtr = MasterSubmix.Pin();
+			check(MasterSubmixPtr.IsValid());
+
+			MasterSubmixPtr->StartSpectrumAnalysis(InSettings);
+		}
+	}
+
+	void FMixerDevice::StopSpectrumAnalysis(USoundSubmix* InSubmix)
+	{
+		Audio::FMixerSubmixPtr* FoundSubmix = Submixes.Find(InSubmix);
+		if (FoundSubmix)
+		{
+			(*FoundSubmix)->StopSpectrumAnalysis();
+		}
+		else
+		{
+			FMixerSubmixWeakPtr MasterSubmix = GetMasterSubmix();
+			FMixerSubmixPtr MasterSubmixPtr = MasterSubmix.Pin();
+			check(MasterSubmixPtr.IsValid());
+
+			MasterSubmixPtr->StopSpectrumAnalysis();
+		}
+	}
+
+	void FMixerDevice::GetMagnitudesForFrequencies(USoundSubmix* InSubmix, const TArray<float>& InFrequencies, TArray<float>& OutMagnitudes)
+	{
+		Audio::FMixerSubmixPtr* FoundSubmix = Submixes.Find(InSubmix);
+		if (FoundSubmix)
+		{
+			(*FoundSubmix)->GetMagnitudeForFrequencies(InFrequencies, OutMagnitudes);
+		}
+		else
+		{
+			FMixerSubmixWeakPtr MasterSubmix = GetMasterSubmix();
+			FMixerSubmixPtr MasterSubmixPtr = MasterSubmix.Pin();
+			check(MasterSubmixPtr.IsValid());
+
+			MasterSubmixPtr->GetMagnitudeForFrequencies(InFrequencies, OutMagnitudes);
+		}
+	}
+
+	void FMixerDevice::GetPhasesForFrequencies(USoundSubmix* InSubmix, const TArray<float>& InFrequencies, TArray<float>& OutPhases)
+	{
+		Audio::FMixerSubmixPtr* FoundSubmix = Submixes.Find(InSubmix);
+		if (FoundSubmix)
+		{
+			(*FoundSubmix)->GetPhaseForFrequencies(InFrequencies, OutPhases);
+		}
+		else
+		{
+			FMixerSubmixWeakPtr MasterSubmix = GetMasterSubmix();
+			FMixerSubmixPtr MasterSubmixPtr = MasterSubmix.Pin();
+			check(MasterSubmixPtr.IsValid());
+
+			MasterSubmixPtr->GetPhaseForFrequencies(InFrequencies, OutPhases);
+		}
+	}
+
 	void FMixerDevice::RegisterSubmixBufferListener(ISubmixBufferListener* InSubmixBufferListener, USoundSubmix* InSubmix)
 	{
 		Audio::FMixerSubmixPtr* FoundSubmix = Submixes.Find(InSubmix);
diff --git a/Engine/Source/Runtime/AudioMixer/Private/AudioMixerSubmix.cpp b/Engine/Source/Runtime/AudioMixer/Private/AudioMixerSubmix.cpp
--- a/Engine/Source/Runtime/AudioMixer/Private/AudioMixerSubmix.cpp
+++ b/Engine/Source/Runtime/AudioMixer/Private/AudioMixerSubmix.cpp
@@ -360,6 +360,29 @@ namespace Audio
 		}
 	}
 
+	void FMixerSubmix::MixBufferDownToMono(const AlignedFloatBuffer& InBuffer, int32 NumInputChannels, AlignedFloatBuffer& OutBuffer)
+	{
+		check(NumInputChannels > 0);
+
+		int32 NumFrames = InBuffer.Num() / NumInputChannels;
+		OutBuffer.Reset();
+		OutBuffer.AddZeroed(NumFrames);
+
+		const float* InData = InBuffer.GetData();
+		float* OutData = OutBuffer.GetData();
+
+		const float GainFactor = 1.0f / FMath::Sqrt((float) NumInputChannels);
+
+		for (int32 FrameIndex = 0; FrameIndex < NumFrames; FrameIndex++)
+		{
+			for (int32 ChannelIndex = 0; ChannelIndex < NumInputChannels; ChannelIndex++)
+			{
+				const int32 InputIndex = FrameIndex * NumInputChannels + ChannelIndex;
+				OutData[FrameIndex] += InData[InputIndex] * GainFactor;
+			}
+		}
+	}
+
 	void FMixerSubmix::SetUpAmbisonicsEncoder()
 	{
 		check(AmbisonicsMixer.IsValid());
@@ -665,6 +688,15 @@ namespace Audio
 			}
 		}
 
+		// If spectrum analysis is enabled for this submix, downmix the resulting audio
+		// and push it to the spectrum analyzer.
+		if (SpectrumAnalyzer.IsValid())
+		{
+			MixBufferDownToMono(InputBuffer, NumChannels, MonoMixBuffer);
+			SpectrumAnalyzer->PushAudio(MonoMixBuffer.GetData(), MonoMixBuffer.Num());
+			SpectrumAnalyzer->PerformAnalysisIfPossible(true, true);
+		}
+
 		// If the channel types match, just do a copy
 		if (ChannelFormat != ParentChannelType || SubmixAmbisonicsDecoderID != INDEX_NONE)
 		{
@@ -844,6 +876,56 @@ namespace Audio
 		OnSubmixEnvelope.AddUnique(OnSubmixEnvelopeBP);
 	}
 
+	void FMixerSubmix::StartSpectrumAnalysis(const FSpectrumAnalyzerSettings& InSettings)
+	{
+		SpectrumAnalyzer.Reset(new FSpectrumAnalyzer(InSettings, MixerDevice->GetSampleRate()));
+	}
+
+	void FMixerSubmix::StopSpectrumAnalysis()
+	{
+		SpectrumAnalyzer.Reset();
+	}
+
+	void FMixerSubmix::GetMagnitudeForFrequencies(const TArray<float>& InFrequencies, TArray<float>& OutMagnitudes)
+	{
+		if (SpectrumAnalyzer.IsValid())
+		{
+			OutMagnitudes.Reset();
+			OutMagnitudes.AddUninitialized(InFrequencies.Num());
+
+			SpectrumAnalyzer->LockOutputBuffer();
+			for (int32 Index = 0; Index < InFrequencies.Num(); Index++)
+			{
+				OutMagnitudes[Index] = SpectrumAnalyzer->GetMagnitudeForFrequency(InFrequencies[Index]);
+			}
+			SpectrumAnalyzer->UnlockOutputBuffer();
+		}
+		else
+		{
+			UE_LOG(LogAudioMixer, Warning, TEXT("Call StartSpectrumAnalysis before calling GetMagnitudeForFrequencies."));
+		}
+	}
+
+	void FMixerSubmix::GetPhaseForFrequencies(const TArray<float>& InFrequencies, TArray<float>& OutPhases)
+	{
+		if (SpectrumAnalyzer.IsValid())
+		{
+			OutPhases.Reset();
+			OutPhases.AddUninitialized(InFrequencies.Num());
+
+			SpectrumAnalyzer->LockOutputBuffer();
+			for (int32 Index = 0; Index < InFrequencies.Num(); Index++)
+			{
+				OutPhases[Index] = SpectrumAnalyzer->GetPhaseForFrequency(InFrequencies[Index]);
+			}
+			SpectrumAnalyzer->UnlockOutputBuffer();
+		}
+		else
+		{
+			UE_LOG(LogAudioMixer, Warning, TEXT("Call StartSpectrumAnalysis before calling GetMagnitudeForFrequencies."));
+		}
+	}
+
 	void FMixerSubmix::BroadcastEnvelope()
 	{
 
diff --git a/Engine/Source/Runtime/AudioMixer/Private/DSP/AudioFFT.cpp b/Engine/Source/Runtime/AudioMixer/Private/DSP/AudioFFT.cpp
--- a/Engine/Source/Runtime/AudioMixer/Private/DSP/AudioFFT.cpp
+++ b/Engine/Source/Runtime/AudioMixer/Private/DSP/AudioFFT.cpp
@@ -21,7 +21,7 @@ namespace Audio
 	{
 		const int32 N = bIsPeriodic ? NumFrames : NumFrames - 1;
 		const float PhaseDelta = 2.0f * PI / N;
-		float Phase = 0;
+		float Phase = 0.0f;
 
 		for (int32 FrameIndex = 0; FrameIndex < NumFrames; FrameIndex++)
 		{
@@ -104,59 +104,6 @@ namespace Audio
 		}
 	}
 
-	FWindow::FWindow(EWindowType InType, int32 InNumFrames, int32 InNumChannels, bool bIsPeriodic)
-		: WindowType(InType)
-		, NumSamples(InNumFrames * InNumChannels)
-	{
-		checkf(NumSamples % 4 == 0, TEXT("For performance reasons, this window's length should be a multiple of 4."));
-		Generate(InNumFrames, InNumChannels, bIsPeriodic);
-	}
-
-	FWindow::~FWindow()
-	{
-
-	}
-
-	void FWindow::Generate(int32 NumFrames, int32 NumChannels, bool bIsPeriodic)
-	{
-		if (WindowType == EWindowType::None)
-		{
-			return;
-		}
-
-		WindowBuffer.Reset();
-		WindowBuffer.AddZeroed(NumSamples);
-
-		switch (WindowType)
-		{
-			case EWindowType::Hann:
-			{
-				GenerateHannWindow(WindowBuffer.GetData(), NumFrames, NumChannels, bIsPeriodic);
-				break;
-			}
-			case EWindowType::Blackman:
-			{
-				GenerateBlackmanWindow(WindowBuffer.GetData(), NumFrames, NumChannels, bIsPeriodic);
-				break;
-			}
-			default:
-			{
-				break;
-			}
-		}
-	}
-
-	void FWindow::ApplyToBuffer(float* InBuffer)
-	{
-		if (WindowType == EWindowType::None)
-		{
-			return;
-		}
-
-		check(IsAligned<float*>(InBuffer, 4));
-		MultiplyBuffersInPlace(WindowBuffer.GetData(), InBuffer, NumSamples);
-	}
-
 	namespace FFTIntrinsics
 	{
 		// Fast bit reversal helper function. Can be used if N is a power of 2. Not well exercised.
diff --git a/Engine/Source/Runtime/AudioMixer/Private/DSP/SpectrumAnalyzer.cpp b/Engine/Source/Runtime/AudioMixer/Private/DSP/SpectrumAnalyzer.cpp
--- a/Engine/Source/Runtime/AudioMixer/Private/DSP/SpectrumAnalyzer.cpp
+++ b/Engine/Source/Runtime/AudioMixer/Private/DSP/SpectrumAnalyzer.cpp
@@ -6,7 +6,7 @@
 namespace Audio
 {
 	FSpectrumAnalyzer::FSpectrumAnalyzer()
-		: CurrentSettings(SpectrumAnalyzerSettings::FSettings())
+		: CurrentSettings(FSpectrumAnalyzerSettings())
 		, bSettingsWereUpdated(false)
 		, bIsInitialized(false)
 		, SampleRate(0.0f)
@@ -17,7 +17,7 @@ namespace Audio
 	{
 	}
 
-	FSpectrumAnalyzer::FSpectrumAnalyzer(const SpectrumAnalyzerSettings::FSettings& InSettings, float InSampleRate)
+	FSpectrumAnalyzer::FSpectrumAnalyzer(const FSpectrumAnalyzerSettings& InSettings, float InSampleRate)
 		: CurrentSettings(InSettings)
 		, bSettingsWereUpdated(false)
 		, bIsInitialized(true)
@@ -31,7 +31,7 @@ namespace Audio
 	}
 
 	FSpectrumAnalyzer::FSpectrumAnalyzer(float InSampleRate)
-		: CurrentSettings(SpectrumAnalyzerSettings::FSettings())
+		: CurrentSettings(FSpectrumAnalyzerSettings())
 		, bSettingsWereUpdated(false)
 		, bIsInitialized(true)
 		, SampleRate(InSampleRate)
@@ -43,13 +43,21 @@ namespace Audio
 		ResetSettings();
 	}
 
+	FSpectrumAnalyzer::~FSpectrumAnalyzer()
+	{
+		if (AsyncAnalysisTask.IsValid())
+		{
+			AsyncAnalysisTask->EnsureCompletion(false);
+		}
+	}
+
 	void FSpectrumAnalyzer::Init(float InSampleRate)
 	{
-		SpectrumAnalyzerSettings::FSettings DefaultSettings = SpectrumAnalyzerSettings::FSettings();
+		FSpectrumAnalyzerSettings DefaultSettings = FSpectrumAnalyzerSettings();
 		Init(DefaultSettings, InSampleRate);
 	}
 
-	void FSpectrumAnalyzer::Init(const SpectrumAnalyzerSettings::FSettings& InSettings, float InSampleRate)
+	void FSpectrumAnalyzer::Init(const FSpectrumAnalyzerSettings& InSettings, float InSampleRate)
 	{
 		CurrentSettings = InSettings;
 		bSettingsWereUpdated = false;
@@ -72,7 +80,15 @@ namespace Audio
 
 		Window = FWindow(CurrentSettings.WindowType, (int32)CurrentSettings.FFTSize, 1, false);
 		FFTSize = (int32) CurrentSettings.FFTSize;
-		HopInSamples = GetCOLAHopSizeForWindow(CurrentSettings.WindowType, (uint32)CurrentSettings.FFTSize);
+		
+		if (FMath::IsNearlyZero(CurrentSettings.HopSize))
+		{
+			HopInSamples = GetCOLAHopSizeForWindow(CurrentSettings.WindowType, (uint32)CurrentSettings.FFTSize);
+		}
+		else
+		{
+			HopInSamples = FMath::FloorToInt((float)CurrentSettings.FFTSize * CurrentSettings.HopSize);
+		}
 		
 		AnalysisTimeDomainBuffer.Reset();
 		AnalysisTimeDomainBuffer.AddZeroed(FFTSize);
@@ -81,7 +97,7 @@ namespace Audio
 		bSettingsWereUpdated = false;
 	}
 
-	void FSpectrumAnalyzer::PerformInterpolation(const FSpectrumAnalyzerFrequencyVector* InFrequencies, SpectrumAnalyzerSettings::EPeakInterpolationMethod InMethod, const float InFreq, float& OutReal, float& OutImag)
+	void FSpectrumAnalyzer::PerformInterpolation(const FSpectrumAnalyzerFrequencyVector* InFrequencies, FSpectrumAnalyzerSettings::EPeakInterpolationMethod InMethod, const float InFreq, float& OutReal, float& OutImag)
 	{
 		const int32 VectorLength = InFrequencies->RealFrequencies.Num();
 		const float NyquistPosition = VectorLength / 2;
@@ -98,7 +114,7 @@ namespace Audio
 
 		switch (InMethod)
 		{
-		case Audio::SpectrumAnalyzerSettings::EPeakInterpolationMethod::NearestNeighbor:
+		case Audio::FSpectrumAnalyzerSettings::EPeakInterpolationMethod::NearestNeighbor:
 		{
 			int32 Index = FMath::RoundToInt(Position);
 			OutReal = InFrequencies->RealFrequencies[Index];
@@ -106,7 +122,7 @@ namespace Audio
 			break;
 		}
 			
-		case Audio::SpectrumAnalyzerSettings::EPeakInterpolationMethod::Linear:
+		case Audio::FSpectrumAnalyzerSettings::EPeakInterpolationMethod::Linear:
 		{
 			const int32 LowerIndex = FMath::FloorToInt(Position);
 			const int32 UpperIndex = FMath::CeilToInt(Position);
@@ -121,7 +137,7 @@ namespace Audio
 			OutImag = FMath::Lerp<float>(y1Imag, y2Imag, PositionFraction);
 			break;
 		}	
-		case Audio::SpectrumAnalyzerSettings::EPeakInterpolationMethod::Quadratic:
+		case Audio::FSpectrumAnalyzerSettings::EPeakInterpolationMethod::Quadratic:
 		{
 			const int32 MidIndex = FMath::RoundToInt(Position);
 			const int32 LowerIndex = FMath::Max(0, MidIndex - 1);
@@ -150,13 +166,13 @@ namespace Audio
 		}
 	}
 
-	void FSpectrumAnalyzer::SetSettings(const SpectrumAnalyzerSettings::FSettings& InSettings)
+	void FSpectrumAnalyzer::SetSettings(const FSpectrumAnalyzerSettings& InSettings)
 	{
 		CurrentSettings = InSettings;
 		bSettingsWereUpdated = true;
 	}
 
-	void FSpectrumAnalyzer::GetSettings(SpectrumAnalyzerSettings::FSettings& OutSettings)
+	void FSpectrumAnalyzer::GetSettings(FSpectrumAnalyzerSettings& OutSettings)
 	{
 		OutSettings = CurrentSettings;
 	}
@@ -285,11 +301,27 @@ namespace Audio
 		return InputQueue.Push(InBuffer, NumSamples) > 0;
 	}
 
-	bool FSpectrumAnalyzer::PerformAnalysisIfPossible()
+	bool FSpectrumAnalyzer::PerformAnalysisIfPossible(bool bUseLatestAudio, bool bAsync)
 	{
 		if (!bIsInitialized)
 		{
 			return false;
+		}		
+
+		if (bAsync)
+		{
+			// if bAsync is true, kick off a new task if one isn't in flight already, and return.
+			if (!AsyncAnalysisTask.IsValid())
+			{
+				AsyncAnalysisTask.Reset(new FSpectrumAnalyzerTask(this, bUseLatestAudio));
+				AsyncAnalysisTask->StartBackgroundTask();
+			}
+			else if (AsyncAnalysisTask->IsDone())
+			{
+				AsyncAnalysisTask->StartBackgroundTask();
+			}
+
+			return true;
 		}
 
 		// If settings were updated, perform resizing and parameter updates here:
@@ -306,6 +338,12 @@ namespace Audio
 		{
 			float* TimeDomainBuffer = AnalysisTimeDomainBuffer.GetData();
 
+			if (bUseLatestAudio)
+			{
+				// If we are only using the latest audio, scrap the oldest audio in the InputQueue:
+				InputQueue.SetNum((uint32)FFTSize);
+			}
+
 			// Perform pop/peek here based on FFT size and hop amount.
 			const int32 PeekAmount = FFTSize - HopInSamples;
 			InputQueue.Pop(TimeDomainBuffer, HopInSamples);
@@ -323,7 +361,6 @@ namespace Audio
 			OutputParams.OutReal = OutputVector->RealFrequencies.GetData();
 			OutputParams.OutImag = OutputVector->ImagFrequencies.GetData();
 
-
 			PerformFFT(InputParams, OutputParams);
 
 			// We're done, so unlock this vector.
@@ -342,7 +379,7 @@ namespace Audio
 		return bIsInitialized;
 	}
 
-	static const int32 SpectrumAnalyzerBufferSize = 3;
+	static const int32 SpectrumAnalyzerBufferSize = 4;
 
 	FSpectrumAnalyzerBuffer::FSpectrumAnalyzerBuffer()
 		: OutputIndex(0)
@@ -350,12 +387,12 @@ namespace Audio
 	{
 	}
 
-	FSpectrumAnalyzerBuffer::FSpectrumAnalyzerBuffer(const SpectrumAnalyzerSettings::FSettings& InSettings)
+	FSpectrumAnalyzerBuffer::FSpectrumAnalyzerBuffer(const FSpectrumAnalyzerSettings& InSettings)
 	{
 		Reset(InSettings);
 	}
 
-	void FSpectrumAnalyzerBuffer::Reset(const SpectrumAnalyzerSettings::FSettings& InSettings)
+	void FSpectrumAnalyzerBuffer::Reset(const FSpectrumAnalyzerSettings& InSettings)
 	{
 		FScopeLock ScopeLock(&BufferIndicesCriticalSection);
 
@@ -426,4 +463,10 @@ namespace Audio
 		ImagFrequencies.Reset();
 		ImagFrequencies.AddZeroed(InFFTSize);
 	}
+
+	void FSpectrumAnalysisAsyncWorker::DoWork()
+	{
+		Analyzer->PerformAnalysisIfPossible(bUseLatestAudio, false);
+	}
+
 }
diff --git a/Engine/Source/Runtime/AudioMixer/Public/AudioMixerBlueprintLibrary.h b/Engine/Source/Runtime/AudioMixer/Public/AudioMixerBlueprintLibrary.h
--- a/Engine/Source/Runtime/AudioMixer/Public/AudioMixerBlueprintLibrary.h
+++ b/Engine/Source/Runtime/AudioMixer/Public/AudioMixerBlueprintLibrary.h
@@ -7,10 +7,57 @@
 #include "SubmixEffects/AudioMixerSubmixEffectDynamicsProcessor.h"
 #include "Sound/SoundEffectSource.h"
 #include "Sound/SampleBuffer.h"
+#include "DSP/SpectrumAnalyzer.h"
 #include "AudioMixerBlueprintLibrary.generated.h"
 
 class USoundSubmix;
 
+UENUM(BlueprintType)
+enum class EFFTSize : uint8
+{
+	// 512
+	DefaultSize, 
+
+	// 64
+	Min,
+
+	// 256
+	Small,
+
+	// 512
+	Medium, 
+
+	// 1024
+	Large, 
+
+	// 4096
+	Max, 
+};
+
+UENUM()
+enum class EFFTPeakInterpolationMethod : uint8
+{
+	NearestNeighbor,
+	Linear,
+	Quadratic
+};
+
+UENUM()
+enum class EFFTWindowType : uint8
+{
+	// No window is applied. Technically a boxcar window.
+	None, 
+
+	// Mainlobe width of -3 dB and sidelove attenuation of ~-40 dB. Good for COLA.
+	Hamming,
+
+	// Mainlobe width of -3 dB and sidelobe attenuation of ~-30dB. Good for COLA.
+	Hann,
+
+	// Mainlobe width of -3 dB and sidelobe attenuation of ~-60db. Tricky for COLA.
+	Blackman
+};
+
 UCLASS(meta=(ScriptName="AudioMixerLibrary"))
 class AUDIOMIXER_API UAudioMixerBlueprintLibrary : public UBlueprintFunctionLibrary
 {
@@ -42,10 +89,26 @@ public:
 	UFUNCTION(BlueprintCallable, Category = "Audio", meta = (WorldContext = "WorldContextObject", AdvancedDisplay = 1))
 	static void PauseRecordingOutput(const UObject* WorldContextObject, USoundSubmix* SubmixToPause = nullptr);
 
-	/** Resume recording audio after pausing. By leaving the Submix To Record field blank, you can record the master output of the game. */
+	/** Resume recording audio after pausing. By leaving the Submix To Pause field blank, you can record the master output of the game. */
 	UFUNCTION(BlueprintCallable, Category = "Audio", meta = (WorldContext = "WorldContextObject", AdvancedDisplay = 1))
 	static void ResumeRecordingOutput(const UObject* WorldContextObject, USoundSubmix* SubmixToPause = nullptr);
 
+	/** Start spectrum analysis of the audio output. By leaving the Submix To Analyze blank, you can analyze the master output of the game. */
+	UFUNCTION(BlueprintCallable, Category = "Audio|Analysis", meta = (WorldContext = "WorldContextObject", AdvancedDisplay = 1))
+	static void StartAnalyzingOutput(const UObject* WorldContextObject, USoundSubmix* SubmixToAnalyze = nullptr, EFFTSize FFTSize = EFFTSize::DefaultSize, EFFTPeakInterpolationMethod InterpolationMethod = EFFTPeakInterpolationMethod::Linear, EFFTWindowType WindowType = EFFTWindowType::Hann, float HopSize = 0);
+
+	/** Start spectrum analysis of the audio output. By leaving the Submix To Stop Analyzing blank, you can analyze the master output of the game. */
+	UFUNCTION(BlueprintCallable, Category = "Audio|Analysis", meta = (WorldContext = "WorldContextObject", AdvancedDisplay = 1))
+	static void StopAnalyzingOutput(const UObject* WorldContextObject, USoundSubmix* SubmixToStopAnalyzing = nullptr);
+
+	/** Start spectrum analysis of the audio output. By leaving the Submix To Analyze blank, you can analyze the master output of the game. */
+	UFUNCTION(BlueprintCallable, Category = "Audio|Analysis", meta = (WorldContext = "WorldContextObject", AdvancedDisplay = 3))
+	static void GetMagnitudeForFrequencies(const UObject* WorldContextObject, const TArray<float>& Frequencies, TArray<float>& Magnitudes, USoundSubmix* SubmixToAnalyze = nullptr);
+
+	/** Start spectrum analysis of the audio output. By leaving the Submix To Analyze blank, you can analyze the master output of the game. */
+	UFUNCTION(BlueprintCallable, Category = "Audio|Analysis", meta = (WorldContext = "WorldContextObject", AdvancedDisplay = 3))
+	static void GetPhaseForFrequencies(const UObject* WorldContextObject, const TArray<float>& Frequencies, TArray<float>& Phases, USoundSubmix* SubmixToAnalyze = nullptr);
+
 	/** Adds source effect entry to preset chain. Only effects the instance of the preset chain */
 	UFUNCTION(BlueprintCallable, Category = "Audio|Effects", meta = (WorldContext = "WorldContextObject"))
 	static void AddSourceEffectToPresetChain(const UObject* WorldContextObject, USoundEffectSourcePresetChain* PresetChain, FSourceEffectChainEntry Entry);
@@ -62,4 +125,7 @@ public:
 	UFUNCTION(BlueprintCallable, Category = "Audio|Effects", meta = (WorldContext = "WorldContextObject"))
 	static int32 GetNumberOfEntriesInSourceEffectChain(const UObject* WorldContextObject, USoundEffectSourcePresetChain* PresetChain);
 
+
+private:
+	static void PopulateSpectrumAnalyzerSettings(EFFTSize FFTSize, EFFTPeakInterpolationMethod InterpolationMethod, EFFTWindowType WindowType, float HopSize, Audio::FSpectrumAnalyzerSettings &OutSettings);
 };
\ No newline at end of file
diff --git a/Engine/Source/Runtime/AudioMixer/Public/AudioMixerDevice.h b/Engine/Source/Runtime/AudioMixer/Public/AudioMixerDevice.h
--- a/Engine/Source/Runtime/AudioMixer/Public/AudioMixerDevice.h
+++ b/Engine/Source/Runtime/AudioMixer/Public/AudioMixerDevice.h
@@ -120,6 +120,12 @@ namespace Audio
 		virtual void StopEnvelopeFollowing(USoundSubmix* InSubmix) override;
 		virtual void AddEnvelopeFollowerDelegate(USoundSubmix* InSubmix, const FOnSubmixEnvelopeBP& OnSubmixEnvelopeBP) override;
 
+		/** Submix Spectrum Analysis */
+		virtual void StartSpectrumAnalysis(USoundSubmix* InSubmix, const Audio::FSpectrumAnalyzerSettings& InSettings) override;
+		virtual void StopSpectrumAnalysis(USoundSubmix* InSubmix) override;
+		virtual void GetMagnitudesForFrequencies(USoundSubmix* InSubmix, const TArray<float>& InFrequencies, TArray<float>& OutMagnitudes);
+		virtual void GetPhasesForFrequencies(USoundSubmix* InSubmix, const TArray<float>& InFrequencies, TArray<float>& OutPhases);
+
 		// Submix buffer listener callbacks
 		virtual void RegisterSubmixBufferListener(ISubmixBufferListener* InSubmixBufferListener, USoundSubmix* InSubmix = nullptr) override;
 		virtual void UnregisterSubmixBufferListener(ISubmixBufferListener* InSubmixBufferListener, USoundSubmix* InSubmix = nullptr) override;
diff --git a/Engine/Source/Runtime/AudioMixer/Public/AudioMixerSubmix.h b/Engine/Source/Runtime/AudioMixer/Public/AudioMixerSubmix.h
--- a/Engine/Source/Runtime/AudioMixer/Public/AudioMixerSubmix.h
+++ b/Engine/Source/Runtime/AudioMixer/Public/AudioMixerSubmix.h
@@ -6,6 +6,7 @@
 #include "Sound/SoundSubmix.h"
 #include "Sound/SampleBuffer.h"
 #include "DSP/EnvelopeFollower.h"
+#include "DSP/SpectrumAnalyzer.h"
 
 class USoundEffectSubmix;
 
@@ -134,6 +135,20 @@ namespace Audio
 		// Adds an envelope follower delegate
 		void AddEnvelopeFollowerDelegate(const FOnSubmixEnvelopeBP& OnSubmixEnvelopeBP);
 
+		// Initializes a new FFT analyzer for this submix and immediately begins feeding audio to it.
+		void StartSpectrumAnalysis(const FSpectrumAnalyzerSettings& InSettings);
+
+		// Terminates whatever FFT Analyzer is being used for this submix.
+		void StopSpectrumAnalysis();
+
+		// Gets the most recent magnitude values for each corresponding value in InFrequencies (in Hz).
+		// This requires StartSpectrumAnalysis to be called first.
+		void GetMagnitudeForFrequencies(const TArray<float>& InFrequencies, TArray<float>& OutMagnitudes);
+
+		// Gets the most recent phase values for each corresponding value in InFrequencies (in Hz).
+		// This requires StartSpectrumAnalysis to be called first.
+		void GetPhaseForFrequencies(const TArray<float>& InFrequencies, TArray<float>& OutPhases);
+
 		// Broadcast the envelope value on the game thread
 		void BroadcastEnvelope();
 
@@ -141,6 +156,8 @@ namespace Audio
 		// Down mix the given buffer to the desired down mix channel count
 		void FormatChangeBuffer(const ESubmixChannelFormat NewChannelType, AlignedFloatBuffer& InBuffer, AlignedFloatBuffer& OutNewBuffer);
 
+		void MixBufferDownToMono(const AlignedFloatBuffer& InBuffer, int32 NumInputChannels, AlignedFloatBuffer& OutBuffer);
+
 		// Set up ambisonics encoder. Called when ambisonics settings are changed.
 		void SetUpAmbisonicsEncoder();
 
@@ -238,6 +255,12 @@ namespace Audio
 		int32 EnvelopeNumChannels;
 		FCriticalSection EnvelopeCriticalSection;
 
+		// Spectrum analyzer:
+		TUniquePtr<FSpectrumAnalyzer> SpectrumAnalyzer;
+		
+		// This buffer is used to downmix the submix output to mono before submitting it to the SpectrumAnalyzer.
+		AlignedFloatBuffer MonoMixBuffer;
+
 		// This buffer is encoded into for each source, then summed into the ambisonics buffer.
 		AlignedFloatBuffer InputAmbisonicsBuffer;
 
diff --git a/Engine/Source/Runtime/AudioMixer/Public/DSP/AudioFFT.h b/Engine/Source/Runtime/AudioMixer/Public/DSP/AudioFFT.h
--- a/Engine/Source/Runtime/AudioMixer/Public/DSP/AudioFFT.h
+++ b/Engine/Source/Runtime/AudioMixer/Public/DSP/AudioFFT.h
@@ -40,25 +40,75 @@ namespace Audio
 		 *                     Generally, set this to false if using this window with an STFT, but use true
 		 *                     if this window will be used on an entire, self-contained signal.
 		 */
-		FWindow(EWindowType InType, int32 InNumFrames, int32 InNumChannels, bool bIsPeriodic);
+		FWindow(EWindowType InType, int32 InNumFrames, int32 InNumChannels, bool bIsPeriodic)
+			: WindowType(InType)
+			, NumSamples(InNumFrames * InNumChannels)
+		{
+			checkf(NumSamples % 4 == 0, TEXT("For performance reasons, this window's length should be a multiple of 4."));
+			Generate(InNumFrames, InNumChannels, bIsPeriodic);
+		}
 
 		// Destructor. Releases memory used for window.
-		~FWindow();
+		~FWindow()
+		{
+		}
 
 		// Apply this window to InBuffer, which is expected to be an interleaved buffer with the same amount of frames
 		// and channels this window was constructed with.
-		void ApplyToBuffer(float* InBuffer);
+		void ApplyToBuffer(float* InBuffer)
+		{
+			if (WindowType == EWindowType::None)
+			{
+				return;
+			}
+
+			check(IsAligned<float*>(InBuffer, 4));
+			MultiplyBuffersInPlace(WindowBuffer.GetData(), InBuffer, NumSamples);
+		}
 
 	private:
+		EWindowType WindowType;
+		AlignedFloatBuffer WindowBuffer;
+		int32 NumSamples;
+
 		// Purposefully hidden constructor.
 		FWindow();
 
 		// Generate the window. Called on constructor.
-		void Generate(int32 NumFrames, int32 NumChannels, bool bIsPeriodic);
+		void Generate(int32 NumFrames, int32 NumChannels, bool bIsPeriodic)
+		{
+			if (WindowType == EWindowType::None)
+			{
+				return;
+			}
 
-		EWindowType WindowType;
-		AlignedFloatBuffer WindowBuffer;
-		int32 NumSamples;
+			WindowBuffer.Reset();
+			WindowBuffer.AddZeroed(NumSamples);
+
+			switch (WindowType)
+			{
+			case EWindowType::Hann:
+			{
+				GenerateHannWindow(WindowBuffer.GetData(), NumFrames, NumChannels, bIsPeriodic);
+				break;
+			}
+			case EWindowType::Hamming:
+			{
+				GenerateHammingWindow(WindowBuffer.GetData(), NumFrames, NumChannels, bIsPeriodic);
+				break;
+			}
+			case EWindowType::Blackman:
+			{
+				GenerateBlackmanWindow(WindowBuffer.GetData(), NumFrames, NumChannels, bIsPeriodic);
+				break;
+			}
+			default:
+			{
+				checkf(false, TEXT("Unknown window type!"));
+				break;
+			}
+			}
+		}
 	};
 
 	struct FFTTimeDomainData
diff --git a/Engine/Source/Runtime/AudioMixer/Public/DSP/Dsp.h b/Engine/Source/Runtime/AudioMixer/Public/DSP/Dsp.h
--- a/Engine/Source/Runtime/AudioMixer/Public/DSP/Dsp.h
+++ b/Engine/Source/Runtime/AudioMixer/Public/DSP/Dsp.h
@@ -553,6 +553,29 @@ namespace Audio
 			return ((int32)CurrentSlack) - NumSamples;
 		}
 
+		// When called, seeks the read or write cursor to only retain either the NumSamples latest data
+		// (if bRetainOldestSamples is false) or the NumSamples oldest data (if bRetainOldestSamples is true)
+		// in the buffer. Cannot be used to increase the capacity of this buffer.
+		void SetNum(uint32 NumSamples, bool bRetainOldestSamples = false)
+		{
+			check(NumSamples < Capacity);
+
+			if (bRetainOldestSamples)
+			{
+				WriteCounter.Set((ReadCounter.GetValue() + NumSamples) % Capacity);
+			}
+			else
+			{
+				int64 ReadCounterNum = ((int32)WriteCounter.GetValue()) - ((int32) NumSamples);
+				if (ReadCounterNum < 0)
+				{
+					ReadCounterNum = Capacity + ReadCounterNum;
+				}
+
+				ReadCounter.Set(ReadCounterNum);
+			}
+		}
+
 		// Get number of samples that can be popped off of the buffer.
 		uint32 Num()
 		{
diff --git a/Engine/Source/Runtime/AudioMixer/Public/DSP/SpectrumAnalyzer.h b/Engine/Source/Runtime/AudioMixer/Public/DSP/SpectrumAnalyzer.h
--- a/Engine/Source/Runtime/AudioMixer/Public/DSP/SpectrumAnalyzer.h
+++ b/Engine/Source/Runtime/AudioMixer/Public/DSP/SpectrumAnalyzer.h
@@ -4,13 +4,15 @@
 
 #include "CoreMinimal.h"
 #include "DSP/Dsp.h"
+#include "DSP/BufferVectorOperations.h"
 #include "DSP/AudioFFT.h"
 #include "Sound/SampleBuffer.h"
 
 namespace Audio
 {
-	
-	namespace SpectrumAnalyzerSettings
+	class FSpectrumAnalyzer;
+
+	struct AUDIOMIXER_API FSpectrumAnalyzerSettings
 	{
 		// Actual FFT size used. For FSpectrumAnalyzer, we never zero pad the input buffer.
 		enum class EFFTSize : uint16
@@ -33,28 +35,25 @@ namespace Audio
 			Quadratic
 		};
 
-		struct FSettings
-		{
-			EWindowType WindowType;
-			EFFTSize FFTSize;
-			EPeakInterpolationMethod InterpolationMethod;
-
-			/**
-			 * Hop size as a percentage of FFTSize.
-			 * 1.0 indicates a full hop.
-			 * Keeping this as 0.0 will use whatever hop size
-			 * can be used for WindowType to maintain COLA.
-			 */
-			float HopSize;
-
-			FSettings()
-				: WindowType(EWindowType::Hann)
-				, FFTSize(EFFTSize::Default)
-				, InterpolationMethod(EPeakInterpolationMethod::Linear)
-				, HopSize(0.0f)
-			{}
-		};
-	}
+		EWindowType WindowType;
+		EFFTSize FFTSize;
+		EPeakInterpolationMethod InterpolationMethod;
+
+		/**
+			* Hop size as a percentage of FFTSize.
+			* 1.0 indicates a full hop.
+			* Keeping this as 0.0 will use whatever hop size
+			* can be used for WindowType to maintain COLA.
+			*/
+		float HopSize;
+
+		FSpectrumAnalyzerSettings()
+			: WindowType(EWindowType::Hann)
+			, FFTSize(EFFTSize::Default)
+			, InterpolationMethod(EPeakInterpolationMethod::Linear)
+			, HopSize(0.0f)
+		{}
+	};
 
 	/**
 	 * This struct contains the output results from a singular FFT operation.
@@ -79,9 +78,9 @@ namespace Audio
 	{
 	public:
 		FSpectrumAnalyzerBuffer();
-		FSpectrumAnalyzerBuffer(const SpectrumAnalyzerSettings::FSettings& InSettings);
+		FSpectrumAnalyzerBuffer(const FSpectrumAnalyzerSettings& InSettings);
 
-		void Reset(const SpectrumAnalyzerSettings::FSettings& InSettings);
+		void Reset(const FSpectrumAnalyzerSettings& InSettings);
 
 		// Input. Used on analysis thread to lock a buffer to write to.
 		FSpectrumAnalyzerFrequencyVector* StartWorkOnBuffer();
@@ -106,6 +105,31 @@ namespace Audio
 		FCriticalSection BufferIndicesCriticalSection;
 	};
 
+	class FSpectrumAnalysisAsyncWorker : public FNonAbandonableTask
+	{
+	protected:
+		FSpectrumAnalyzer* Analyzer;
+		bool bUseLatestAudio;
+
+	public:
+		FSpectrumAnalysisAsyncWorker(FSpectrumAnalyzer* InAnalyzer, bool bInUseLatestAudio)
+			: Analyzer(InAnalyzer)
+			, bUseLatestAudio(bInUseLatestAudio)
+		{}
+
+		FORCEINLINE TStatId GetStatId() const
+		{
+			RETURN_QUICK_DECLARE_CYCLE_STAT(FSpectrumAnalysisAsyncWorker, STATGROUP_ThreadPoolAsyncTasks);
+		}
+
+		void DoWork();
+
+	private:
+		FSpectrumAnalysisAsyncWorker();
+	};
+
+	typedef FAsyncTask<FSpectrumAnalysisAsyncWorker> FSpectrumAnalyzerTask;
+
 	/**
 	 * Class built to be a rolling spectrum analyzer for arbitrary, monaural audio data.
 	 * Class is meant to scale accuracy with CPU and memory budgets.
@@ -120,17 +144,19 @@ namespace Audio
 
 		// If an instance is created using either of these constructors, Init() is not neccessary.
 		FSpectrumAnalyzer(float InSampleRate);
-		FSpectrumAnalyzer(const SpectrumAnalyzerSettings::FSettings& InSettings, float InSampleRate);
+		FSpectrumAnalyzer(const FSpectrumAnalyzerSettings& InSettings, float InSampleRate);
+
+		~FSpectrumAnalyzer();
 
 		// Initialize sample rate of analyzer if not known at time of construction
 		void Init(float InSampleRate);
-		void Init(const SpectrumAnalyzerSettings::FSettings& InSettings, float InSampleRate);
+		void Init(const FSpectrumAnalyzerSettings& InSettings, float InSampleRate);
 
 		// Update the settings used by this Spectrum Analyzer. Safe to call on any thread, but should not be called every tick.
-		void SetSettings(const SpectrumAnalyzerSettings::FSettings& InSettings);
+		void SetSettings(const FSpectrumAnalyzerSettings& InSettings);
 
 		// Get the current settings used by this Spectrum Analyzer.
-		void GetSettings(SpectrumAnalyzerSettings::FSettings& OutSettings);
+		void GetSettings(FSpectrumAnalyzerSettings& OutSettings);
 
 		// Samples magnitude (linearly) for a given frequency, in Hz.
 		float GetMagnitudeForFrequency(float InFrequency);
@@ -150,7 +176,10 @@ namespace Audio
 		bool PushAudio(const float* InBuffer, int32 NumSamples);
 
 		// Thread safe call to perform actual FFT. Returns true if it performed the FFT, false otherwise.
-		bool PerformAnalysisIfPossible();
+		// If bAsync is true, this function will kick off an async task.
+		// If bUseLatestAudio is set to true, this function will flush the entire input buffer, potentially losing data.
+		// Otherwise it will only consume enough samples necessary to perform a single FFT.
+		bool PerformAnalysisIfPossible(bool bUseLatestAudio = false, bool bAsync = false);
 
 		// Returns false if this instance of FSpectrumAnalyzer was constructed with the default constructor 
 		// and Init() has not been called yet.
@@ -162,10 +191,10 @@ namespace Audio
 		void ResetSettings();
 
 		// Called in GetMagnitudeForFrequency and GetPhaseForFrequency.
-		void PerformInterpolation(const FSpectrumAnalyzerFrequencyVector* InFrequencies, SpectrumAnalyzerSettings::EPeakInterpolationMethod InMethod, const float InFreq, float& OutReal, float& OutImag);
+		void PerformInterpolation(const FSpectrumAnalyzerFrequencyVector* InFrequencies, FSpectrumAnalyzerSettings::EPeakInterpolationMethod InMethod, const float InFreq, float& OutReal, float& OutImag);
 
 		// Cached current settings. Only actually used in ResetSettings().
-		SpectrumAnalyzerSettings::FSettings CurrentSettings;
+		FSpectrumAnalyzerSettings CurrentSettings;
 		volatile bool bSettingsWereUpdated;
 
 		volatile bool bIsInitialized;
@@ -183,5 +212,9 @@ namespace Audio
 
 		// if non-null, owns pointer to locked frequency vector we're using.
 		const FSpectrumAnalyzerFrequencyVector* LockedFrequencyVector;
+
+		// This is used if PerformAnalysisIfPossible is called
+		// with bAsync = true.
+		TUniquePtr<FSpectrumAnalyzerTask> AsyncAnalysisTask;
 	};
 }
diff --git a/Engine/Source/Runtime/Engine/Public/AudioDevice.h b/Engine/Source/Runtime/Engine/Public/AudioDevice.h
--- a/Engine/Source/Runtime/Engine/Public/AudioDevice.h
+++ b/Engine/Source/Runtime/Engine/Public/AudioDevice.h
@@ -15,6 +15,7 @@
 #include "Sound/SoundSourceBus.h"
 #include "Sound/AudioSettings.h"
 #include "AudioDeviceManager.h"
+#include "DSP/SpectrumAnalyzer.h"
 #include "EngineGlobals.h"
 
 class FAudioEffectsManager;
@@ -1204,6 +1205,26 @@ public:
 		UE_LOG(LogAudio, Error, TEXT("Envelope following submixes only works with the audio mixer. Please run using -audiomixer or set INI file to use submix recording."));
 	}
 
+	virtual void StartSpectrumAnalysis(USoundSubmix* InSubmix, const Audio::FSpectrumAnalyzerSettings& InSettings)
+	{
+		UE_LOG(LogAudio, Error, TEXT("Spectrum analysis of submixes only works with the audio mixer. Please run using -audiomixer or set INI file to use submix recording."));
+	}
+
+	virtual void StopSpectrumAnalysis(USoundSubmix* InSubmix)
+	{
+		UE_LOG(LogAudio, Error, TEXT("Spectrum analysis of submixes only works with the audio mixer. Please run using -audiomixer or set INI file to use submix recording."));
+	}
+
+	virtual void GetMagnitudesForFrequencies(USoundSubmix* InSubmix, const TArray<float>& InFrequencies, TArray<float>& OutMagnitudes)
+	{
+		UE_LOG(LogAudio, Error, TEXT("Spectrum analysis of submixes only works with the audio mixer. Please run using -audiomixer or set INI file to use submix recording."));
+	}
+
+	virtual void GetPhasesForFrequencies(USoundSubmix* InSubmix, const TArray<float>& InFrequencies, TArray<float>& OutPhases)
+	{
+		UE_LOG(LogAudio, Error, TEXT("Spectrum analysis of submixes only works with the audio mixer. Please run using -audiomixer or set INI file to use submix recording."));
+	}
+
 protected:
 	friend class FSoundSource;
 
diff --git a/Engine/Source/Runtime/MediaAssets/Private/Assets/MediaSoundComponent.cpp b/Engine/Source/Runtime/MediaAssets/Private/Assets/MediaSoundComponent.cpp
--- a/Engine/Source/Runtime/MediaAssets/Private/Assets/MediaSoundComponent.cpp
+++ b/Engine/Source/Runtime/MediaAssets/Private/Assets/MediaSoundComponent.cpp
@@ -439,7 +439,7 @@ int32 UMediaSoundComponent::OnGenerateAudio(float* OutAudio, int32 NumSamples)
 			}
 
 			// Launch an analysis task with this audio
-			(new FAutoDeleteAsyncTask<FMediaSoundComponentSpectrumAnalysisTask>(&SpectrumAnalyzer, &SpectrumAnalysisCounter))->StartBackgroundTask();
+			SpectrumAnalyzer.PerformAnalysisIfPossible(true, true);
 		}
 
 		SET_FLOAT_STAT(STAT_MediaUtils_MediaSoundComponentSync, FMath::Abs((Time - OutTime).GetTotalMilliseconds()));
@@ -467,25 +467,25 @@ void UMediaSoundComponent::SetEnableSpectralAnalysis(bool bInSpectralAnalysisEna
 
 void UMediaSoundComponent::SetSpectralAnalysisSettings(TArray<float> InFrequenciesToAnalyze, EMediaSoundComponentFFTSize InFFTSize)
 {
-	Audio::SpectrumAnalyzerSettings::EFFTSize SpectrumAnalyzerSize;
+	Audio::FSpectrumAnalyzerSettings::EFFTSize SpectrumAnalyzerSize;
 
 	switch (InFFTSize)
 	{
 		case EMediaSoundComponentFFTSize::Min_64: 
-			SpectrumAnalyzerSize = Audio::SpectrumAnalyzerSettings::EFFTSize::Min_64; 
+			SpectrumAnalyzerSize = Audio::FSpectrumAnalyzerSettings::EFFTSize::Min_64;
 			break;
 		
 		case EMediaSoundComponentFFTSize::Small_256: 
-			SpectrumAnalyzerSize = Audio::SpectrumAnalyzerSettings::EFFTSize::Small_256; 
+			SpectrumAnalyzerSize = Audio::FSpectrumAnalyzerSettings::EFFTSize::Small_256;
 			break;
 		
 		default:
 		case EMediaSoundComponentFFTSize::Medium_512:
-			SpectrumAnalyzerSize = Audio::SpectrumAnalyzerSettings::EFFTSize::Medium_512; 
+			SpectrumAnalyzerSize = Audio::FSpectrumAnalyzerSettings::EFFTSize::Medium_512;
 			break;
 
 		case EMediaSoundComponentFFTSize::Large_1024: 
-			SpectrumAnalyzerSize = Audio::SpectrumAnalyzerSettings::EFFTSize::Large_1024; 
+			SpectrumAnalyzerSize = Audio::FSpectrumAnalyzerSettings::EFFTSize::Large_1024;
 			break;
 	}
 
diff --git a/Engine/Source/Runtime/MediaAssets/Public/MediaSoundComponent.h b/Engine/Source/Runtime/MediaAssets/Public/MediaSoundComponent.h
--- a/Engine/Source/Runtime/MediaAssets/Public/MediaSoundComponent.h
+++ b/Engine/Source/Runtime/MediaAssets/Public/MediaSoundComponent.h
@@ -64,34 +64,6 @@ struct FMediaSoundComponentSpectralData
 	float Magnitude;
 };
 
-class FMediaSoundComponentSpectrumAnalysisTask : public FNonAbandonableTask
-{
-	friend class FAutoDeleteAsyncTask<FMediaSoundComponentSpectrumAnalysisTask>;
-
-	FMediaSoundComponentSpectrumAnalysisTask(Audio::FSpectrumAnalyzer* InAnalyzer, FThreadSafeCounter* InTaskCounter)
-		: Analyzer(InAnalyzer)
-		, TaskCounter(InTaskCounter)
-	{
-		TaskCounter->Increment();
-	}
-
-	void DoWork()
-	{
-		while (Analyzer->PerformAnalysisIfPossible());
-
-		TaskCounter->Decrement();
-	}
-
-	FORCEINLINE TStatId GetStatId() const
-	{
-		RETURN_QUICK_DECLARE_CYCLE_STAT(FTimeSynthSpectrumAnalysisTask, STATGROUP_ThreadPoolAsyncTasks);
-	}
-
-	Audio::FSpectrumAnalyzer* Analyzer;
-	FThreadSafeCounter* TaskCounter;
-};
-
-
 /**
  * Implements a sound component for playing a media player's audio output.
  */
@@ -302,8 +274,7 @@ private:
 
 	/** Spectrum analyzer used for anlayzing audio in media. */
 	Audio::FSpectrumAnalyzer SpectrumAnalyzer;
-	Audio::SpectrumAnalyzerSettings::FSettings SpectrumAnalyzerSettings;
-	FThreadSafeCounter SpectrumAnalysisCounter;
+	Audio::FSpectrumAnalyzerSettings SpectrumAnalyzerSettings;
 
 	/** Scratch buffer to mix in source audio to from decoder */
 	Audio::AlignedFloatBuffer AudioScratchBuffer;
