Fix for missing initial pre-cache buffer with audio mixer introduced in CL 4696296

diff --git a/Engine/Source/Runtime/AudioMixer/Private/AudioMixerSourceBuffer.cpp b/Engine/Source/Runtime/AudioMixer/Private/AudioMixerSourceBuffer.cpp
--- a/Engine/Source/Runtime/AudioMixer/Private/AudioMixerSourceBuffer.cpp
+++ b/Engine/Source/Runtime/AudioMixer/Private/AudioMixerSourceBuffer.cpp
@@ -179,6 +179,7 @@ namespace Audio
 
 		// Prepare the buffer for the PCM submission
 		SourceVoiceBuffers[0]->AudioData.Reset(NumSamplesPerBuffer);
+		SourceVoiceBuffers[0]->AudioData.AddUninitialized(NumSamplesPerBuffer);
 
 		RawPCMDataBuffer.GetNextBuffer(SourceVoiceBuffers[0].Get(), NumSamplesPerBuffer);
 
