Fix build error by downgrading log. DLL is only used in editor on import currently.

diff --git a/Engine/Source/Runtime/AudioMixer/Private/SoundFileIO/SoundFileIOManagerImpl.cpp b/Engine/Source/Runtime/AudioMixer/Private/SoundFileIO/SoundFileIOManagerImpl.cpp
--- a/Engine/Source/Runtime/AudioMixer/Private/SoundFileIO/SoundFileIOManagerImpl.cpp
+++ b/Engine/Source/Runtime/AudioMixer/Private/SoundFileIO/SoundFileIOManagerImpl.cpp
@@ -100,7 +100,7 @@ namespace Audio
 		SoundFileDllHandle = GetSoundFileDllHandle();
 		if (!SoundFileDllHandle)
 		{
-			UE_LOG(LogAudioMixer, Error, TEXT("Failed to load Sound File dll"));
+			UE_LOG(LogAudioMixer, Display, TEXT("Failed to load Sound File dll"));
 			return false;
 		}
 
