Unblocked RoboMerge

diff --git a/Engine/Source/Runtime/AudioMixer/Private/AudioMixerChannelMaps.cpp b/Engine/Source/Runtime/AudioMixer/Private/AudioMixerChannelMaps.cpp
--- a/Engine/Source/Runtime/AudioMixer/Private/AudioMixerChannelMaps.cpp
+++ b/Engine/Source/Runtime/AudioMixer/Private/AudioMixerChannelMaps.cpp
@@ -192,7 +192,7 @@ namespace Audio
 	static TArray<TArray<float>> ChannelMapCache;
 	static TArray<TArray<float>> VorbisChannelMapCache;
 
-	int32 FMixerDevice::GetChannelMapCacheId(const int32 NumSourceChannels, const int32 NumOutputChannels, const bool bIsCenterChannelOnly) const
+	int32 FMixerDevice::GetChannelMapCacheId(const int32 NumSourceChannels, const int32 NumOutputChannels, const bool bIsCenterChannelOnly)
 	{
 		int32 Index = (NumSourceChannels - 1) + AUDIO_MIXER_MAX_OUTPUT_CHANNELS * (NumOutputChannels - 1);
 		if (bIsCenterChannelOnly)
@@ -204,7 +204,14 @@ namespace Audio
 
 	void FMixerDevice::Get2DChannelMap(bool bIsVorbis, const ESubmixChannelFormat InSubmixChannelType, const int32 NumSourceChannels, const bool bIsCenterChannelOnly, Audio::AlignedFloatBuffer& OutChannelMap) const
 	{
+		AUDIO_MIXER_CHECK_GAME_THREAD(MixerDevice);
+
 		int32 NumOutputChannels = GetNumChannelsForSubmixFormat(InSubmixChannelType);
+		Get2DChannelMap(bIsVorbis, NumSourceChannels, NumOutputChannels, bIsCenterChannelOnly, OutChannelMap);
+	}
+
+	void FMixerDevice::Get2DChannelMap(bool bIsVorbis, const int32 NumSourceChannels, const int32 NumOutputChannels, const bool bIsCenterChannelOnly, Audio::AlignedFloatBuffer& OutChannelMap)
+	{
 		if (NumSourceChannels > 8 || NumOutputChannels > 8)
 		{
 			// Return a zero'd channel map buffer in the case of an unsupported channel configuration
@@ -230,7 +237,7 @@ namespace Audio
 		const int32 OutputChannelMapIndex = NumOutputChannels - 1;
 		check(OutputChannelMapIndex < ARRAY_COUNT(OutputChannelMaps));
 
-		float* Matrix = OutputChannelMaps[OutputChannelMapIndex];
+		float* RESTRICT Matrix = OutputChannelMaps[OutputChannelMapIndex];
 		check(Matrix != nullptr);
 
 		// Mono input sources have some special cases to take into account
diff --git a/Engine/Source/Runtime/AudioMixer/Private/AudioMixerSource.cpp b/Engine/Source/Runtime/AudioMixer/Private/AudioMixerSource.cpp
--- a/Engine/Source/Runtime/AudioMixer/Private/AudioMixerSource.cpp
+++ b/Engine/Source/Runtime/AudioMixer/Private/AudioMixerSource.cpp
@@ -847,9 +847,10 @@ namespace Audio
 				if (ChannelType != ESubmixChannelFormat::Ambisonics)
 				{
 					check(Buffer);
+					const uint32 NumChannels = Buffer->NumChannels;
 					if (ComputeChannelMap(ChannelType, Buffer->NumChannels, ChannelMapInfo.ChannelMap))
 					{
-						MixerSourceVoice->SetChannelMap(ChannelType, ChannelMapInfo.ChannelMap, bIs3D, WaveInstance->bCenterChannelOnly);
+						MixerSourceVoice->SetChannelMap(ChannelType, NumChannels, ChannelMapInfo.ChannelMap, bIs3D, WaveInstance->bCenterChannelOnly);
 					}
 				}
 			}
diff --git a/Engine/Source/Runtime/AudioMixer/Private/AudioMixerSourceManager.cpp b/Engine/Source/Runtime/AudioMixer/Private/AudioMixerSourceManager.cpp
--- a/Engine/Source/Runtime/AudioMixer/Private/AudioMixerSourceManager.cpp
+++ b/Engine/Source/Runtime/AudioMixer/Private/AudioMixerSourceManager.cpp
@@ -146,6 +146,14 @@ namespace Audio
 
 		MixerSources.Init(nullptr, NumTotalSources);
 
+		// Populate downmix array:
+		DownmixDataArray.Reset();
+		for (int32 Index = 0; Index < NumTotalSources; Index++)
+		{
+			DownmixDataArray.Emplace(2, MixerDevice->GetNumDeviceChannels(), NumOutputFrames);
+		}
+		
+
 		SourceInfos.AddDefaulted(NumTotalSources);
 
 		for (int32 i = 0; i < NumTotalSources; ++i)
@@ -183,7 +191,6 @@ namespace Audio
 			SourceInfo.SourceEnvelopeFollower = Audio::FEnvelopeFollower(MixerDevice->SampleRate, 10, 100, Audio::EPeakMode::Peak);
 			SourceInfo.SourceEnvelopeValue = 0.0f;
 			SourceInfo.bEffectTailsDone = false;
-			SourceInfo.PostEffectBuffers = nullptr;
 		
 			SourceInfo.bIs3D = false;
 			SourceInfo.bIsCenterChannelOnly = false;
@@ -198,12 +205,15 @@ namespace Audio
 			SourceInfo.bUseOcclusionPlugin = false;
 			SourceInfo.bUseReverbPlugin = false;
 			SourceInfo.bHasStarted = false;
-			SourceInfo.bIsDebugMode = false;
 			SourceInfo.bOutputToBusOnly = false;
 			SourceInfo.bIsVorbis = false;
 			SourceInfo.bIsBypassingLPF = false;
 			SourceInfo.bIsBypassingHPF = false;
 
+#if AUDIO_MIXER_ENABLE_DEBUG_MODE
+			SourceInfo.bIsDebugMode = false;
+#endif // AUDIO_MIXER_ENABLE_DEBUG_MODE
+
 			SourceInfo.NumInputChannels = 0;
 			SourceInfo.NumPostEffectChannels = 0;
 			SourceInfo.NumInputFrames = 0;
@@ -420,7 +430,6 @@ namespace Audio
 		SourceInfo.CurrentFrameIndex = 0;
 		SourceInfo.NumFramesPlayed = 0;
 		SourceInfo.StartTime = 0.0;
-		SourceInfo.PostEffectBuffers = nullptr;
 		SourceInfo.bIs3D = false;
 		SourceInfo.bIsCenterChannelOnly = false;
 		SourceInfo.bIsActive = false;
@@ -434,22 +443,18 @@ namespace Audio
 		SourceInfo.bUseOcclusionPlugin = false;
 		SourceInfo.bUseReverbPlugin = false;
 		SourceInfo.bHasStarted = false;
-		SourceInfo.bIsDebugMode = false;
 		SourceInfo.bOutputToBusOnly = false;
 		SourceInfo.bIsBypassingLPF = false;
 		SourceInfo.bIsBypassingHPF = false;
+
+#if AUDIO_MIXER_ENABLE_DEBUG_MODE
+		SourceInfo.bIsDebugMode = false;
 		SourceInfo.DebugName = FString();
+#endif //AUDIO_MIXER_ENABLE_DEBUG_MODE
+
 		SourceInfo.NumInputChannels = 0;
 		SourceInfo.NumPostEffectChannels = 0;
 
-		// Reset submix channel infos
-		for (int32 i = 0; i < (int32)ESubmixChannelFormat::Count; ++i)
-		{
-			SourceInfo.SubmixChannelInfo[i].OutputBuffer.Reset();
-			SourceInfo.SubmixChannelInfo[i].ChannelMapParam.Reset();
-			SourceInfo.SubmixChannelInfo[i].bUsed = false;
-		}
-
 		GameThreadInfo.bNeedsSpeakerMap[SourceId] = false;
 	}
 
@@ -706,6 +711,11 @@ namespace Audio
 
 			AUDIO_MIXER_CHECK(SourceInfo.SubmixSends.Num() == 0);
 
+			// Initialize a new downmix data:
+			check(SourceId < SourceInfos.Num());
+			const int32 SourceInputChannels = SourceInfo.bUseHRTFSpatializer ? 2 : SourceInfo.NumInputChannels;
+			FSourceDownmixData& DownmixData = InitializeDownmixForSource(SourceId, SourceInputChannels, MixerDevice->GetDeviceOutputChannels(), NumOutputFrames);
+
 			for (int32 i = 0; i < InitParams.SubmixSends.Num(); ++i)
 			{
 				const FMixerSourceSubmixSend& MixerSubmixSend = InitParams.SubmixSends[i];
@@ -719,8 +729,9 @@ namespace Audio
 					// Prepare output buffers and speaker map entries for every submix channel type
 					const ESubmixChannelFormat SubmixChannelType = SubmixPtr->GetSubmixChannels();
 
-					// Flag that we're going to be using this channel info entry
-					SourceInfo.SubmixChannelInfo[(int32)SubmixChannelType].bUsed = true;
+					FSubmixChannelTypeInfo& SubmixChannelInfo = GetChannelInfoForFormat(SubmixChannelType, DownmixData);
+
+					SubmixChannelInfo.bInUse = true;
 				}
 			}
 
@@ -942,37 +953,44 @@ namespace Audio
 		});
 	}
 
-	void FMixerSourceManager::SetChannelMap(const int32 SourceId, const ESubmixChannelFormat SubmixChannelType, const Audio::AlignedFloatBuffer& ChannelMap, const bool bInIs3D, const bool bInIsCenterChannelOnly)
+	void FMixerSourceManager::SetChannelMap(const int32 SourceId, const ESubmixChannelFormat SubmixChannelType, const uint32 NumInputChannels, const Audio::AlignedFloatBuffer& ChannelMap, const bool bInIs3D, const bool bInIsCenterChannelOnly)
 	{
 		AUDIO_MIXER_CHECK(SourceId < NumTotalSources);
 		AUDIO_MIXER_CHECK(GameThreadInfo.bIsBusy[SourceId]);
 		AUDIO_MIXER_CHECK_GAME_THREAD(MixerDevice);
 
-		AudioMixerThreadCommand([this, SourceId, SubmixChannelType, ChannelMap, bInIs3D, bInIsCenterChannelOnly]()
+		AudioMixerThreadCommand([this, SourceId, SubmixChannelType, NumInputChannels, ChannelMap, bInIs3D, bInIsCenterChannelOnly]()
 		{
 			AUDIO_MIXER_CHECK_AUDIO_PLAT_THREAD(MixerDevice);
 
 			check(NumOutputFrames > 0);
 
 			FSourceInfo& SourceInfo = SourceInfos[SourceId];
+			FSourceDownmixData& DownmixData = DownmixDataArray[SourceId];
+
+			if (DownmixData.NumInputChannels != NumInputChannels)
+			{
+				// This means that this source has been reinitialized as a different source while this command was in flight,
+				// In which case it is of no use to us. Exit.
+				return;
+			}
 
 			// Set whether or not this is a 3d channel map and if its center channel only. Used for reseting channel maps on device change.
 			SourceInfo.bIs3D = bInIs3D;
 			SourceInfo.bIsCenterChannelOnly = bInIsCenterChannelOnly;
 
-			FSubmixChannelTypeInfo& ChannelTypeInfo = SourceInfo.SubmixChannelInfo[(int32)SubmixChannelType];
-			ChannelTypeInfo.bUsed = true;
+			FSubmixChannelTypeInfo& ChannelTypeInfo = GetChannelInfoForFormat(SubmixChannelType, DownmixData);
+			ChannelTypeInfo.bInUse = true;
 
 			// Fix up the channel map in case device output count changed
 			if (SubmixChannelType == ESubmixChannelFormat::Device)
 			{
-				const int32 NumSourceChannels = SourceInfo.bUseHRTFSpatializer ? 2 : SourceInfo.NumInputChannels;
-				const int32 NumOutputChannels = MixerDevice->GetNumDeviceChannels();
-				const int32 ChannelMapSize = NumSourceChannels * NumOutputChannels;
+				const uint32 ChannelMapSize = ChannelTypeInfo.ChannelMap.CopySize / sizeof(float);
 
 				// If this is true, then the device changed while the command was in-flight
 				if (ChannelMap.Num() != ChannelMapSize)
 				{
+					// todo: investigate turning this into a stack array
 					Audio::AlignedFloatBuffer NewChannelMap;
 
 					// If 3d then just zero it out, we'll get another channel map shortly
@@ -984,22 +1002,23 @@ namespace Audio
 					// Otherwise, get an appropriate channel map for the new device configuration
 					else
 					{
-						MixerDevice->Get2DChannelMap(SourceInfo.bIsVorbis, ESubmixChannelFormat::Device, NumSourceChannels, bInIsCenterChannelOnly, NewChannelMap);
+						const uint32 NumOutputChannels = ChannelMapSize / NumInputChannels;
+						FMixerDevice::Get2DChannelMap(SourceInfo.bIsVorbis, NumInputChannels, NumOutputChannels, bInIsCenterChannelOnly, NewChannelMap);
 					}
 
 					// Make sure we've been flagged to be using this submix channel type entry
-					ChannelTypeInfo.ChannelMapParam.SetChannelMap(NewChannelMap, NumOutputChannels);
+					ChannelTypeInfo.ChannelMap.SetChannelMap(NewChannelMap.GetData());
 				}
 				else
 				{
 					GameThreadInfo.bNeedsSpeakerMap[SourceId] = false;
-					ChannelTypeInfo.ChannelMapParam.SetChannelMap(ChannelMap, NumOutputFrames);
+					ChannelTypeInfo.ChannelMap.SetChannelMap(ChannelMap.GetData());
 				}
 			}
 			else
 			{
 				// Since we're artificially mixing to this channel count, we don't need to deal with device reset
-				ChannelTypeInfo.ChannelMapParam.SetChannelMap(ChannelMap, NumOutputFrames);
+				ChannelTypeInfo.ChannelMap.SetChannelMap(ChannelMap.GetData());
 			}
 		});
 	}
@@ -1064,9 +1083,13 @@ namespace Audio
 				{
 					SourceInfo.SubmixSends.Add(InSubmixSend);
 
+					FSourceDownmixData& DownmixData = DownmixDataArray[SourceId];
+
 					// Flag that we're now using this submix channel info
-					ESubmixChannelFormat ChannelType = InSubmixPtr->GetSubmixChannels();
-					SourceInfo.SubmixChannelInfo[(int32)ChannelType].bUsed = true;
+					const ESubmixChannelFormat ChannelType = InSubmixPtr->GetSubmixChannels();
+					FSubmixChannelTypeInfo& ChannelInfo = GetChannelInfoForFormat(ChannelType, DownmixData);
+
+					ChannelInfo.bInUse = true;
 				}
 
 				InSubmixPtr->AddOrSetSourceVoice(MixerSources[SourceId], InSubmixSend.SendLevel);
@@ -1263,7 +1286,7 @@ namespace Audio
 			{
 				// Get the source's rendered bus data
 				const FMixerBus* Bus = Buses.Find(SourceInfo.BusId);
-				const float* BusBufferPtr = Bus->GetCurrentBusBuffer();
+				const float* RESTRICT BusBufferPtr = Bus->GetCurrentBusBuffer();
 
 				int32 NumFramesPlayed = NumOutputFrames;
 				if (SourceInfo.BusDurationFrames != INDEX_NONE)
@@ -1402,7 +1425,7 @@ namespace Audio
 		SourceInfo.DistanceAttenuationSourceStart = SourceInfo.DistanceAttenuationSourceDestination;
 	}
 
-	void FMixerSourceManager::ComputePluginAudio(FSourceInfo& SourceInfo, int32 SourceId, int32 NumSamples)
+	void FMixerSourceManager::ComputePluginAudio(FSourceInfo& SourceInfo, FSourceDownmixData& DownmixData, int32 SourceId, int32 NumSamples)
 	{
 		if (BypassAudioPluginsCvar)
 		{
@@ -1410,7 +1433,7 @@ namespace Audio
 			SourceInfo.NumPostEffectChannels = SourceInfo.NumInputChannels;
 
 			// Set the ptr to use for post-effect buffers:
-			SourceInfo.PostEffectBuffers = &SourceInfo.SourceBuffer;
+			DownmixData.PostEffectBuffers = &SourceInfo.SourceBuffer;
 
 			return;
 		}
@@ -1443,8 +1466,8 @@ namespace Audio
 			if (!MixerDevice->bReverbIsExternalSend)
 			{
 				// Copy the reverb-processed data back to the source buffer
-				SourceInfo.ReverbPluginOutputBuffer.Reset();
-				SourceInfo.ReverbPluginOutputBuffer.Append(SourceInfo.AudioPluginOutputData.AudioBuffer);
+				DownmixData.ReverbPluginOutputBuffer.Reset();
+				DownmixData.ReverbPluginOutputBuffer.Append(SourceInfo.AudioPluginOutputData.AudioBuffer);
 				bShouldMixInReverb = true;
 			}
 		}
@@ -1472,7 +1495,7 @@ namespace Audio
 			// Copy the occlusion-processed data back to the source buffer and mix with the reverb plugin output buffer
 			if (bShouldMixInReverb)
 			{
-				const float* ReverbPluginOutputBufferPtr = SourceInfo.ReverbPluginOutputBuffer.GetData();
+				const float* ReverbPluginOutputBufferPtr = DownmixData.ReverbPluginOutputBuffer.GetData();
 				const float* AudioPluginOutputDataPtr = SourceInfo.AudioPluginOutputData.AudioBuffer.GetData();
 
 				Audio::SumBuffers(ReverbPluginOutputBufferPtr, AudioPluginOutputDataPtr, PostDistanceAttenBufferPtr, NumSamples);
@@ -1484,7 +1507,7 @@ namespace Audio
 		}
 		else if (bShouldMixInReverb)
 		{
-			const float* ReverbPluginOutputBufferPtr = SourceInfo.ReverbPluginOutputBuffer.GetData();
+			const float* ReverbPluginOutputBufferPtr = DownmixData.ReverbPluginOutputBuffer.GetData();
 			Audio::MixInBufferFast(ReverbPluginOutputBufferPtr, PostDistanceAttenBufferPtr, NumSamples);
 		}
 
@@ -1521,13 +1544,13 @@ namespace Audio
 				SourceInfo.NumPostEffectChannels = SourceInfo.NumInputChannels;
 
 				// Set the ptr to use for post-effect buffers rather than the plugin output data (since the plugin won't have output audio data)
-				SourceInfo.PostEffectBuffers = &SourceInfo.SourceBuffer;
+				DownmixData.PostEffectBuffers = &SourceInfo.SourceBuffer;
 			}
 			else
 			{
 				// Otherwise, we are now a 2-channel file and should not be spatialized using normal 3d spatialization
 				SourceInfo.NumPostEffectChannels = 2;
-				SourceInfo.PostEffectBuffers = &SourceInfo.AudioPluginOutputData.AudioBuffer;
+				DownmixData.PostEffectBuffers = &SourceInfo.AudioPluginOutputData.AudioBuffer;
 			}
 		}
 		else
@@ -1536,7 +1559,298 @@ namespace Audio
 			SourceInfo.NumPostEffectChannels = SourceInfo.NumInputChannels;
 
 			// Set the ptr to use for post-effect buffers
-			SourceInfo.PostEffectBuffers = &SourceInfo.SourceBuffer;
+			DownmixData.PostEffectBuffers = &SourceInfo.SourceBuffer;
+		}
+	}
+
+	void FMixerSourceManager::ComputeDownmix3D(FSourceDownmixData& DownmixData)
+	{
+		// This enormous switch statement handles using the correct function for a given number of input and output channels.
+		// For 3D sources, we interpolate from ChannelStartGains to ChannelDestinationGains.
+		if (DownmixData.NumInputChannels == 1)
+		{
+			switch (DownmixData.NumDeviceChannels)
+			{
+			case 8:
+				Audio::MixMonoTo8ChannelsFast(*DownmixData.PostEffectBuffers, DownmixData.DeviceSubmixInfo.OutputBuffer, DownmixData.DeviceSubmixInfo.ChannelMap.ChannelStartGains, DownmixData.DeviceSubmixInfo.ChannelMap.ChannelDestinationGains);
+				break;
+			case 6:
+				Audio::MixMonoTo6ChannelsFast(*DownmixData.PostEffectBuffers, DownmixData.DeviceSubmixInfo.OutputBuffer, DownmixData.DeviceSubmixInfo.ChannelMap.ChannelStartGains, DownmixData.DeviceSubmixInfo.ChannelMap.ChannelDestinationGains);
+				break;
+			case 4:
+				Audio::MixMonoTo4ChannelsFast(*DownmixData.PostEffectBuffers, DownmixData.DeviceSubmixInfo.OutputBuffer, DownmixData.DeviceSubmixInfo.ChannelMap.ChannelStartGains, DownmixData.DeviceSubmixInfo.ChannelMap.ChannelDestinationGains);
+				break;
+			case 2:
+				Audio::MixMonoTo2ChannelsFast(*DownmixData.PostEffectBuffers, DownmixData.DeviceSubmixInfo.OutputBuffer, DownmixData.DeviceSubmixInfo.ChannelMap.ChannelStartGains, DownmixData.DeviceSubmixInfo.ChannelMap.ChannelDestinationGains);
+				break;
+			}
+
+			if (DownmixData.StereoSubmixInfo.bInUse)
+			{
+				Audio::MixMonoTo2ChannelsFast(*DownmixData.PostEffectBuffers, DownmixData.StereoSubmixInfo.OutputBuffer, DownmixData.StereoSubmixInfo.ChannelMap.ChannelStartGains, DownmixData.StereoSubmixInfo.ChannelMap.ChannelDestinationGains);
+			}
+
+			if (DownmixData.QuadSubmixInfo.bInUse)
+			{
+				Audio::MixMonoTo4ChannelsFast(*DownmixData.PostEffectBuffers, DownmixData.QuadSubmixInfo.OutputBuffer, DownmixData.QuadSubmixInfo.ChannelMap.ChannelStartGains, DownmixData.QuadSubmixInfo.ChannelMap.ChannelDestinationGains);
+			}
+
+			if (DownmixData.FiveOneSubmixInfo.bInUse)
+			{
+				Audio::MixMonoTo6ChannelsFast(*DownmixData.PostEffectBuffers, DownmixData.FiveOneSubmixInfo.OutputBuffer, DownmixData.FiveOneSubmixInfo.ChannelMap.ChannelStartGains, DownmixData.FiveOneSubmixInfo.ChannelMap.ChannelDestinationGains);
+			}
+
+			if (DownmixData.SevenOneSubmixInfo.bInUse)
+			{
+				Audio::MixMonoTo8ChannelsFast(*DownmixData.PostEffectBuffers, DownmixData.SevenOneSubmixInfo.OutputBuffer, DownmixData.SevenOneSubmixInfo.ChannelMap.ChannelStartGains, DownmixData.SevenOneSubmixInfo.ChannelMap.ChannelDestinationGains);
+			}
+		}
+		else if (DownmixData.NumInputChannels == 2)
+		{
+			switch (DownmixData.NumDeviceChannels)
+			{
+			case 8:
+				Audio::Mix2ChannelsTo8ChannelsFast(*DownmixData.PostEffectBuffers, DownmixData.DeviceSubmixInfo.OutputBuffer, DownmixData.DeviceSubmixInfo.ChannelMap.ChannelStartGains, DownmixData.DeviceSubmixInfo.ChannelMap.ChannelDestinationGains);
+				break;
+			case 6:
+				Audio::Mix2ChannelsTo6ChannelsFast(*DownmixData.PostEffectBuffers, DownmixData.DeviceSubmixInfo.OutputBuffer, DownmixData.DeviceSubmixInfo.ChannelMap.ChannelStartGains, DownmixData.DeviceSubmixInfo.ChannelMap.ChannelDestinationGains);
+				break;
+			case 4:
+				Audio::Mix2ChannelsTo4ChannelsFast(*DownmixData.PostEffectBuffers, DownmixData.DeviceSubmixInfo.OutputBuffer, DownmixData.DeviceSubmixInfo.ChannelMap.ChannelStartGains, DownmixData.DeviceSubmixInfo.ChannelMap.ChannelDestinationGains);
+				break;
+			case 2:
+				Audio::Mix2ChannelsTo2ChannelsFast(*DownmixData.PostEffectBuffers, DownmixData.DeviceSubmixInfo.OutputBuffer, DownmixData.DeviceSubmixInfo.ChannelMap.ChannelStartGains, DownmixData.DeviceSubmixInfo.ChannelMap.ChannelDestinationGains);
+				break;
+			}
+
+			if (DownmixData.StereoSubmixInfo.bInUse)
+			{
+				Audio::Mix2ChannelsTo2ChannelsFast(*DownmixData.PostEffectBuffers, DownmixData.StereoSubmixInfo.OutputBuffer, DownmixData.StereoSubmixInfo.ChannelMap.ChannelStartGains, DownmixData.StereoSubmixInfo.ChannelMap.ChannelDestinationGains);
+			}
+
+			if (DownmixData.QuadSubmixInfo.bInUse)
+			{
+				Audio::Mix2ChannelsTo4ChannelsFast(*DownmixData.PostEffectBuffers, DownmixData.QuadSubmixInfo.OutputBuffer, DownmixData.QuadSubmixInfo.ChannelMap.ChannelStartGains, DownmixData.QuadSubmixInfo.ChannelMap.ChannelDestinationGains);
+			}
+
+			if (DownmixData.FiveOneSubmixInfo.bInUse)
+			{
+				Audio::Mix2ChannelsTo6ChannelsFast(*DownmixData.PostEffectBuffers, DownmixData.FiveOneSubmixInfo.OutputBuffer, DownmixData.FiveOneSubmixInfo.ChannelMap.ChannelStartGains, DownmixData.FiveOneSubmixInfo.ChannelMap.ChannelDestinationGains);
+			}
+
+			if (DownmixData.SevenOneSubmixInfo.bInUse)
+			{
+				Audio::Mix2ChannelsTo8ChannelsFast(*DownmixData.PostEffectBuffers, DownmixData.SevenOneSubmixInfo.OutputBuffer, DownmixData.SevenOneSubmixInfo.ChannelMap.ChannelStartGains, DownmixData.SevenOneSubmixInfo.ChannelMap.ChannelDestinationGains);
+			}
+		}
+		else
+		{
+			// Use generic calls:
+			Audio::DownmixBuffer(DownmixData.NumInputChannels, DownmixData.NumDeviceChannels, *DownmixData.PostEffectBuffers, DownmixData.DeviceSubmixInfo.OutputBuffer, DownmixData.DeviceSubmixInfo.ChannelMap.ChannelStartGains, DownmixData.DeviceSubmixInfo.ChannelMap.ChannelDestinationGains);
+
+			if (DownmixData.StereoSubmixInfo.bInUse)
+			{
+				Audio::DownmixBuffer(DownmixData.NumInputChannels, 2, *DownmixData.PostEffectBuffers, DownmixData.StereoSubmixInfo.OutputBuffer, DownmixData.StereoSubmixInfo.ChannelMap.ChannelStartGains, DownmixData.StereoSubmixInfo.ChannelMap.ChannelDestinationGains);
+			}
+
+			if (DownmixData.QuadSubmixInfo.bInUse)
+			{
+				Audio::DownmixBuffer(DownmixData.NumInputChannels, 4, *DownmixData.PostEffectBuffers, DownmixData.QuadSubmixInfo.OutputBuffer, DownmixData.QuadSubmixInfo.ChannelMap.ChannelStartGains, DownmixData.QuadSubmixInfo.ChannelMap.ChannelDestinationGains);
+			}
+
+			if (DownmixData.FiveOneSubmixInfo.bInUse)
+			{
+				Audio::DownmixBuffer(DownmixData.NumInputChannels, 6, *DownmixData.PostEffectBuffers, DownmixData.FiveOneSubmixInfo.OutputBuffer, DownmixData.FiveOneSubmixInfo.ChannelMap.ChannelStartGains, DownmixData.FiveOneSubmixInfo.ChannelMap.ChannelDestinationGains);
+			}
+
+			if (DownmixData.SevenOneSubmixInfo.bInUse)
+			{
+				Audio::DownmixBuffer(DownmixData.NumInputChannels, 8, *DownmixData.PostEffectBuffers, DownmixData.SevenOneSubmixInfo.OutputBuffer, DownmixData.SevenOneSubmixInfo.ChannelMap.ChannelStartGains, DownmixData.SevenOneSubmixInfo.ChannelMap.ChannelDestinationGains);
+			}
+		}
+	}
+
+	void FMixerSourceManager::ComputeDownmix2D(FSourceDownmixData& DownmixData)
+	{
+		// This enormous switch statement handles using the correct function for a given number of input and output channels.
+		// For 2D sources, we just apply the gain matrix in ChannelDestionationGains with no interpolation.
+		if (DownmixData.NumInputChannels == 1)
+		{
+			switch (DownmixData.NumDeviceChannels)
+			{
+			case 8:
+				Audio::MixMonoTo8ChannelsFast(*DownmixData.PostEffectBuffers, DownmixData.DeviceSubmixInfo.OutputBuffer, DownmixData.DeviceSubmixInfo.ChannelMap.ChannelDestinationGains);
+				break;
+			case 6:
+				Audio::MixMonoTo6ChannelsFast(*DownmixData.PostEffectBuffers, DownmixData.DeviceSubmixInfo.OutputBuffer, DownmixData.DeviceSubmixInfo.ChannelMap.ChannelDestinationGains);
+				break;
+			case 4:
+				Audio::MixMonoTo4ChannelsFast(*DownmixData.PostEffectBuffers, DownmixData.DeviceSubmixInfo.OutputBuffer, DownmixData.DeviceSubmixInfo.ChannelMap.ChannelDestinationGains);
+				break;
+			case 2:
+				Audio::MixMonoTo2ChannelsFast(*DownmixData.PostEffectBuffers, DownmixData.DeviceSubmixInfo.OutputBuffer, DownmixData.DeviceSubmixInfo.ChannelMap.ChannelDestinationGains);
+				break;
+			}
+
+			if (DownmixData.StereoSubmixInfo.bInUse)
+			{
+				Audio::MixMonoTo2ChannelsFast(*DownmixData.PostEffectBuffers, DownmixData.StereoSubmixInfo.OutputBuffer, DownmixData.StereoSubmixInfo.ChannelMap.ChannelDestinationGains);
+			}
+
+			if (DownmixData.QuadSubmixInfo.bInUse)
+			{
+				Audio::MixMonoTo4ChannelsFast(*DownmixData.PostEffectBuffers, DownmixData.QuadSubmixInfo.OutputBuffer, DownmixData.QuadSubmixInfo.ChannelMap.ChannelDestinationGains);
+			}
+
+			if (DownmixData.FiveOneSubmixInfo.bInUse)
+			{
+				Audio::MixMonoTo6ChannelsFast(*DownmixData.PostEffectBuffers, DownmixData.FiveOneSubmixInfo.OutputBuffer, DownmixData.FiveOneSubmixInfo.ChannelMap.ChannelDestinationGains);
+			}
+
+			if (DownmixData.SevenOneSubmixInfo.bInUse)
+			{
+				Audio::MixMonoTo8ChannelsFast(*DownmixData.PostEffectBuffers, DownmixData.SevenOneSubmixInfo.OutputBuffer, DownmixData.SevenOneSubmixInfo.ChannelMap.ChannelDestinationGains);
+			}
+		}
+		else if (DownmixData.NumInputChannels == 2)
+		{
+			switch (DownmixData.NumDeviceChannels)
+			{
+			case 8:
+				Audio::Mix2ChannelsTo8ChannelsFast(*DownmixData.PostEffectBuffers, DownmixData.DeviceSubmixInfo.OutputBuffer, DownmixData.DeviceSubmixInfo.ChannelMap.ChannelDestinationGains);
+				break;
+			case 6:
+				Audio::Mix2ChannelsTo6ChannelsFast(*DownmixData.PostEffectBuffers, DownmixData.DeviceSubmixInfo.OutputBuffer, DownmixData.DeviceSubmixInfo.ChannelMap.ChannelDestinationGains);
+				break;
+			case 4:
+				Audio::Mix2ChannelsTo4ChannelsFast(*DownmixData.PostEffectBuffers, DownmixData.DeviceSubmixInfo.OutputBuffer, DownmixData.DeviceSubmixInfo.ChannelMap.ChannelDestinationGains);
+				break;
+			case 2:
+				Audio::Mix2ChannelsTo2ChannelsFast(*DownmixData.PostEffectBuffers, DownmixData.DeviceSubmixInfo.OutputBuffer, DownmixData.DeviceSubmixInfo.ChannelMap.ChannelDestinationGains);
+				break;
+			}
+
+			if (DownmixData.StereoSubmixInfo.bInUse)
+			{
+				Audio::Mix2ChannelsTo2ChannelsFast(*DownmixData.PostEffectBuffers, DownmixData.StereoSubmixInfo.OutputBuffer, DownmixData.StereoSubmixInfo.ChannelMap.ChannelDestinationGains);
+			}
+
+			if (DownmixData.QuadSubmixInfo.bInUse)
+			{
+				Audio::Mix2ChannelsTo4ChannelsFast(*DownmixData.PostEffectBuffers, DownmixData.QuadSubmixInfo.OutputBuffer, DownmixData.QuadSubmixInfo.ChannelMap.ChannelDestinationGains);
+			}
+
+			if (DownmixData.FiveOneSubmixInfo.bInUse)
+			{
+				Audio::Mix2ChannelsTo6ChannelsFast(*DownmixData.PostEffectBuffers, DownmixData.FiveOneSubmixInfo.OutputBuffer, DownmixData.FiveOneSubmixInfo.ChannelMap.ChannelDestinationGains);
+			}
+
+			if (DownmixData.SevenOneSubmixInfo.bInUse)
+			{
+				Audio::Mix2ChannelsTo8ChannelsFast(*DownmixData.PostEffectBuffers, DownmixData.SevenOneSubmixInfo.OutputBuffer, DownmixData.SevenOneSubmixInfo.ChannelMap.ChannelDestinationGains);
+			}
+		}
+		else
+		{
+			// Use generic calls:
+			Audio::DownmixBuffer(DownmixData.NumInputChannels, DownmixData.NumDeviceChannels, *DownmixData.PostEffectBuffers, DownmixData.DeviceSubmixInfo.OutputBuffer, DownmixData.DeviceSubmixInfo.ChannelMap.ChannelDestinationGains);
+
+			if (DownmixData.StereoSubmixInfo.bInUse)
+			{
+				Audio::DownmixBuffer(DownmixData.NumInputChannels, 2, *DownmixData.PostEffectBuffers, DownmixData.StereoSubmixInfo.OutputBuffer, DownmixData.StereoSubmixInfo.ChannelMap.ChannelDestinationGains);
+			}
+
+			if (DownmixData.QuadSubmixInfo.bInUse)
+			{
+				Audio::DownmixBuffer(DownmixData.NumInputChannels, 4, *DownmixData.PostEffectBuffers, DownmixData.QuadSubmixInfo.OutputBuffer, DownmixData.QuadSubmixInfo.ChannelMap.ChannelDestinationGains);
+			}
+
+			if (DownmixData.FiveOneSubmixInfo.bInUse)
+			{
+				Audio::DownmixBuffer(DownmixData.NumInputChannels, 6, *DownmixData.PostEffectBuffers, DownmixData.FiveOneSubmixInfo.OutputBuffer, DownmixData.FiveOneSubmixInfo.ChannelMap.ChannelDestinationGains);
+
+			}
+
+			if (DownmixData.SevenOneSubmixInfo.bInUse)
+			{
+				Audio::DownmixBuffer(DownmixData.NumInputChannels, 8, *DownmixData.PostEffectBuffers, DownmixData.SevenOneSubmixInfo.OutputBuffer, DownmixData.SevenOneSubmixInfo.ChannelMap.ChannelDestinationGains);
+			}
+		}
+	}
+
+	FMixerSourceManager::FSourceDownmixData& FMixerSourceManager::InitializeDownmixForSource(const int32 SourceId, const int32 NumInputChannels, const int32 NumOutputChannels, const int32 InNumOutputFrames)
+	{
+		DownmixDataArray[SourceId].ResetData(NumInputChannels);
+		return DownmixDataArray[SourceId];
+	}
+
+	const FMixerSourceManager::FSubmixChannelTypeInfo& FMixerSourceManager::GetChannelInfoForFormat(const ESubmixChannelFormat InFormat, const FSourceDownmixData& InDownmixData) const
+	{
+		switch (InFormat)
+		{
+		case ESubmixChannelFormat::Device:
+		{
+			return InDownmixData.DeviceSubmixInfo;
+		}
+		case ESubmixChannelFormat::SevenDotOne:
+		{
+			return InDownmixData.SevenOneSubmixInfo;
+		}
+		case ESubmixChannelFormat::FiveDotOne:
+		{
+			return InDownmixData.FiveOneSubmixInfo;
+		}
+		case ESubmixChannelFormat::Quad:
+		{
+			return InDownmixData.QuadSubmixInfo;
+		}
+		case ESubmixChannelFormat::Stereo:
+		{
+			return InDownmixData.StereoSubmixInfo;
+		}
+		case ESubmixChannelFormat::Ambisonics:
+		{
+			return InDownmixData.AmbisonicsSubmixInfo;
+		}
+		default:
+			check(false);
+			return InDownmixData.DeviceSubmixInfo;
+			break;
+		}
+	}
+
+	FMixerSourceManager::FSubmixChannelTypeInfo& FMixerSourceManager::GetChannelInfoForFormat(const ESubmixChannelFormat InFormat, FSourceDownmixData& InDownmixData)
+	{
+		switch (InFormat)
+		{
+		case ESubmixChannelFormat::Device:
+		{
+			return InDownmixData.DeviceSubmixInfo;
+		}
+		case ESubmixChannelFormat::SevenDotOne:
+		{
+			return InDownmixData.SevenOneSubmixInfo;
+		}
+		case ESubmixChannelFormat::FiveDotOne:
+		{
+			return InDownmixData.FiveOneSubmixInfo;
+		}
+		case ESubmixChannelFormat::Quad:
+		{
+			return InDownmixData.QuadSubmixInfo;
+		}
+		case ESubmixChannelFormat::Stereo:
+		{
+			return InDownmixData.StereoSubmixInfo;
+		}
+		case ESubmixChannelFormat::Ambisonics:
+		{
+			return InDownmixData.AmbisonicsSubmixInfo;
+		}
+		default:
+			check(false);
+			return InDownmixData.DeviceSubmixInfo;
+			break;
 		}
 	}
 
@@ -1561,6 +1875,8 @@ namespace Audio
 				continue;
 			}
 
+			FSourceDownmixData& DownmixData = DownmixDataArray[SourceId];
+
 			// Copy and store the current state of the pre-distance attenuation buffer before we feed it through our source effects
 			// This is used by pre-effect sends
 			if (SourceInfo.BusSends[(int32)EBusSendType::PreEffect].Num() > 0)
@@ -1681,7 +1997,7 @@ namespace Audio
 					int32 SampleIndex = SourceInfo.NumInputChannels * Frame;
 
 					// Apply filters, if necessary.
-					if (LPFFreq < MAX_FILTER_FREQUENCY)
+					if (LPFFreq < (MAX_FILTER_FREQUENCY - KINDA_SMALL_NUMBER))
 					{
 						//If we stopped processing the low pass filter, we need to clear the filter's memory to prevent pops.
 						if (SourceInfo.bIsBypassingLPF)
@@ -1698,7 +2014,7 @@ namespace Audio
 						FMemory::Memcpy(&PostDistanceAttenBufferPtr[SampleIndex], &PreDistanceAttenBufferPtr[SampleIndex], SourceInfo.NumInputChannels * sizeof(float));
 					}
 
-					if (!DisableHPFilteringCvar && HPFFreq > 0.0f)
+					if (!DisableHPFilteringCvar && HPFFreq > KINDA_SMALL_NUMBER)
 					{
 						//If we stopped processing the low pass filter, we need to clear the filter's memory to prevent pops.
 						if (SourceInfo.bIsBypassingHPF)
@@ -1723,7 +2039,7 @@ namespace Audio
 				ApplyDistanceAttenuation(SourceInfo, NumSamples);
 
 				// Send source audio to plugins
-				ComputePluginAudio(SourceInfo, SourceId, NumSamples);
+				ComputePluginAudio(SourceInfo, DownmixData, SourceId, NumSamples);
 			}
 			else if (DisableFilteringCvar)
 			{
@@ -1734,7 +2050,7 @@ namespace Audio
 				ApplyDistanceAttenuation(SourceInfo, NumSamples);
 
 				// Send source audio to plugins
-				ComputePluginAudio(SourceInfo, SourceId, NumSamples);
+				ComputePluginAudio(SourceInfo, DownmixData, SourceId, NumSamples);
 			}
 
 			// Check the source effect tails condition
@@ -1772,82 +2088,13 @@ namespace Audio
 				continue;
 			}
 
-			// Perform output buffer computation for each submix channel type output
-			for (int32 SubmixInfoId = 0; SubmixInfoId < (int32)ESubmixChannelFormat::Count; ++SubmixInfoId)
+			if (SourceInfo.bIs3D)
 			{
-				FSubmixChannelTypeInfo& ChannelInfo = SourceInfo.SubmixChannelInfo[SubmixInfoId];
-
-				if (!ChannelInfo.bUsed)
-				{
-					continue;
-				}
-
-				ESubmixChannelFormat ChannelFormat = (ESubmixChannelFormat)SubmixInfoId;
-				if (ChannelFormat == ESubmixChannelFormat::Ambisonics)
-				{
-					// there are no channel maps for ambisonics format since it's a "pass through" to the ambisonics submix
-					ChannelInfo.OutputBuffer.Reset();
-					ChannelInfo.OutputBuffer.AddZeroed(NumOutputFrames * SourceInfo.NumInputChannels);
-
-					// Simply copy the post-effect buffers to the output buffer
-					FMemory::Memcpy((void*)ChannelInfo.OutputBuffer.GetData(), (void*)SourceInfo.PostEffectBuffers->GetData(), SourceInfo.PostEffectBuffers->Num()*sizeof(float));
-				}
-				else
-				{
-					int32 NumOutputChannels = MixerDevice->GetNumChannelsForSubmixFormat((ESubmixChannelFormat)SubmixInfoId);
-
-					ChannelInfo.OutputBuffer.Reset();
-					ChannelInfo.OutputBuffer.AddZeroed(NumOutputFrames * NumOutputChannels);
-
-					// If we're paused, then early return now
-					if (SourceInfo.bIsPaused)
-					{
-						continue;
-					}
-
-					float* OutputBufferPtr = ChannelInfo.OutputBuffer.GetData();
-					float* PostEffetBufferPtr = SourceInfo.PostEffectBuffers->GetData();
-
-					// Apply speaker mapping
-					for (int32 Frame = 0; Frame < NumOutputFrames; ++Frame)
-					{
-						const int32 PostEffectChannels = SourceInfo.NumPostEffectChannels;
-
-						float SourceSampleValue = 0.0f;
-
-						// Make sure that our channel map is appropriate for the source channel and output channel count!
-						ChannelInfo.ChannelMapParam.UpdateChannelMap();
-
-						// For each source channel, compute the output channel mapping
-						for (int32 SourceChannel = 0; SourceChannel < PostEffectChannels; ++SourceChannel)
-						{
-							const int32 SourceSampleIndex = Frame * PostEffectChannels + SourceChannel;
-							SourceSampleValue = PostEffetBufferPtr[SourceSampleIndex];
-
-							for (int32 OutputChannel = 0; OutputChannel < NumOutputChannels; ++OutputChannel)
-							{
-								// Look up the channel map value (maps input channels to output channels) for the source
-								// This is the step that either applies the spatialization algorithm or just maps a 2d sound
-								const int32 ChannelMapIndex = NumOutputChannels * SourceChannel + OutputChannel;
-								const float ChannelMapValue = ChannelInfo.ChannelMapParam.GetChannelValue(ChannelMapIndex);
-
-								// If we have a non-zero sample value, write it out. Note that most 3d audio channel maps
-								// for surround sound will result in 0.0 sample values so this branch should save a bunch of multiplies + adds
-								if (ChannelMapValue > 0.0f)
-								{
-									// Scale the input source sample for this source channel value
-									const float SampleValue = SourceSampleValue * ChannelMapValue;
-									const int32 OutputSampleIndex = Frame * NumOutputChannels + OutputChannel;
-
-									OutputBufferPtr[OutputSampleIndex] += SampleValue;
-								}
-							}
-						}
-					}
-
-					// Reset the channel map param interpolation
-					ChannelInfo.ChannelMapParam.ResetInterpolation();
-				}
+				ComputeDownmix3D(DownmixDataArray[SourceId]);
+			}
+			else
+			{
+				ComputeDownmix2D(DownmixDataArray[SourceId]);
 			}
 		}
 	}
@@ -1904,18 +2151,19 @@ namespace Audio
 			// Don't need to mix into submixes if the source is paused
 			if (!SourceInfo.bIsPaused && !SourceInfo.bIsDone && SourceInfo.bIsPlaying)
 			{
-				const FSubmixChannelTypeInfo& ChannelInfo = SourceInfo.SubmixChannelInfo[(int32)InSubmixChannelType];
+				const FSourceDownmixData& DownmixData = DownmixDataArray[SourceId];
+				const FSubmixChannelTypeInfo& ChannelInfo = GetChannelInfoForFormat(InSubmixChannelType, DownmixData);
 // 				if (InSubmixChannelType == ESubmixChannelFormat::Ambisonics && SourceInfo.NumInputChannels == 1)
 // 				{
 // 					// TODO: encode 3d sources (non-ambisonics files) to ambisonics
 // 				}
 
-				check(ChannelInfo.bUsed);
-				check(ChannelInfo.OutputBuffer.Num() == OutWetBuffer.Num());
+				const float* RESTRICT SourceOutputBufferPtr = ChannelInfo.OutputBuffer.GetData();
 
-				const float* SourceOutputBufferPtr = ChannelInfo.OutputBuffer.GetData();
-				const int32 OutWetBufferSize = OutWetBuffer.Num();
-				float* OutWetBufferPtr = OutWetBuffer.GetData();
+				// TODO: Figure out a fix for this race condition on device swap:
+				// const int32 OutWetBufferSize = OutWetBuffer.Num();
+				const int32 OutWetBufferSize = FMath::Min<float>(OutWetBuffer.Num(), ChannelInfo.OutputBuffer.Num());
+				float* RESTRICT OutWetBufferPtr = OutWetBuffer.GetData();
 
 				Audio::MixInBufferFast(SourceOutputBufferPtr, OutWetBufferPtr, OutWetBufferSize, SendLevel);
 			}
@@ -1924,43 +2172,49 @@ namespace Audio
 
 	void FMixerSourceManager::UpdateDeviceChannelCount(const int32 InNumOutputChannels)
 	{
-		NumOutputSamples = NumOutputFrames * MixerDevice->GetNumDeviceChannels();
-
-		// Update all source's to appropriate channel maps
-		for (int32 SourceId = 0; SourceId < NumTotalSources; ++SourceId)
+		AudioMixerThreadCommand([this, InNumOutputChannels]()
 		{
-			FSourceInfo& SourceInfo = SourceInfos[SourceId];
-
-			// Don't need to do anything if it's not active
-			if (!SourceInfo.bIsActive)
-			{
-				continue;
-			}
-
-			FSubmixChannelTypeInfo& ChannelTypeInfo = SourceInfo.SubmixChannelInfo[(int32)ESubmixChannelFormat::Device];
+			NumOutputSamples = NumOutputFrames * MixerDevice->GetNumDeviceChannels();
 
-			if (ChannelTypeInfo.bUsed)
+			// Update all source's to appropriate channel maps
+			for (int32 SourceId = 0; SourceId < NumTotalSources; ++SourceId)
 			{
-				SourceInfo.ScratchChannelMap.Reset();
-				const int32 NumSourceChannels = SourceInfo.bUseHRTFSpatializer ? 2 : SourceInfo.NumInputChannels;
+				FSourceInfo& SourceInfo = SourceInfos[SourceId];
 
-				// If this is a 3d source, then just zero out the channel map, it'll cause a temporary blip
-				// but it should reset in the next tick
-				if (SourceInfo.bIs3D)
+				// Don't need to do anything if it's not active
+				if (!SourceInfo.bIsActive)
 				{
-					GameThreadInfo.bNeedsSpeakerMap[SourceId] = true;
-					SourceInfo.ScratchChannelMap.AddZeroed(NumSourceChannels * InNumOutputChannels);
+					continue;
 				}
-				// If it's a 2D sound, then just get a new channel map appropriate for the new device channel count
-				else
+
+				FSourceDownmixData& DownmixData = DownmixDataArray[SourceId];
+				DownmixData.ResetNumberOfDeviceChannels(InNumOutputChannels);
+
+				FSubmixChannelTypeInfo& ChannelTypeInfo = GetChannelInfoForFormat(ESubmixChannelFormat::Device, DownmixData);
+
+				if (ChannelTypeInfo.bInUse)
 				{
 					SourceInfo.ScratchChannelMap.Reset();
-					MixerDevice->Get2DChannelMap(SourceInfo.bIsVorbis, ESubmixChannelFormat::Device, NumSourceChannels, SourceInfo.bIsCenterChannelOnly, SourceInfo.ScratchChannelMap);
-				}
+					const int32 NumSourceChannels = SourceInfo.bUseHRTFSpatializer ? 2 : SourceInfo.NumInputChannels;
+
+					// If this is a 3d source, then just zero out the channel map, it'll cause a temporary blip
+					// but it should reset in the next tick
+					if (SourceInfo.bIs3D)
+					{
+						GameThreadInfo.bNeedsSpeakerMap[SourceId] = true;
+						SourceInfo.ScratchChannelMap.AddZeroed(NumSourceChannels * InNumOutputChannels);
+					}
+					// If it's a 2D sound, then just get a new channel map appropriate for the new device channel count
+					else
+					{
+						SourceInfo.ScratchChannelMap.Reset();
+						MixerDevice->Get2DChannelMap(SourceInfo.bIsVorbis, ESubmixChannelFormat::Device, NumSourceChannels, SourceInfo.bIsCenterChannelOnly, SourceInfo.ScratchChannelMap);
+					}
 			
-				ChannelTypeInfo.ChannelMapParam.SetChannelMap(SourceInfo.ScratchChannelMap, NumOutputFrames);
+					ChannelTypeInfo.ChannelMap.SetChannelMap(SourceInfo.ScratchChannelMap.GetData());
+				}
 			}
-		}
+		});
 	}
 
 	void FMixerSourceManager::UpdateSourceEffectChain(const uint32 InSourceEffectChainId, const TArray<FSourceEffectChainEntry>& InSourceEffectChain, const bool bPlayEffectChainTails)
diff --git a/Engine/Source/Runtime/AudioMixer/Private/AudioMixerSourceManager.h b/Engine/Source/Runtime/AudioMixer/Private/AudioMixerSourceManager.h
--- a/Engine/Source/Runtime/AudioMixer/Private/AudioMixerSourceManager.h
+++ b/Engine/Source/Runtime/AudioMixer/Private/AudioMixerSourceManager.h
@@ -138,71 +138,40 @@ namespace Audio
 		{}
 	};
 
-	class FSourceChannelMap
+	struct FSourceChannelMap
 	{
-	public:
-		FSourceChannelMap() {}
+		alignas(16) float ChannelStartGains[AUDIO_MIXER_MAX_OUTPUT_CHANNELS * AUDIO_MIXER_MAX_OUTPUT_CHANNELS];
+		alignas(16) float ChannelDestinationGains[AUDIO_MIXER_MAX_OUTPUT_CHANNELS * AUDIO_MIXER_MAX_OUTPUT_CHANNELS];
 
-		FORCEINLINE void Reset()
-		{
-			ChannelValues.Reset();
-		}
+		// This is the number of bytes the gain array is using:
+		// (Number of input channels * number of output channels) * sizeof float.
+		int32 CopySize;
 
-		FORCEINLINE void SetChannelMap(const Audio::AlignedFloatBuffer& ChannelMap, const int32 InNumInterpFrames)
+		FSourceChannelMap(int32 InNumInChannels, int32 InNumOutChannels) 
+			: CopySize(InNumInChannels * InNumOutChannels * sizeof(float))
 		{
-			if (ChannelValues.Num() != ChannelMap.Num())
-			{
-				ChannelValues.Reset();
-				for (int32 i = 0; i < ChannelMap.Num(); ++i)
-				{
-					ChannelValues.Add(FParam());
-					ChannelValues[i].SetValue(ChannelMap[i], InNumInterpFrames);
-				}
-			}
-			else
-			{
-				for (int32 i = 0; i < ChannelMap.Num(); ++i)
-				{
-					ChannelValues[i].SetValue(ChannelMap[i], InNumInterpFrames);
-				}
-			}
+			checkSlow(InNumInChannels <= AUDIO_MIXER_MAX_OUTPUT_CHANNELS);
+			checkSlow(InNumOutChannels <= AUDIO_MIXER_MAX_OUTPUT_CHANNELS);
 		}
 
-		FORCEINLINE void UpdateChannelMap()
+		FORCEINLINE void Reset(int32 InNumInChannels, int32 InNumOutChannels)
 		{
-			const int32 NumChannelValues = ChannelValues.Num();
-			for (int32 i = 0; i < NumChannelValues; ++i)
-			{
-				ChannelValues[i].Update();
-			}
+			CopySize = InNumInChannels * InNumOutChannels * sizeof(float);
+			FMemory::Memzero(ChannelDestinationGains, CopySize);
 		}
 
-		FORCEINLINE void ResetInterpolation()
+		FORCEINLINE void SetChannelMap(const float* RESTRICT InChannelGains)
 		{
-			const int32 NumChannelValues = ChannelValues.Num();
-			for (int32 i = 0; i < NumChannelValues; ++i)
-			{
-				ChannelValues[i].Reset();
-			}
+			// TODO: See if you can find a way to assign this without going back to memory.
+			FMemory::Memcpy(ChannelStartGains, ChannelDestinationGains, CopySize);
+			FMemory::Memcpy(ChannelDestinationGains, InChannelGains, CopySize);
 		}
 
-		FORCEINLINE float GetChannelValue(int ChannelIndex)
-		{
-			return ChannelValues[ChannelIndex].GetValue();
-		}
-
-		FORCEINLINE void PadZeroes(const int32 ToSize, const int32 InNumInterpFrames)
+	private:
+		FSourceChannelMap()
+			: CopySize(0)
 		{
-			int32 CurrentSize = ChannelValues.Num();
-			for (int32 i = CurrentSize; i < ToSize; ++i)
-			{
-				ChannelValues.Add(FParam());
-				ChannelValues[i].SetValue(0.0f, InNumInterpFrames);
-			}
 		}
-
-	private:
-		TArray<FParam> ChannelValues;
 	};
 
 	struct FSourceManagerInitParams
@@ -238,7 +207,7 @@ namespace Audio
 		void SetVolume(const int32 SourceId, const float Volume);
 		void SetDistanceAttenuation(const int32 SourceId, const float DistanceAttenuation);
 		void SetSpatializationParams(const int32 SourceId, const FSpatializationParams& InParams);
-		void SetChannelMap(const int32 SourceId, const ESubmixChannelFormat SubmixChannelType, const Audio::AlignedFloatBuffer& InChannelMap, const bool bInIs3D, const bool bInIsCenterChannelOnly);
+		void SetChannelMap(const int32 SourceId, const ESubmixChannelFormat SubmixChannelType, const uint32 NumInputChannels, const Audio::AlignedFloatBuffer& InChannelMap, const bool bInIs3D, const bool bInIsCenterChannelOnly);
 		void SetLPFFrequency(const int32 SourceId, const float Frequency);
 		void SetHPFFrequency(const int32 SourceId, const float Frequency);
 
@@ -345,18 +314,75 @@ namespace Audio
 
 		struct FSubmixChannelTypeInfo
 		{
-			// Channel map parameter
-			FSourceChannelMap ChannelMapParam;
+			FSourceChannelMap ChannelMap;
+			Audio::AlignedFloatBuffer OutputBuffer;
+			bool bInUse;
+
+			FSubmixChannelTypeInfo(uint32 InNumInChannels, uint32 InNumOutputChannels, uint32 NumFrames)
+				: ChannelMap(InNumInChannels, InNumOutputChannels)
+				, bInUse(false)
+			{
+				OutputBuffer.Reset();
+				OutputBuffer.AddUninitialized(NumFrames * InNumOutputChannels);
+			}
+
+			void Reset(uint32 InNumInChannels, uint32 InNumOutputChannels, uint32 NumFrames)
+			{
+				ChannelMap.Reset(InNumInChannels, InNumOutputChannels);
+
+				OutputBuffer.Reset();
+				OutputBuffer.AddUninitialized(NumFrames * InNumOutputChannels);
+			}
+		};
+
+		struct FSourceDownmixData
+		{
+			// Output data, after computing a block of sample data, this is read back from mixers
+			Audio::AlignedFloatBuffer ReverbPluginOutputBuffer;
+			Audio::AlignedFloatBuffer* PostEffectBuffers;
 
-			// Output buffer based on channel map param
-			TArray<float> OutputBuffer;
+			// Data needed for outputting to submixes for each channel configuration.
+			FSubmixChannelTypeInfo DeviceSubmixInfo;
+			FSubmixChannelTypeInfo StereoSubmixInfo;
+			FSubmixChannelTypeInfo QuadSubmixInfo;
+			FSubmixChannelTypeInfo FiveOneSubmixInfo;
+			FSubmixChannelTypeInfo SevenOneSubmixInfo;
+			FSubmixChannelTypeInfo AmbisonicsSubmixInfo;
+
+			uint32 NumInputChannels;
+			const uint32 NumFrames;
+			uint32 NumDeviceChannels;
+
+			FSourceDownmixData(uint32 SourceNumChannels, uint32 NumDeviceOutputChannels, uint32 InNumFrames)
+				: DeviceSubmixInfo(FSubmixChannelTypeInfo(SourceNumChannels, NumDeviceOutputChannels, InNumFrames))
+				, StereoSubmixInfo(FSubmixChannelTypeInfo(SourceNumChannels, 2, InNumFrames))
+				, QuadSubmixInfo(FSubmixChannelTypeInfo(SourceNumChannels, 4, InNumFrames))
+				, FiveOneSubmixInfo(FSubmixChannelTypeInfo(SourceNumChannels, 6, InNumFrames))
+				, SevenOneSubmixInfo(FSubmixChannelTypeInfo(SourceNumChannels, 8, InNumFrames))
+				, AmbisonicsSubmixInfo(FSubmixChannelTypeInfo(SourceNumChannels, 4, InNumFrames))
+				, NumInputChannels(SourceNumChannels)
+				, NumFrames(InNumFrames)
+				, NumDeviceChannels(NumDeviceOutputChannels)
+			{
+			}
 
-			// Whether or not this channel type is used
-			bool bUsed;
+			void ResetNumberOfDeviceChannels(const uint32 NumDeviceOutputChannels)
+			{
+				NumDeviceChannels = NumDeviceOutputChannels;
+				DeviceSubmixInfo.Reset(NumInputChannels, NumDeviceOutputChannels, NumFrames);
+			}
 
-			FSubmixChannelTypeInfo()
-				: bUsed(false)
-			{}
+			void ResetData(const uint32 InNumInputChannels)
+			{
+				NumInputChannels = InNumInputChannels;
+
+				DeviceSubmixInfo.Reset(NumInputChannels, NumDeviceChannels, NumFrames);
+				StereoSubmixInfo.Reset(NumInputChannels, NumDeviceChannels, NumFrames);
+				QuadSubmixInfo.Reset(NumInputChannels, NumDeviceChannels, NumFrames);
+				FiveOneSubmixInfo.Reset(NumInputChannels, NumDeviceChannels, NumFrames);
+				SevenOneSubmixInfo.Reset(NumInputChannels, NumDeviceChannels, NumFrames);
+				AmbisonicsSubmixInfo.Reset(NumInputChannels, NumDeviceChannels, NumFrames);
+			}
 		};
 
 		struct FSourceInfo
@@ -432,13 +458,6 @@ namespace Audio
 			FSpatializationParams SpatParams;
 			Audio::AlignedFloatBuffer ScratchChannelMap;
 
-			// Output data, after computing a block of sample data, this is read back from mixers
-			Audio::AlignedFloatBuffer ReverbPluginOutputBuffer;
-			Audio::AlignedFloatBuffer* PostEffectBuffers;
-
-			// Data needed for outputting to submixes
-			FSubmixChannelTypeInfo SubmixChannelInfo[(int32) ESubmixChannelFormat::Count];
-
 			// State management
 			uint8 bIs3D:1;
 			uint8 bIsCenterChannelOnly:1;
@@ -457,9 +476,6 @@ namespace Audio
 			uint8 bIsVorbis:1;
 			uint8 bIsBypassingLPF:1;
 			uint8 bIsBypassingHPF:1;
-			uint8 bIsDebugMode:1;
-
-			FString DebugName;
 
 			// Source format info
 			int32 NumInputChannels;
@@ -468,16 +484,34 @@ namespace Audio
 
 			// ID for associated Audio Component if there is one, 0 otherwise
 			uint64 AudioComponentID;
+
+#if AUDIO_MIXER_ENABLE_DEBUG_MODE
+			uint8 bIsDebugMode : 1;
+			FString DebugName;
+#endif // AUDIO_MIXER_ENABLE_DEBUG_MODE
 		};
 
-		void ApplyDistanceAttenuation(FSourceInfo& InSourceInfo, int32 NumSamples);
-		void ComputePluginAudio(FSourceInfo& InSourceInfo, int32 SourceId, int32 NumSamples);
+		static void ApplyDistanceAttenuation(FSourceInfo& InSourceInfo, int32 NumSamples);
+		void ComputePluginAudio(FSourceInfo& InSourceInfo, FSourceDownmixData& DownmixData, int32 SourceId, int32 NumSamples);
+
+		static void ComputeDownmix3D(FSourceDownmixData& DownmixData);
+		static void ComputeDownmix2D(FSourceDownmixData& DownmixData);
+
+		// This function is effectively equivalent to calling DownmixDataArray.EmplaceAt_GetRef(args...), but bypasses that function's intrinsic call to AddUninitialized.
+		FSourceDownmixData& InitializeDownmixForSource(const int32 SourceId, const int32 NumInputChannels, const int32 NumOutputChannels, const int32 InNumOutputFrames);
+
+		const FSubmixChannelTypeInfo& GetChannelInfoForFormat(const ESubmixChannelFormat InFormat, const FSourceDownmixData& InDownmixData) const;
+		FSubmixChannelTypeInfo& GetChannelInfoForFormat(const ESubmixChannelFormat InFormat, FSourceDownmixData& InDownmixData);
 
 		// Array of listener transforms
 		TArray<FTransform> ListenerTransforms;
 
 		// Array of source infos.
 		TArray<FSourceInfo> SourceInfos;
+		
+		// These structs are used for guaranteed vectorization when downmixing
+		// sources.
+		TArray<FSourceDownmixData> DownmixDataArray;
 
 		// Map of bus object Id's to bus data. 
 		TMap<uint32, FMixerBus> Buses;
@@ -493,7 +527,7 @@ namespace Audio
 		{
 			TArray<int32> FreeSourceIndices;
 			TArray<bool> bIsBusy;
-			TArray <bool> bNeedsSpeakerMap;
+			TArray<bool> bNeedsSpeakerMap;
 			TArray<bool> bIsDebugMode;
 		} GameThreadInfo;
 
diff --git a/Engine/Source/Runtime/AudioMixer/Private/AudioMixerSourceVoice.cpp b/Engine/Source/Runtime/AudioMixer/Private/AudioMixerSourceVoice.cpp
--- a/Engine/Source/Runtime/AudioMixer/Private/AudioMixerSourceVoice.cpp
+++ b/Engine/Source/Runtime/AudioMixer/Private/AudioMixerSourceVoice.cpp
@@ -141,11 +141,11 @@ namespace Audio
 		}
 	}
 
-	void FMixerSourceVoice::SetChannelMap(ESubmixChannelFormat InChannelType, const Audio::AlignedFloatBuffer& InChannelMap, const bool bInIs3D, const bool bInIsCenterChannelOnly)
+	void FMixerSourceVoice::SetChannelMap(ESubmixChannelFormat InChannelType, const uint32 NumInputChannels, const Audio::AlignedFloatBuffer& InChannelMap, const bool bInIs3D, const bool bInIsCenterChannelOnly)
 	{
 		AUDIO_MIXER_CHECK_GAME_THREAD(MixerDevice);
 
-		SourceManager->SetChannelMap(SourceId, InChannelType, InChannelMap, bInIs3D, bInIsCenterChannelOnly);
+		SourceManager->SetChannelMap(SourceId, InChannelType, NumInputChannels, InChannelMap, bInIs3D, bInIsCenterChannelOnly);
 	}
 
 	void FMixerSourceVoice::SetSpatializationParams(const FSpatializationParams& InParams)
diff --git a/Engine/Source/Runtime/AudioMixer/Private/AudioMixerSourceVoice.h b/Engine/Source/Runtime/AudioMixer/Private/AudioMixerSourceVoice.h
--- a/Engine/Source/Runtime/AudioMixer/Private/AudioMixerSourceVoice.h
+++ b/Engine/Source/Runtime/AudioMixer/Private/AudioMixerSourceVoice.h
@@ -48,7 +48,7 @@ namespace Audio
 		void SetHPFFrequency(const float InFrequency);
 
 		// Sets the source voice's channel map (2d or 3d).
-		void SetChannelMap(ESubmixChannelFormat InChannelType, const Audio::AlignedFloatBuffer& InChannelMap, const bool bInIs3D, const bool bInIsCenterChannelOnly);
+		void SetChannelMap(ESubmixChannelFormat InChannelType, const uint32 NumInputChannels, const Audio::AlignedFloatBuffer& InChannelMap, const bool bInIs3D, const bool bInIsCenterChannelOnly);
 
 		// Sets params used by HRTF spatializer
 		void SetSpatializationParams(const FSpatializationParams& InParams);
diff --git a/Engine/Source/Runtime/AudioMixer/Private/DSP/BufferVectorOperations.cpp b/Engine/Source/Runtime/AudioMixer/Private/DSP/BufferVectorOperations.cpp
--- a/Engine/Source/Runtime/AudioMixer/Private/DSP/BufferVectorOperations.cpp
+++ b/Engine/Source/Runtime/AudioMixer/Private/DSP/BufferVectorOperations.cpp
@@ -1,6 +1,7 @@
 // Copyright 1998-2019 Epic Games, Inc. All Rights Reserved.
 
 #include "DSP/BufferVectorOperations.h"
+#include "AudioMixer.h"
 
 #define AUDIO_USE_SIMD 1
 
@@ -60,7 +61,7 @@ namespace Audio
 		MultiplyBufferByConstantInPlace(InBuffer.GetData(), InBuffer.Num(), InGain);
 	}
 
-	void MultiplyBufferByConstantInPlace(float* InBuffer, int32 NumSamples, float InGain)
+	void MultiplyBufferByConstantInPlace(float* RESTRICT InBuffer, int32 NumSamples, float InGain)
 	{
 		const VectorRegister Gain = VectorLoadFloat1(&InGain);
 
@@ -77,7 +78,7 @@ namespace Audio
 		FadeBufferFast(OutFloatBuffer.GetData(), OutFloatBuffer.Num(), StartValue, EndValue);
 	}
 
-	void FadeBufferFast(float* OutFloatBuffer, int32 NumSamples, const float StartValue, const float EndValue)
+	void FadeBufferFast(float* RESTRICT OutFloatBuffer, int32 NumSamples, const float StartValue, const float EndValue)
 	{
 		checkf(IsAligned<float*>(OutFloatBuffer, 4), TEXT("Memory must be aligned to use vector operations."));
 		checkf(NumSamples % 4 == 0, TEXT("Please use a buffer size that is a multiple of 4."));
@@ -155,7 +156,7 @@ namespace Audio
 		MixInBufferFast(InFloatBuffer.GetData(), BufferToSumTo.GetData(), InFloatBuffer.Num(), Gain);
 	}
 
-	void MixInBufferFast(const float* InFloatBuffer, float* BufferToSumTo, int32 NumSamples, const float Gain)
+	void MixInBufferFast(const float* RESTRICT InFloatBuffer, float* RESTRICT BufferToSumTo, int32 NumSamples, const float Gain)
 	{
 		checkf(IsAligned<const float*>(InFloatBuffer, 4), TEXT("Memory must be aligned to use vector operations."));
 		checkf(IsAligned<float*>(BufferToSumTo, 4), TEXT("Memory must be aligned to use vector operations."));
@@ -179,7 +180,7 @@ namespace Audio
 #endif
 	}
 
-	void MixInBufferFast(const float* InFloatBuffer, float* BufferToSumTo, int32 NumSamples)
+	void MixInBufferFast(const float* RESTRICT InFloatBuffer, float* RESTRICT BufferToSumTo, int32 NumSamples)
 	{
 		checkf(IsAligned<const float*>(InFloatBuffer, 4), TEXT("Memory must be aligned to use vector operations."));
 		checkf(IsAligned<float*>(BufferToSumTo, 4), TEXT("Memory must be aligned to use vector operations."));
@@ -206,7 +207,7 @@ namespace Audio
 		SumBuffers(InFloatBuffer1.GetData(), InFloatBuffer2.GetData(), OutputBuffer.GetData(), OutputBuffer.Num());
 	}
 
-	void SumBuffers(const float* InFloatBuffer1, const float* InFloatBuffer2, float* OutputBuffer, int32 NumSamples)
+	void SumBuffers(const float* RESTRICT InFloatBuffer1, const float* RESTRICT InFloatBuffer2, float* RESTRICT OutputBuffer, int32 NumSamples)
 	{
 		checkf(IsAligned<const float*>(InFloatBuffer1, 4), TEXT("Memory must be aligned to use vector operations."));
 		checkf(IsAligned<const float*>(InFloatBuffer2, 4), TEXT("Memory must be aligned to use vector operations."));
@@ -235,7 +236,7 @@ namespace Audio
 		MultiplyBuffersInPlace(InFloatBuffer.GetData(), BufferToMultiply.GetData(), BufferToMultiply.Num());
 	}
 
-	void MultiplyBuffersInPlace(const float* InFloatBuffer, float* BufferToMultiply, int32 NumSamples)
+	void MultiplyBuffersInPlace(const float* RESTRICT InFloatBuffer, float* RESTRICT BufferToMultiply, int32 NumSamples)
 	{
 		checkf(IsAligned<const float*>(InFloatBuffer, 4), TEXT("Memory must be aligned to use vector operations."));
 		checkf(IsAligned<const float*>(BufferToMultiply, 4), TEXT("Memory must be aligned to use vector operations."));
@@ -256,7 +257,7 @@ namespace Audio
 		return GetMagnitude(Buffer.GetData(), Buffer.Num());
 	}
 
-	float GetMagnitude(const float* Buffer, int32 NumSamples)
+	float GetMagnitude(const float* RESTRICT Buffer, int32 NumSamples)
 	{
 		checkf(NumSamples % 4 == 0, TEXT("Please use a buffer size that is a multiple of 4."));
 
@@ -293,7 +294,7 @@ namespace Audio
 		return GetAverageAmplitude(Buffer.GetData(), Buffer.Num());
 	}
 
-	float GetAverageAmplitude(const float* Buffer, int32 NumSamples)
+	float GetAverageAmplitude(const float* RESTRICT Buffer, int32 NumSamples)
 	{
 		checkf(NumSamples % 4 == 0, TEXT("Please use a buffer size that is a multiple of 4."));
 
@@ -319,4 +320,1116 @@ namespace Audio
 		return (PartionedSums[0] + PartionedSums[1] + PartionedSums[2] + PartionedSums[3]) / NumSamples;
 #endif
 	}
+
+	/**
+	 * CHANNEL MIXING OPERATIONS:
+	 * To understand these functions, it's best that you have prior experience reading SIMD code.
+	 * These functions are all variations on component-wise matrix multiplies. 
+	 * There are two types of functions below:
+	 * 
+	 * Apply[N]ChannelGain:
+	 * These are all in-place multiplies of an N-length gain vector and an N-length frame.
+	 * There are two flavors of every variant of this function: The non-interpolating form (which takes a single gain matrix)
+	 * And the interpolating form (which takes a start gain matrix and interpolates to the end gain matrix over the given number of frames).
+	 * All non-interpolating forms of these functions use the following steps:
+	 *    1. Create a const GainVector, or series of GainVectors, that maps to the multiplies required for each iteration.
+	 *    2. In a loop:
+	 *           i.   load a frame or number of frames into a vector register or series of vector registers (these are named Result).
+	 *           ii.  perform a vector multiply on result with the corresponding gain vector.
+	 *           iii. store the result vector in the same position in the buffer we loaded from.
+	 *
+	 * The interpolating forms of these functions use the following steps:
+	 *    1. Initialize a non-const GainVector, or series of GainVectors, from StartGains, that maps to the multiplies required for each iteration.
+	 *    2. Compute the amount we add to GainVector for each iteration to reach Destination Gains and store it in the const GainDeltasVector.
+	 *    3. In a loop:
+	 *           i.   load a frame or number of frames into a vector register or series of vector registers (these are named Result).
+	 *           ii.  perform a vector multiply on result with the corresponding gain vector.
+	 *           iii. store the result vector in the same position in the buffer we loaded from.
+	 *           iv.  increment each GainVector by it's corresponding GainDeltasVector.
+	 *
+	 *
+	 * MixMonoTo[N]ChannelsFast and Mix2ChannelsTo[N]ChannelsFast:
+	 * These, like Apply[N]ChannelGain, all have non-interpolating and interpolating forms.
+	 * All non-interpolating forms of these functions use the following steps:
+	 *    1. Create a const GainVector, or series of GainVectors, that maps to the multiplies required for each input channel for each iteration.
+	 *    2. In a loop:
+	 *           i.   load a frame or number of frames into a const vector register or series of const vector registers (these are named Input).
+	 *           ii.  perform a vector multiply on input with the corresponding gain vector and store the result in a new vector or series of vectors named Result.
+	 *           iii. if there is a second input channel, store the results of the following MultiplyAdd operation to Results: (Gain Vectors for second channel) * (Input vectors for second channel) + (Result vectors from step ii).
+	 *
+	 * Interpolating forms of these functions use the following steps:
+	 *    1. Initialize a non-const GainVector, or series of GainVectors, from StartGains, that maps to the multiplies required for each input channel for each iteration.
+	 *    2. Compute the amount we add to each GainVector for each iteration to reach the vector's corresponding DestinationGains and store it in a corresponding GainDeltaVector.
+	 *    3. In a loop:
+	 *           i.   load a frame or number of frames into a const vector register or series of const vector registers (these are named Input).
+	 *           ii.  perform a vector multiply on input with the corresponding gain vector and store the result in a new vector or series of vectors named Result.
+	 *           iii. if there is a second input channel, store the results of the following MultiplyAdd operation to Results: (Gain Vectors for second channel) * (Input vectors for second channel) + (Result vectors from step ii).
+	 *           iv.  increment each GainVector by it's corresponding GainDeltasVector.
+	 * 
+	 * DETERMINING THE VECTOR LAYOUT FOR EACH FUNCTION:
+	 * For every variant of Mix[N]ChannelsTo[N]ChannelsFast, we use the least common multiple of the number of output channels and the SIMD vector length (4) to calulate the length of our matrix.
+	 * For example, MixMonoTo4ChannelsFast can use a single VectorRegister for each variable. GainVector's values are [g0, g1, g2, g3], input channels are mapped to [i0, i0, i0, i0], and output channels are mapped to [o0, o1, o2, o3].
+	 * MixMonoTo8ChannelsFast has an LCM of 8, so we use two VectorRegister for each variable. This results in the following layout:
+	 * GainVector1:   [g0, g1, g2, g3] GainVector2:   [g4, g5, g6, g7]
+	 * InputVector1:  [i0, i0, i0, i0] InputVector2:  [i0, i0, i0, i0]
+	 * ResultVector1: [o0, o1, o2, o3] ResultVector2: [o4, o5, o6, o7]
+	 *
+	 * The general naming convention for vector variables is [Name]Vector[VectorIndex] for MixMonoTo[N]ChannelsFast functions.
+	 * For Mix2ChannelsTo[N]ChannelsFast functions, the naming convention for vector variables is [Name]Vector[VectorIndex][InputChannelIndex].
+	 *
+	 * For clarity, the layout of vectors for each function variant is given in a block comment above that function.
+	 */
+
+	void Apply2ChannelGain(AlignedFloatBuffer& StereoBuffer, const float* RESTRICT Gains)
+	{
+		Apply2ChannelGain(StereoBuffer.GetData(), StereoBuffer.Num(), Gains);
+	}
+
+	void Apply2ChannelGain(float* RESTRICT StereoBuffer, int32 NumSamples, const float* RESTRICT Gains)
+	{
+		const VectorRegister GainVector = VectorLoadFloat2(Gains);
+		
+		for (int32 i = 0; i < NumSamples; i += 4)
+		{
+			VectorRegister Result = VectorLoadAligned(&StereoBuffer[i]);
+			Result = VectorMultiply(Result, GainVector);
+			VectorStoreAligned(Result, &StereoBuffer[i]);
+		}
+	}
+
+	void Apply2ChannelGain(AlignedFloatBuffer& StereoBuffer, const float* RESTRICT StartGains, const float* RESTRICT EndGains)
+	{
+		Apply2ChannelGain(StereoBuffer.GetData(), StereoBuffer.Num(), StartGains, EndGains);
+	}
+
+	void Apply2ChannelGain(float* RESTRICT StereoBuffer, int32 NumSamples, const float* RESTRICT StartGains, const float* RESTRICT EndGains)
+	{
+		// Initialize GainVector at StartGains and compute GainDeltasVector:
+		VectorRegister GainVector = VectorLoadFloat2(StartGains);
+		const VectorRegister DestinationVector = VectorLoadFloat2(EndGains);
+		const VectorRegister NumFramesVector = VectorSetFloat1(NumSamples / 4);
+		const VectorRegister GainDeltasVector = VectorDivide(VectorSubtract(DestinationVector, GainVector), NumFramesVector);
+
+		for (int32 i = 0; i < NumSamples; i += 4)
+		{
+			VectorRegister Result = VectorLoadAligned(&StereoBuffer[i]);
+			Result = VectorMultiply(Result, GainVector);
+			VectorStoreAligned(Result, &StereoBuffer[i]);
+
+			GainVector = VectorAdd(GainVector, GainDeltasVector);
+		}
+	}
+
+	void MixMonoTo2ChannelsFast(const AlignedFloatBuffer& MonoBuffer, AlignedFloatBuffer& DestinationBuffer, const float* RESTRICT Gains)
+	{
+		MixMonoTo2ChannelsFast(MonoBuffer.GetData(), DestinationBuffer.GetData(), DestinationBuffer.Num() / 2, Gains);
+	}
+
+	/**
+	* See CHANNEL MIXING OPERATIONS above for more info.
+	* 2 frames per iteration:
+	* +------------+---------+---------+---------+---------+
+	* | VectorName | Index 0 | Index 1 | Index 2 | Index 3 |
+	* +------------+---------+---------+---------+---------+
+	* | Gain       | g0      | g1      | g0      | g1      |
+	* |            | *       | *       | *       | *       |
+	* | Input      | i0      | i0      | i0      | i0      |
+	* |            | =       | =       | =       | =       |
+	* | Output     | o0      | o1      | o2      | o3      |
+	* +------------+---------+---------+---------+---------+
+	*/
+	void MixMonoTo2ChannelsFast(const float* RESTRICT MonoBuffer, float* RESTRICT DestinationBuffer, int32 NumFrames, const float* RESTRICT Gains)
+	{
+		const VectorRegister GainVector = VectorLoadFloat2(Gains);
+		for (int32 i = 0; i < NumFrames; i += 2)
+		{
+			VectorRegister Result = VectorSet(MonoBuffer[i], MonoBuffer[i], MonoBuffer[i + 1], MonoBuffer[i + 1]);
+			Result = VectorMultiply(Result, GainVector);
+			VectorStoreAligned(Result, &DestinationBuffer[i*2]);
+		}
+	}
+
+	void MixMonoTo2ChannelsFast(const AlignedFloatBuffer& MonoBuffer, AlignedFloatBuffer& DestinationBuffer, const float* RESTRICT StartGains, const float* RESTRICT EndGains)
+	{
+		MixMonoTo2ChannelsFast(MonoBuffer.GetData(), DestinationBuffer.GetData(), DestinationBuffer.Num() / 2, StartGains, EndGains);
+	}
+
+	/**
+	* See CHANNEL MIXING OPERATIONS above for more info.
+	* 2 frames per iteration:
+	* +------------+---------+---------+---------+---------+
+	* | VectorName | Index 0 | Index 1 | Index 2 | Index 3 |
+	* +------------+---------+---------+---------+---------+
+	* | Gain       | g0      | g1      | g0      | g1      |
+	* |            | *       | *       | *       | *       |
+	* | Input      | i0      | i0      | i0      | i0      |
+	* |            | =       | =       | =       | =       |
+	* | Output     | o0      | o1      | o2      | o3      |
+	* +------------+---------+---------+---------+---------+
+	*/
+	void MixMonoTo2ChannelsFast(const float* RESTRICT MonoBuffer, float* RESTRICT DestinationBuffer, int32 NumFrames, const float* RESTRICT StartGains, const float* RESTRICT EndGains)
+	{
+		// Initialize GainVector at StartGains and compute GainDeltasVector:
+		VectorRegister GainVector = VectorLoadFloat2(StartGains);
+		const VectorRegister DestinationVector = VectorLoadFloat2(EndGains);
+		const VectorRegister NumFramesVector = VectorSetFloat1(NumFrames / 2);
+		const VectorRegister GainDeltasVector = VectorDivide(VectorSubtract(DestinationVector, GainVector), NumFramesVector);
+
+		// To help with stair stepping, we initialize the second frame in GainVector to be half a GainDeltas vector higher than the first frame.
+		const VectorRegister VectorOfHalf = VectorSet(0.0f, 0.0f, 0.5f, 0.5f);
+		const VectorRegister HalfOfDeltaVector = VectorMultiply(GainDeltasVector, VectorOfHalf);
+		GainVector = VectorAdd(GainVector, HalfOfDeltaVector);
+
+		for (int32 i = 0; i < NumFrames; i += 2)
+		{
+			VectorRegister Result = VectorSet(MonoBuffer[i], MonoBuffer[i], MonoBuffer[i + 1], MonoBuffer[i + 1]);
+			Result = VectorMultiply(Result, GainVector);
+			VectorStoreAligned(Result, &DestinationBuffer[i*2]);
+
+			GainVector = VectorAdd(GainVector, GainDeltasVector);
+		}
+	}
+
+	void Mix2ChannelsTo2ChannelsFast(const AlignedFloatBuffer& SourceBuffer, AlignedFloatBuffer& DestinationBuffer, const float* RESTRICT Gains)
+	{
+		Mix2ChannelsTo2ChannelsFast(SourceBuffer.GetData(), DestinationBuffer.GetData(), DestinationBuffer.Num() / 2, Gains);
+	}
+
+	/**
+	* See CHANNEL MIXING OPERATIONS above for more info.
+	* 2 frames per iteration:
+	* +------------+---------+---------+---------+---------+
+	* | VectorName | Index 0 | Index 1 | Index 2 | Index 3 |
+	* +------------+---------+---------+---------+---------+
+	* | Gain1      | g0      | g1      | g0      | g1      |
+	* |            | *       | *       | *       | *       |
+	* | Input1     | i0      | i0      | i2      | i2      |
+	* |            | +       | +       | +       | +       |
+	* | Gain2      | g2      | g3      | g2      | g3      |
+	* |            | *       | *       | *       | *       |
+	* | Input2     | i1      | i1      | i3      | i3      |
+	* |            | =       | =       | =       | =       |
+	* | Output     | o0      | o1      | o2      | o3      |
+	* +------------+---------+---------+---------+---------+
+	*/
+
+	void Mix2ChannelsTo2ChannelsFast(const float* RESTRICT SourceBuffer, float* RESTRICT DestinationBuffer, int32 NumFrames, const float* RESTRICT Gains)
+	{
+		const VectorRegister GainVector1 = VectorLoadFloat2(Gains);
+		const VectorRegister GainVector2 = VectorLoadFloat2(Gains + 2);
+
+		for (int32 i = 0; i < NumFrames; i += 2)
+		{
+			const VectorRegister Input1 = VectorSet(SourceBuffer[i * 2], SourceBuffer[i * 2], SourceBuffer[i * 2 + 2], SourceBuffer[i * 2 + 2]);
+			const VectorRegister Input2 = VectorSet(SourceBuffer[i * 2 + 1], SourceBuffer[i * 2 + 1], SourceBuffer[i * 2 + 3], SourceBuffer[i * 2 + 3]);
+
+			VectorRegister Result = VectorMultiply(Input1, GainVector1);
+			Result = VectorMultiplyAdd(Input2, GainVector2, Result);
+
+			VectorStoreAligned(Result, &DestinationBuffer[i * 2]);
+		}
+	}
+
+	void Mix2ChannelsTo2ChannelsFast(const AlignedFloatBuffer& SourceBuffer, AlignedFloatBuffer& DestinationBuffer, const float* RESTRICT StartGains, const float* RESTRICT EndGains)
+	{
+		Mix2ChannelsTo2ChannelsFast(SourceBuffer.GetData(), DestinationBuffer.GetData(), DestinationBuffer.Num() / 2, StartGains, EndGains);
+	}
+
+	/**
+	* See CHANNEL MIXING OPERATIONS above for more info.
+	* 2 frames per iteration:
+	* +------------+---------+---------+---------+---------+
+	* | VectorName | Index 0 | Index 1 | Index 2 | Index 3 |
+	* +------------+---------+---------+---------+---------+
+	* | Gain1      | g0      | g1      | g0      | g1      |
+	* |            | *       | *       | *       | *       |
+	* | Input1     | i0      | i0      | i2      | i2      |
+	* |            | +       | +       | +       | +       |
+	* | Gain2      | g2      | g3      | g2      | g3      |
+	* |            | *       | *       | *       | *       |
+	* | Input2     | i1      | i1      | i3      | i3      |
+	* |            | =       | =       | =       | =       |
+	* | Output     | o0      | o1      | o2      | o3      |
+	* +------------+---------+---------+---------+---------+
+	*/
+
+	void Mix2ChannelsTo2ChannelsFast(const float* RESTRICT SourceBuffer, float* RESTRICT DestinationBuffer, int32 NumFrames, const float* RESTRICT StartGains, const float* RESTRICT EndGains)
+	{
+		const VectorRegister NumFramesVector = VectorSetFloat1(NumFrames / 2);
+
+		VectorRegister GainVector1 = VectorLoadFloat2(StartGains);
+		const VectorRegister DestinationVector1 = VectorLoadFloat2(EndGains);
+		const VectorRegister GainDeltasVector1 = VectorDivide(VectorSubtract(DestinationVector1, GainVector1), NumFramesVector);
+
+		// To help with stair stepping, we initialize the second frame in GainVector to be half a GainDeltas vector higher than the first frame.
+		const VectorRegister VectorOfHalf = VectorSet(0.0f, 0.0f, 0.5f, 0.5f);
+
+		const VectorRegister HalfOfDeltaVector1 = VectorMultiply(GainDeltasVector1, VectorOfHalf);
+		GainVector1 = VectorAdd(GainVector1, HalfOfDeltaVector1);
+
+		VectorRegister GainVector2 = VectorLoadFloat2(StartGains + 2);
+		const VectorRegister DestinationVector2 = VectorLoadFloat2(EndGains + 2);
+		const VectorRegister GainDeltasVector2 = VectorDivide(VectorSubtract(DestinationVector2, GainVector2), NumFramesVector);
+
+		const VectorRegister HalfOfDeltaVector2 = VectorMultiply(GainDeltasVector2, VectorOfHalf);
+		GainVector2 = VectorAdd(GainVector2, HalfOfDeltaVector2);
+
+		for (int32 i = 0; i < NumFrames; i += 2)
+		{
+			const VectorRegister Input1 = VectorSet(SourceBuffer[i * 2], SourceBuffer[i * 2], SourceBuffer[i * 2 + 2], SourceBuffer[i * 2 + 2]);
+			const VectorRegister Input2 = VectorSet(SourceBuffer[i * 2 + 1], SourceBuffer[i * 2 + 1], SourceBuffer[i * 2 + 3], SourceBuffer[i * 2 + 3]);
+
+			VectorRegister Result = VectorMultiply(Input1, GainVector1);
+			Result = VectorMultiplyAdd(Input2, GainVector2, Result);
+
+			VectorStoreAligned(Result, &DestinationBuffer[i * 2]);
+
+			GainVector1 = VectorAdd(GainVector1, GainDeltasVector1);
+			GainVector2 = VectorAdd(GainVector2, GainDeltasVector2);
+		}
+	}
+
+	void Apply4ChannelGain(AlignedFloatBuffer& InterleavedBuffer, const float* RESTRICT Gains)
+	{
+		Apply4ChannelGain(InterleavedBuffer.GetData(), InterleavedBuffer.Num(), Gains);
+	}
+
+	void Apply4ChannelGain(float* RESTRICT InterleavedBuffer, int32 NumSamples, const float* RESTRICT Gains)
+	{
+		const VectorRegister GainVector = VectorLoadAligned(Gains);
+
+		for (int32 i = 0; i < NumSamples; i += 4)
+		{
+			VectorRegister Result = VectorLoadAligned(&InterleavedBuffer[i]);
+			Result = VectorMultiply(Result, GainVector);
+			VectorStoreAligned(Result, &InterleavedBuffer[i]);
+		}
+	}
+
+	void Apply4ChannelGain(AlignedFloatBuffer& InterleavedBuffer, const float* RESTRICT StartGains, const float* RESTRICT EndGains)
+	{
+		Apply4ChannelGain(InterleavedBuffer.GetData(), InterleavedBuffer.Num(), StartGains, EndGains);
+	}
+
+	void Apply4ChannelGain(float* RESTRICT InterleavedBuffer, int32 NumSamples, const float* RESTRICT StartGains, const float* RESTRICT EndGains)
+	{
+		// Initialize GainVector at StartGains and compute GainDeltasVector:
+		VectorRegister GainVector = VectorLoadAligned(StartGains);
+		const VectorRegister DestinationVector = VectorLoadAligned(EndGains);
+		const VectorRegister NumFramesVector = VectorSetFloat1(NumSamples / 4);
+		const VectorRegister GainDeltasVector = VectorDivide(VectorSubtract(DestinationVector, GainVector), NumFramesVector);
+
+		for (int32 i = 0; i < NumSamples; i += 4)
+		{
+			VectorRegister Result = VectorLoadAligned(&InterleavedBuffer[i]);
+			Result = VectorMultiply(Result, GainVector);
+			VectorStoreAligned(Result, &InterleavedBuffer[i]);
+
+			GainVector = VectorAdd(GainVector, GainDeltasVector);
+		}
+	}
+
+	void MixMonoTo4ChannelsFast(const AlignedFloatBuffer& MonoBuffer, AlignedFloatBuffer& DestinationBuffer, const float* RESTRICT Gains)
+	{
+		MixMonoTo4ChannelsFast(MonoBuffer.GetData(), DestinationBuffer.GetData(), DestinationBuffer.Num() / 4, Gains);
+	}
+
+	/**
+	* See CHANNEL MIXING OPERATIONS above for more info.
+	* 1 frame per iteration:
+	* +------------+---------+---------+---------+---------+
+	* | VectorName | Index 0 | Index 1 | Index 2 | Index 3 |
+	* +------------+---------+---------+---------+---------+
+	* | Gain       | g0      | g1      | g2      | g3      |
+	* |            | *       | *       | *       | *       |
+	* | Input      | i0      | i0      | i0      | i0      |
+	* |            | =       | =       | =       | =       |
+	* | Output     | o0      | o1      | o2      | o3      |
+	* +------------+---------+---------+---------+---------+
+	*/
+
+	void MixMonoTo4ChannelsFast(const float* RESTRICT MonoBuffer, float* RESTRICT DestinationBuffer, int32 NumFrames, const float* RESTRICT Gains)
+	{
+		const VectorRegister GainVector = VectorLoadAligned(Gains);
+
+		for (int32 i = 0; i < NumFrames; i++)
+		{
+			VectorRegister Result = VectorLoadFloat1(&MonoBuffer[i]);
+			Result = VectorMultiply(Result, GainVector);
+			VectorStoreAligned(Result, &DestinationBuffer[i*4]);
+		}
+	}
+
+	void MixMonoTo4ChannelsFast(const AlignedFloatBuffer& MonoBuffer, AlignedFloatBuffer& DestinationBuffer, const float* RESTRICT StartGains, const float* RESTRICT EndGains)
+	{
+		MixMonoTo4ChannelsFast(MonoBuffer.GetData(), DestinationBuffer.GetData(), DestinationBuffer.Num() / 4, StartGains, EndGains);
+	}
+
+
+	/**
+	* See CHANNEL MIXING OPERATIONS above for more info.
+	* 1 frame per iteration:
+	* +------------+---------+---------+---------+---------+
+	* | VectorName | Index 0 | Index 1 | Index 2 | Index 3 |
+	* +------------+---------+---------+---------+---------+
+	* | Gain       | g0      | g1      | g2      | g3      |
+	* |            | *       | *       | *       | *       |
+	* | Input      | i0      | i0      | i0      | i0      |
+	* |            | =       | =       | =       | =       |
+	* | Output     | o0      | o1      | o2      | o3      |
+	* +------------+---------+---------+---------+---------+
+	*/
+	void MixMonoTo4ChannelsFast(const float* RESTRICT MonoBuffer, float* RESTRICT DestinationBuffer, int32 NumFrames, const float* RESTRICT StartGains, const float* RESTRICT EndGains)
+	{
+		VectorRegister GainVector = VectorLoadAligned(StartGains);
+		const VectorRegister DestinationVector = VectorLoadAligned(EndGains);
+		const VectorRegister NumFramesVector = VectorSetFloat1(NumFrames);
+		const VectorRegister GainDeltasVector = VectorDivide(VectorSubtract(DestinationVector, GainVector), NumFramesVector);
+
+		for (int32 i = 0; i < NumFrames; i++)
+		{
+			VectorRegister Result = VectorLoadFloat1(&MonoBuffer[i]);
+			Result = VectorMultiply(Result, GainVector);
+			VectorStoreAligned(Result, &DestinationBuffer[i * 4]);
+
+			GainVector = VectorAdd(GainVector, GainDeltasVector);
+		}
+	}
+
+	void Mix2ChannelsTo4ChannelsFast(const AlignedFloatBuffer& SourceBuffer, AlignedFloatBuffer& DestinationBuffer, const float* RESTRICT Gains)
+	{
+		Mix2ChannelsTo4ChannelsFast(SourceBuffer.GetData(), DestinationBuffer.GetData(), DestinationBuffer.Num() / 4, Gains);
+	}
+
+	/**
+	* See CHANNEL MIXING OPERATIONS above for more info.
+	* 1 frame per iteration:
+	* +------------+---------+---------+---------+---------+
+	* | VectorName | Index 0 | Index 1 | Index 2 | Index 3 |
+	* +------------+---------+---------+---------+---------+
+	* | Gain1      | g0      | g1      | g2      | g3      |
+	* |            | *       | *       | *       | *       |
+	* | Input1     | i0      | i0      | i0      | i0      |
+	* |            | +       | +       | +       | +       |
+	* | Gain2      | g4      | g5      | g6      | g7      |
+	* |            | *       | *       | *       | *       |
+	* | Input2     | i1      | i1      | i1      | i1      |
+	* |            | =       | =       | =       | =       |
+	* | Output     | o0      | o1      | o2      | o3      |
+	* +------------+---------+---------+---------+---------+
+	*/
+
+	void Mix2ChannelsTo4ChannelsFast(const float* RESTRICT SourceBuffer, float* RESTRICT DestinationBuffer, int32 NumFrames, const float* RESTRICT Gains)
+	{
+		const VectorRegister GainVector1 = VectorLoadAligned(Gains);
+		const VectorRegister GainVector2 = VectorLoadAligned(Gains + 4);
+
+		for (int32 i = 0; i < NumFrames; i++)
+		{
+			const VectorRegister Input1 = VectorLoadFloat1(&SourceBuffer[i * 2]);
+			const VectorRegister Input2 = VectorLoadFloat1(&SourceBuffer[i * 2 + 1]);
+
+			VectorRegister Result = VectorMultiply(Input1, GainVector1);
+			Result = VectorMultiplyAdd(Input2, GainVector2, Result);
+			VectorStoreAligned(Result, &DestinationBuffer[i * 4]);
+		}
+	}
+
+	void Mix2ChannelsTo4ChannelsFast(const AlignedFloatBuffer& SourceBuffer, AlignedFloatBuffer& DestinationBuffer, const float* RESTRICT StartGains, const float* RESTRICT EndGains)
+	{
+		Mix2ChannelsTo4ChannelsFast(SourceBuffer.GetData(), DestinationBuffer.GetData(), DestinationBuffer.Num() / 4, StartGains, EndGains);
+	}
+
+	/**
+	* See CHANNEL MIXING OPERATIONS above for more info.
+	* 1 frame per iteration:
+	* +------------+---------+---------+---------+---------+
+	* | VectorName | Index 0 | Index 1 | Index 2 | Index 3 |
+	* +------------+---------+---------+---------+---------+
+	* | Gain1      | g0      | g1      | g2      | g3      |
+	* |            | *       | *       | *       | *       |
+	* | Input1     | i0      | i0      | i0      | i0      |
+	* |            | +       | +       | +       | +       |
+	* | Gain2      | g4      | g5      | g6      | g7      |
+	* |            | *       | *       | *       | *       |
+	* | Input2     | i1      | i1      | i1      | i1      |
+	* |            | =       | =       | =       | =       |
+	* | Output     | o0      | o1      | o2      | o3      |
+	* +------------+---------+---------+---------+---------+
+	*/
+	void Mix2ChannelsTo4ChannelsFast(const float* RESTRICT SourceBuffer, float* RESTRICT DestinationBuffer, int32 NumFrames, const float* RESTRICT StartGains, const float* RESTRICT EndGains)
+	{
+		const VectorRegister NumFramesVector = VectorSetFloat1(NumFrames);
+
+		VectorRegister GainVector1 = VectorLoadAligned(StartGains);
+		const VectorRegister DestinationVector1 = VectorLoadAligned(EndGains);
+		const VectorRegister GainDeltasVector1 = VectorDivide(VectorSubtract(DestinationVector1, GainVector1), NumFramesVector);
+
+		VectorRegister GainVector2 = VectorLoadAligned(StartGains + 4);
+		const VectorRegister DestinationVector2 = VectorLoadAligned(EndGains + 4);
+		const VectorRegister GainDeltasVector2 = VectorDivide(VectorSubtract(DestinationVector2, GainVector2), NumFramesVector);
+
+		for (int32 i = 0; i < NumFrames; i++)
+		{
+			const VectorRegister Input1 = VectorLoadFloat1(&SourceBuffer[i * 2]);
+			const VectorRegister Input2 = VectorLoadFloat1(&SourceBuffer[i * 2 + 1]);
+
+			VectorRegister Result = VectorMultiply(Input1, GainVector1);
+			Result = VectorMultiplyAdd(Input2, GainVector2, Result);
+			VectorStoreAligned(Result, &DestinationBuffer[i * 4]);
+
+			GainVector1 = VectorAdd(GainVector1, GainDeltasVector1);
+			GainVector2 = VectorAdd(GainVector2, GainDeltasVector2);
+		}
+	}
+
+	void Apply6ChannelGain(AlignedFloatBuffer& InterleavedBuffer, const float* RESTRICT Gains)
+	{
+		Apply6ChannelGain(InterleavedBuffer.GetData(), InterleavedBuffer.Num(), Gains);
+	}
+
+	void Apply6ChannelGain(float* RESTRICT InterleavedBuffer, int32 NumSamples, const float* RESTRICT Gains)
+	{
+		const VectorRegister GainVector1 = VectorLoadAligned(Gains);
+		const VectorRegister GainVector2 = VectorSet(Gains[4], Gains[5], Gains[0], Gains[1]);
+		const VectorRegister GainVector3 = VectorLoad(&Gains[2]);
+
+		for (int32 i = 0; i < NumSamples; i += 12)
+		{
+			VectorRegister Result = VectorLoadAligned(&InterleavedBuffer[i]);
+			Result = VectorMultiply(Result, GainVector1);
+			VectorStoreAligned(Result, &InterleavedBuffer[i]);
+
+			Result = VectorLoadAligned(&InterleavedBuffer[i + 4]);
+			Result = VectorMultiply(Result, GainVector2);
+			VectorStoreAligned(Result, &InterleavedBuffer[i + 4]);
+
+			Result = VectorLoadAligned(&InterleavedBuffer[i + 8]);
+			Result = VectorMultiply(Result, GainVector3);
+			VectorStoreAligned(Result, &InterleavedBuffer[i + 8]);
+		}
+	}
+
+	void Apply6ChannelGain(AlignedFloatBuffer& InterleavedBuffer, const float* RESTRICT StartGains, const float* RESTRICT EndGains)
+	{
+		Apply6ChannelGain(InterleavedBuffer.GetData(), InterleavedBuffer.Num(), StartGains, EndGains);
+	}
+
+	void Apply6ChannelGain(float* RESTRICT InterleavedBuffer, int32 NumSamples, const float* RESTRICT StartGains, const float* RESTRICT EndGains)
+	{
+		const VectorRegister NumFramesVector = VectorSetFloat1(NumSamples / 12);
+
+		VectorRegister GainVector1 = VectorLoadAligned(StartGains);
+		const VectorRegister DestinationVector1 = VectorLoadAligned(EndGains);
+		const VectorRegister GainDeltasVector1 = VectorDivide(VectorSubtract(DestinationVector1, GainVector1), NumFramesVector);
+
+		VectorRegister GainVector2 = VectorSet(StartGains[4], StartGains[5], StartGains[0], StartGains[1]);
+		const VectorRegister DestinationVector2 = VectorSet(EndGains[4], EndGains[5], EndGains[0], EndGains[1]);
+		const VectorRegister GainDeltasVector2 = VectorDivide(VectorSubtract(DestinationVector2, GainVector2), NumFramesVector);
+
+		VectorRegister GainVector3 = VectorLoad(&StartGains[2]);
+		const VectorRegister DestinationVector3 = VectorLoad(&EndGains[2]);
+		const VectorRegister GainDeltasVector3 = VectorDivide(VectorSubtract(DestinationVector3, GainVector3), NumFramesVector);
+
+		for (int32 i = 0; i < NumSamples; i += 12)
+		{
+			VectorRegister Result = VectorLoadAligned(&InterleavedBuffer[i]);
+			Result = VectorMultiply(Result, GainVector1);
+			VectorStoreAligned(Result, &InterleavedBuffer[i]);
+
+			GainVector1 = VectorAdd(GainVector1, GainDeltasVector1);
+
+			Result = VectorLoadAligned(&InterleavedBuffer[i + 4]);
+			Result = VectorMultiply(Result, GainVector2);
+			VectorStoreAligned(Result, &InterleavedBuffer[i + 4]);
+
+			GainVector2 = VectorAdd(GainVector2, GainDeltasVector2);
+
+			Result = VectorLoadAligned(&InterleavedBuffer[i + 8]);
+			Result = VectorMultiply(Result, GainVector3);
+			VectorStoreAligned(Result, &InterleavedBuffer[i + 8]);
+
+			GainVector3 = VectorAdd(GainVector3, GainDeltasVector3);
+		}
+	}
+
+	void MixMonoTo6ChannelsFast(const AlignedFloatBuffer& MonoBuffer, AlignedFloatBuffer& DestinationBuffer, const float* RESTRICT Gains)
+	{
+		MixMonoTo6ChannelsFast(MonoBuffer.GetData(), DestinationBuffer.GetData(), DestinationBuffer.Num() / 6, Gains);
+	}
+
+	/**
+	* See CHANNEL MIXING OPERATIONS above for more info.
+	* 2 frames per iteration:
+	* +------------+----------+---------+---------+---------+----------+---------+---------+---------+----------+---------+----------+----------+
+	* |            | Vector 1 |         |         |         | Vector 2 |         |         |         | Vector 3 |         |          |          |
+	* +------------+----------+---------+---------+---------+----------+---------+---------+---------+----------+---------+----------+----------+
+	* |            | Index 0  | Index 1 | Index 2 | Index 3 | Index 4  | Index 5 | Index 6 | Index 7 | Index 8  | Index 9 | Index 10 | Index 11 |
+	* | Gain       | g0       | g1      | g2      | g3      | g4       | g5      | g0      | g1      | g2       | g3      | g4       | g5       |
+	* |            | *        | *       | *       | *       | *        | *       | *       | *       | *        | *       | *        | *        |
+	* | Input      | i0       | i0      | i0      | i0      | i0       | i0      | i1      | i1      | i1       | i1      | i1       | i1       |
+	* |            | =        | =       | =       | =       | =        | =       | =       | =       | =        | =       | =        | =        |
+	* | Output     | o0       | o1      | o2      | o3      | o4       | o5      | o6      | o7      | o8       | o9      | o10      | o11      |
+	* +------------+----------+---------+---------+---------+----------+---------+---------+---------+----------+---------+----------+----------+
+	*/
+
+	void MixMonoTo6ChannelsFast(const float* RESTRICT MonoBuffer, float* RESTRICT DestinationBuffer, int32 NumFrames, const float* RESTRICT Gains)
+	{
+		const VectorRegister GainVector1 = VectorLoadAligned(Gains);
+		const VectorRegister GainVector2 = VectorSet(Gains[4], Gains[5], Gains[0], Gains[1]);
+		const VectorRegister GainVector3 = VectorLoad(&Gains[2]);
+
+		for (int32 i = 0; i < NumFrames; i += 2)
+		{
+			const VectorRegister Input1 = VectorLoadFloat1(&MonoBuffer[i]);
+			const VectorRegister Input2 = VectorSet(MonoBuffer[i], MonoBuffer[i], MonoBuffer[i + 1], MonoBuffer[i + 1]);
+			const VectorRegister Input3 = VectorLoadFloat1(&MonoBuffer[i + 1]);
+
+			VectorRegister Result = VectorMultiply(Input1, GainVector1);
+			VectorStoreAligned(Result, &DestinationBuffer[i * 6]);
+
+			Result = VectorMultiply(Input2, GainVector2);
+			VectorStoreAligned(Result, &DestinationBuffer[i * 6 + 4]);
+
+			Result = VectorMultiply(Input3, GainVector3);
+			VectorStoreAligned(Result, &DestinationBuffer[i * 6 + 8]);
+		}
+	}
+
+	void MixMonoTo6ChannelsFast(const AlignedFloatBuffer& MonoBuffer, AlignedFloatBuffer& DestinationBuffer, const float* RESTRICT StartGains, const float* RESTRICT EndGains)
+	{
+		MixMonoTo6ChannelsFast(MonoBuffer.GetData(), DestinationBuffer.GetData(), DestinationBuffer.Num() / 6, StartGains, EndGains);
+	}
+
+	/**
+	* See CHANNEL MIXING OPERATIONS above for more info.
+	* 2 frames per iteration:
+	* +------------+----------+---------+---------+---------+----------+---------+---------+---------+----------+---------+----------+----------+
+	* |            | Vector 1 |         |         |         | Vector 2 |         |         |         | Vector 3 |         |          |          |
+	* +------------+----------+---------+---------+---------+----------+---------+---------+---------+----------+---------+----------+----------+
+	* |            | Index 0  | Index 1 | Index 2 | Index 3 | Index 4  | Index 5 | Index 6 | Index 7 | Index 8  | Index 9 | Index 10 | Index 11 |
+	* | Gain       | g0       | g1      | g2      | g3      | g4       | g5      | g0      | g1      | g2       | g3      | g4       | g5       |
+	* |            | *        | *       | *       | *       | *        | *       | *       | *       | *        | *       | *        | *        |
+	* | Input      | i0       | i0      | i0      | i0      | i0       | i0      | i1      | i1      | i1       | i1      | i1       | i1       |
+	* |            | =        | =       | =       | =       | =        | =       | =       | =       | =        | =       | =        | =        |
+	* | Output     | o0       | o1      | o2      | o3      | o4       | o5      | o6      | o7      | o8       | o9      | o10      | o11      |
+	* +------------+----------+---------+---------+---------+----------+---------+---------+---------+----------+---------+----------+----------+
+	*/
+
+	void MixMonoTo6ChannelsFast(const float* RESTRICT MonoBuffer, float* RESTRICT DestinationBuffer, int32 NumFrames, const float* RESTRICT StartGains, const float* RESTRICT EndGains)
+	{
+		const VectorRegister NumFramesVector = VectorSetFloat1(NumFrames / 2);
+
+		VectorRegister GainVector1 = VectorLoadAligned(StartGains);
+		const VectorRegister DestinationVector1 = VectorLoadAligned(EndGains);
+		const VectorRegister GainDeltasVector1 = VectorDivide(VectorSubtract(DestinationVector1, GainVector1), NumFramesVector);
+
+		VectorRegister GainVector2 = VectorSet(StartGains[4], StartGains[5], StartGains[0], StartGains[1]);
+		const VectorRegister DestinationVector2 = VectorSet(EndGains[4], EndGains[5], EndGains[0], EndGains[1]);
+		const VectorRegister GainDeltasVector2 = VectorDivide(VectorSubtract(DestinationVector2, GainVector2), NumFramesVector);
+
+		VectorRegister GainVector3 = VectorLoad(&StartGains[2]);
+		const VectorRegister DestinationVector3 = VectorLoad(&EndGains[2]);
+		const VectorRegister GainDeltasVector3 = VectorDivide(VectorSubtract(DestinationVector3, GainVector3), NumFramesVector);
+
+		for (int32 i = 0; i < NumFrames; i += 2)
+		{
+			const VectorRegister Input1 = VectorLoadFloat1(&MonoBuffer[i]);
+			const VectorRegister Input2 = VectorSet(MonoBuffer[i], MonoBuffer[i], MonoBuffer[i + 1], MonoBuffer[i + 1]);
+			const VectorRegister Input3 = VectorLoadFloat1(&MonoBuffer[i + 1]);
+
+			VectorRegister Result = VectorMultiply(Input1, GainVector1);
+			VectorStoreAligned(Result, &DestinationBuffer[i * 6]);
+
+			GainVector1 = VectorAdd(GainVector1, GainDeltasVector1);
+
+			Result = VectorMultiply(Input2, GainVector2);
+			VectorStoreAligned(Result, &DestinationBuffer[i * 6 + 4]);
+
+			GainVector2 = VectorAdd(GainVector2, GainDeltasVector2);
+
+			Result = VectorMultiply(Input3, GainVector3);
+			VectorStoreAligned(Result, &DestinationBuffer[i * 6 + 8]);
+
+			GainVector3 = VectorAdd(GainVector3, GainDeltasVector3);
+		}
+	}
+
+	void Mix2ChannelsTo6ChannelsFast(const AlignedFloatBuffer& SourceBuffer, AlignedFloatBuffer& DestinationBuffer, const float* RESTRICT Gains)
+	{
+		Mix2ChannelsTo6ChannelsFast(SourceBuffer.GetData(), DestinationBuffer.GetData(), DestinationBuffer.Num() / 6, Gains);
+	}
+
+	/**
+	* See CHANNEL MIXING OPERATIONS above for more info.
+	* 2 frames per iteration:
+	* +--------+----------+---------+---------+---------+----------+---------+---------+---------+----------+---------+----------+----------+
+	* |        | Vector 1 |         |         |         | Vector 2 |         |         |         | Vector 3 |         |          |          |
+	* +--------+----------+---------+---------+---------+----------+---------+---------+---------+----------+---------+----------+----------+
+	* |        | Index 0  | Index 1 | Index 2 | Index 3 | Index 4  | Index 5 | Index 6 | Index 7 | Index 8  | Index 9 | Index 10 | Index 11 |
+	* | Gain1  | g0       | g1      | g2      | g3      | g4       | g5      | g0      | g1      | g2       | g3      | g4       | g5       |
+	* |        | *        | *       | *       | *       | *        | *       | *       | *       | *        | *       | *        | *        |
+	* | Input1 | i0       | i0      | i0      | i0      | i0       | i0      | i2      | i2      | i2       | i2      | i2       | i2       |
+	* |        | +        | +       | +       | +       | +        | +       | +       | +       | +        | +       | +        | +        |
+	* | Gain2  | g6       | g7      | g8      | g9      | g10      | g11     | g6      | g7      | g8       | g9      | g10      | g11      |
+	* |        | *        | *       | *       | *       | *        | *       | *       | *       | *        | *       | *        | *        |
+	* | Input2 | i1       | i1      | i1      | i1      | i1       | i1      | i1      | i3      | i3       | i3      | i3       | i3       |
+	* |        | =        | =       | =       | =       | =        | =       | =       | =       | =        | =       | =        | =        |
+	* | Output | o0       | o1      | o2      | o3      | o4       | o5      | o6      | o7      | o8       | o9      | o10      | o11      |
+	* +--------+----------+---------+---------+---------+----------+---------+---------+---------+----------+---------+----------+----------+
+	*/
+
+	void Mix2ChannelsTo6ChannelsFast(const float* RESTRICT SourceBuffer, float* RESTRICT DestinationBuffer, int32 NumFrames, const float* RESTRICT Gains)
+	{
+		const VectorRegister GainVector11 = VectorLoadAligned(Gains);
+		const VectorRegister GainVector21 = VectorSet(Gains[4], Gains[5], Gains[0], Gains[1]);
+		const VectorRegister GainVector31 = VectorLoad(&Gains[2]);
+
+		const VectorRegister GainVector12 = VectorLoadAligned(Gains + 6);
+		const VectorRegister GainVector22 = VectorSet(Gains[10], Gains[11], Gains[6], Gains[7]);
+		const VectorRegister GainVector32 = VectorLoadAligned(&Gains[8]);
+
+		for (int32 FrameIndex = 0; FrameIndex < NumFrames; FrameIndex += 2)
+		{
+			const int32 InputIndex = FrameIndex * 2;
+			const int32 OutputIndex = FrameIndex * 6;
+
+			const VectorRegister Input11 = VectorLoadFloat1(&SourceBuffer[InputIndex]);
+			const VectorRegister Input21 = VectorSet(SourceBuffer[InputIndex], SourceBuffer[InputIndex], SourceBuffer[InputIndex + 2], SourceBuffer[InputIndex + 2]);
+			const VectorRegister Input31 = VectorLoadFloat1(&SourceBuffer[InputIndex + 2]);
+
+			const VectorRegister Input12 = VectorLoadFloat1(&SourceBuffer[InputIndex + 1]);
+			const VectorRegister Input22 = VectorSet(SourceBuffer[InputIndex + 1], SourceBuffer[InputIndex + 1], SourceBuffer[InputIndex + 3], SourceBuffer[InputIndex + 3]);
+			const VectorRegister Input32 = VectorLoadFloat1(&SourceBuffer[InputIndex + 3]);
+
+			VectorRegister Result = VectorMultiply(Input11, GainVector11);
+			Result = VectorMultiplyAdd(Input12, GainVector12, Result);
+			VectorStoreAligned(Result, &DestinationBuffer[OutputIndex]);
+
+			Result = VectorMultiply(Input21, GainVector21);
+			Result = VectorMultiplyAdd(Input22, GainVector22, Result);
+			VectorStoreAligned(Result, &DestinationBuffer[OutputIndex + 4]);
+
+			Result = VectorMultiply(Input31, GainVector31);
+			Result = VectorMultiplyAdd(Input32, GainVector31, Result);
+			VectorStoreAligned(Result, &DestinationBuffer[OutputIndex + 8]);
+		}
+	}
+
+	void Mix2ChannelsTo6ChannelsFast(const AlignedFloatBuffer& SourceBuffer, AlignedFloatBuffer& DestinationBuffer, const float* RESTRICT StartGains, const float* RESTRICT EndGains)
+	{
+		Mix2ChannelsTo6ChannelsFast(SourceBuffer.GetData(), DestinationBuffer.GetData(), DestinationBuffer.Num() / 6, StartGains, EndGains);
+	}
+
+	/**
+	* See CHANNEL MIXING OPERATIONS above for more info.
+	* 2 frames per iteration:
+	* +--------+----------+---------+---------+---------+----------+---------+---------+---------+----------+---------+----------+----------+
+	* |        | Vector 1 |         |         |         | Vector 2 |         |         |         | Vector 3 |         |          |          |
+	* +--------+----------+---------+---------+---------+----------+---------+---------+---------+----------+---------+----------+----------+
+	* |        | Index 0  | Index 1 | Index 2 | Index 3 | Index 4  | Index 5 | Index 6 | Index 7 | Index 8  | Index 9 | Index 10 | Index 11 |
+	* | Gain1  | g0       | g1      | g2      | g3      | g4       | g5      | g0      | g1      | g2       | g3      | g4       | g5       |
+	* |        | *        | *       | *       | *       | *        | *       | *       | *       | *        | *       | *        | *        |
+	* | Input1 | i0       | i0      | i0      | i0      | i0       | i0      | i2      | i2      | i2       | i2      | i2       | i2       |
+	* |        | +        | +       | +       | +       | +        | +       | +       | +       | +        | +       | +        | +        |
+	* | Gain2  | g6       | g7      | g8      | g9      | g10      | g11     | g6      | g7      | g8       | g9      | g10      | g11      |
+	* |        | *        | *       | *       | *       | *        | *       | *       | *       | *        | *       | *        | *        |
+	* | Input2 | i1       | i1      | i1      | i1      | i1       | i1      | i1      | i3      | i3       | i3      | i3       | i3       |
+	* |        | =        | =       | =       | =       | =        | =       | =       | =       | =        | =       | =        | =        |
+	* | Output | o0       | o1      | o2      | o3      | o4       | o5      | o6      | o7      | o8       | o9      | o10      | o11      |
+	* +--------+----------+---------+---------+---------+----------+---------+---------+---------+----------+---------+----------+----------+
+	*/
+
+	void Mix2ChannelsTo6ChannelsFast(const float* RESTRICT SourceBuffer, float* RESTRICT DestinationBuffer, int32 NumFrames, const float* RESTRICT StartGains, const float* RESTRICT EndGains)
+	{
+		const VectorRegister NumFramesVector = VectorSetFloat1(NumFrames / 2);
+
+		VectorRegister GainVector11 = VectorLoadAligned(StartGains);
+		const VectorRegister DestinationVector11 = VectorLoadAligned(EndGains);
+		const VectorRegister GainDeltasVector11 = VectorDivide(VectorSubtract(DestinationVector11, GainVector11), NumFramesVector);
+
+		VectorRegister GainVector21 = VectorSet(StartGains[4], StartGains[5], StartGains[0], StartGains[1]);
+		const VectorRegister DestinationVector21 = VectorSet(EndGains[4], EndGains[5], EndGains[0], EndGains[1]);
+		const VectorRegister GainDeltasVector21 = VectorDivide(VectorSubtract(DestinationVector21, GainVector21), NumFramesVector);
+
+		// In order to ease stair stepping, we ensure that the second frame is initialized to half the GainDelta more than the first frame.
+		// This gives us a consistent increment across every frame.
+		const VectorRegister DeltaHalf21 = VectorSet(0.0f, 0.0f, 0.5f, 0.5f);
+		const VectorRegister InitializedDelta21 = VectorMultiply(GainDeltasVector21, DeltaHalf21);
+		GainVector21 = VectorAdd(GainVector21, InitializedDelta21);
+
+		VectorRegister GainVector31 = VectorLoad(&StartGains[2]);
+		const VectorRegister DestinationVector31 = VectorLoad(&EndGains[2]);
+		const VectorRegister GainDeltasVector31 = VectorDivide(VectorSubtract(DestinationVector31, GainVector31), NumFramesVector);
+
+		const VectorRegister DeltaHalf31 = VectorSetFloat1(0.5f);
+		const VectorRegister InitializedDelta31 = VectorMultiply(GainDeltasVector31, DeltaHalf31);
+		GainVector31 = VectorAdd(GainVector31, InitializedDelta31);
+
+		VectorRegister GainVector12 = VectorLoadAligned(StartGains + 6);
+		const VectorRegister DestinationVector12 = VectorLoadAligned(EndGains + 6);
+		const VectorRegister GainDeltasVector12 = VectorDivide(VectorSubtract(DestinationVector12, GainVector12), NumFramesVector);
+
+		VectorRegister GainVector22 = VectorSet(StartGains[10], StartGains[11], StartGains[6], StartGains[7]);
+		const VectorRegister DestinationVector22 = VectorSet(EndGains[10], EndGains[11], EndGains[6], EndGains[7]);
+		const VectorRegister GainDeltasVector22 = VectorDivide(VectorSubtract(DestinationVector22, GainVector22), NumFramesVector);
+
+		const VectorRegister DeltaHalf22 = VectorSet(0.0f, 0.0f, 0.5f, 0.5f);
+		const VectorRegister InitializedDelta22 = VectorMultiply(GainDeltasVector22, DeltaHalf22);
+		GainVector22 = VectorAdd(GainVector22, InitializedDelta22);
+
+		VectorRegister GainVector32 = VectorLoadAligned(StartGains + 8);
+		const VectorRegister DestinationVector32 = VectorLoadAligned(EndGains + 8);
+		const VectorRegister GainDeltasVector32 = VectorDivide(VectorSubtract(DestinationVector32, GainVector32), NumFramesVector);
+
+		const VectorRegister DeltaHalf32 = VectorSetFloat1(0.5f);
+		const VectorRegister InitializedDelta32 = VectorMultiply(GainDeltasVector32, DeltaHalf32);
+		GainVector32 = VectorAdd(GainVector32, InitializedDelta32);
+
+		for (int32 FrameIndex = 0; FrameIndex < NumFrames; FrameIndex += 2)
+		{
+			const int32 InputIndex = FrameIndex * 2;
+			const int32 OutputIndex = FrameIndex * 6;
+
+			const VectorRegister Input11 = VectorLoadFloat1(&SourceBuffer[InputIndex]);
+			const VectorRegister Input21 = VectorSet(SourceBuffer[InputIndex], SourceBuffer[InputIndex], SourceBuffer[InputIndex + 2], SourceBuffer[InputIndex + 2]);
+			const VectorRegister Input31 = VectorLoadFloat1(&SourceBuffer[InputIndex + 2]);
+
+			const VectorRegister Input12 = VectorLoadFloat1(&SourceBuffer[InputIndex + 1]);
+			const VectorRegister Input22 = VectorSet(SourceBuffer[InputIndex + 1], SourceBuffer[InputIndex + 1], SourceBuffer[InputIndex + 3], SourceBuffer[InputIndex + 3]);
+			const VectorRegister Input32 = VectorLoadFloat1(&SourceBuffer[InputIndex + 3]);
+
+			VectorRegister Result = VectorMultiply(Input11, GainVector11);
+			Result = VectorMultiplyAdd(Input12, GainVector12, Result);
+			VectorStoreAligned(Result, &DestinationBuffer[OutputIndex]);
+
+			GainVector11 = VectorAdd(GainVector11, GainDeltasVector11);
+			GainVector12 = VectorAdd(GainVector12, GainDeltasVector12);
+
+			Result = VectorMultiply(Input21, GainVector21);
+			Result = VectorMultiplyAdd(Input22, GainVector22, Result);
+			VectorStoreAligned(Result, &DestinationBuffer[OutputIndex + 4]);
+
+			GainVector21 = VectorAdd(GainVector21, GainDeltasVector21);
+			GainVector22 = VectorAdd(GainVector22, GainDeltasVector22);
+
+			Result = VectorMultiply(Input31, GainVector31);
+			Result = VectorMultiplyAdd(Input32, GainVector31, Result);
+			VectorStoreAligned(Result, &DestinationBuffer[OutputIndex + 8]);
+
+			GainVector31 = VectorAdd(GainVector31, GainDeltasVector31);
+			GainVector32 = VectorAdd(GainVector32, GainDeltasVector32);
+		}
+	}
+
+	void Apply8ChannelGain(AlignedFloatBuffer& InterleavedBuffer, const float* RESTRICT Gains)
+	{
+		Apply8ChannelGain(InterleavedBuffer.GetData(), InterleavedBuffer.Num(), Gains);
+	}
+
+	void Apply8ChannelGain(float* RESTRICT InterleavedBuffer, int32 NumSamples, const float* RESTRICT Gains)
+	{
+		const VectorRegister GainVector1 = VectorLoadAligned(Gains);
+		const VectorRegister GainVector2 = VectorLoadAligned(Gains + 4);
+
+		for (int32 i = 0; i < NumSamples; i += 8)
+		{
+			VectorRegister Result = VectorLoadAligned(&InterleavedBuffer[i]);
+			Result = VectorMultiply(Result, GainVector1);
+			VectorStoreAligned(Result, &InterleavedBuffer[i]);
+
+			Result = VectorLoadAligned(&InterleavedBuffer[i + 4]);
+			Result = VectorMultiply(Result, GainVector2);
+			VectorStoreAligned(Result, &InterleavedBuffer[i + 4]);
+		}
+	}
+
+	void Apply8ChannelGain(AlignedFloatBuffer& InterleavedBuffer, const float* RESTRICT StartGains, const float* RESTRICT EndGains)
+	{
+		Apply8ChannelGain(InterleavedBuffer.GetData(), InterleavedBuffer.Num(), StartGains, EndGains);
+	}
+
+	void Apply8ChannelGain(float* RESTRICT InterleavedBuffer, int32 NumSamples, const float* RESTRICT StartGains, const float* RESTRICT EndGains)
+	{
+		const VectorRegister NumFramesVector = VectorSetFloat1(NumSamples / 8);
+
+		VectorRegister GainVector1 = VectorLoadAligned(StartGains);
+		const VectorRegister DestinationVector1 = VectorLoadAligned(EndGains);
+		const VectorRegister GainDeltasVector1 = VectorDivide(VectorSubtract(DestinationVector1, GainVector1), NumFramesVector);
+
+		VectorRegister GainVector2 = VectorLoadAligned(StartGains + 4);
+		const VectorRegister DestinationVector2 = VectorLoadAligned(EndGains + 4);
+		const VectorRegister GainDeltasVector2 = VectorDivide(VectorSubtract(DestinationVector2, GainVector2), NumFramesVector);
+
+		for (int32 i = 0; i < NumSamples; i += 8)
+		{
+			VectorRegister Result = VectorLoadAligned(&InterleavedBuffer[i]);
+			Result = VectorMultiply(Result, GainVector1);
+			VectorStoreAligned(Result, &InterleavedBuffer[i]);
+
+			GainVector1 = VectorAdd(GainVector1, GainDeltasVector1);
+
+			Result = VectorLoadAligned(&InterleavedBuffer[i + 4]);
+			Result = VectorMultiply(Result, GainVector2);
+			VectorStoreAligned(Result, &InterleavedBuffer[i + 4]);
+
+			GainVector2 = VectorAdd(GainVector2, GainDeltasVector2);
+		}
+	}
+
+	void MixMonoTo8ChannelsFast(const AlignedFloatBuffer& MonoBuffer, AlignedFloatBuffer& DestinationBuffer, const float* RESTRICT Gains)
+	{
+		MixMonoTo8ChannelsFast(MonoBuffer.GetData(), DestinationBuffer.GetData(), DestinationBuffer.Num() / 8, Gains);
+	}
+
+	/**
+	* See CHANNEL MIXING OPERATIONS above for more info.
+	* 1 frames per iteration:
+	* +------------+----------+---------+---------+---------+----------+---------+---------+---------+
+	* |            | Vector 1 |         |         |         | Vector 2 |         |         |         |
+	* | VectorName | Index 0  | Index 1 | Index 2 | Index 3 | Index 4  | Index 5 | Index 6 | Index 7 |
+	* | Gain       | g0       | g1      | g2      | g3      | g4       | g5      | g6      | g7      |
+	* |            | *        | *       | *       | *       | *        | *       | *       | *       |
+	* | Input      | i0       | i0      | i0      | i0      | i0       | i0      | i0      | i0      |
+	* |            | =        | =       | =       | =       | =        | =       | =       | =       |
+	* | Output     | o0       | o1      | o2      | o3      | o4       | o5      | o6      | o7      |
+	* +------------+----------+---------+---------+---------+----------+---------+---------+---------+
+	*/
+
+	void MixMonoTo8ChannelsFast(const float* RESTRICT MonoBuffer, float* RESTRICT DestinationBuffer, int32 NumFrames, const float* RESTRICT Gains)
+	{
+		const VectorRegister GainVector1 = VectorLoadAligned(Gains);
+		const VectorRegister GainVector2 = VectorLoadAligned(Gains + 4);
+
+		for (int32 i = 0; i < NumFrames; i++)
+		{
+			VectorRegister Result = VectorLoadFloat1(&MonoBuffer[i]);
+			Result = VectorMultiply(Result, GainVector1);
+			VectorStoreAligned(Result, &DestinationBuffer[i * 8]);
+
+			Result = VectorLoadFloat1(&MonoBuffer[i]);
+			Result = VectorMultiply(Result, GainVector2);
+			VectorStoreAligned(Result, &DestinationBuffer[i * 8 + 4]);
+		}
+	}
+
+	void MixMonoTo8ChannelsFast(const AlignedFloatBuffer& MonoBuffer, AlignedFloatBuffer& DestinationBuffer, const float* RESTRICT StartGains, const float* RESTRICT EndGains)
+	{
+		MixMonoTo8ChannelsFast(MonoBuffer.GetData(), DestinationBuffer.GetData(), DestinationBuffer.Num() / 8, StartGains, EndGains);
+	}
+
+	/**
+	* See CHANNEL MIXING OPERATIONS above for more info.
+	* 1 frames per iteration:
+	* +------------+----------+---------+---------+---------+----------+---------+---------+---------+
+	* |            | Vector 1 |         |         |         | Vector 2 |         |         |         |
+	* | VectorName | Index 0  | Index 1 | Index 2 | Index 3 | Index 4  | Index 5 | Index 6 | Index 7 |
+	* | Gain       | g0       | g1      | g2      | g3      | g4       | g5      | g6      | g7      |
+	* |            | *        | *       | *       | *       | *        | *       | *       | *       |
+	* | Input      | i0       | i0      | i0      | i0      | i0       | i0      | i0      | i0      |
+	* |            | =        | =       | =       | =       | =        | =       | =       | =       |
+	* | Output     | o0       | o1      | o2      | o3      | o4       | o5      | o6      | o7      |
+	* +------------+----------+---------+---------+---------+----------+---------+---------+---------+
+	*/
+	void MixMonoTo8ChannelsFast(const float* RESTRICT MonoBuffer, float* RESTRICT DestinationBuffer, int32 NumFrames, const float* RESTRICT StartGains, const float* RESTRICT EndGains)
+	{
+		const VectorRegister NumFramesVector = VectorSetFloat1(NumFrames);
+
+		VectorRegister GainVector1 = VectorLoadAligned(StartGains);
+		const VectorRegister DestinationVector1 = VectorLoadAligned(EndGains);
+		const VectorRegister GainDeltasVector1 = VectorDivide(VectorSubtract(DestinationVector1, GainVector1), NumFramesVector);
+
+		VectorRegister GainVector2 = VectorLoadAligned(StartGains + 4);
+		const VectorRegister DestinationVector2 = VectorLoadAligned(EndGains + 4);
+		const VectorRegister GainDeltasVector2 = VectorDivide(VectorSubtract(DestinationVector2, GainVector2), NumFramesVector);
+
+		for (int32 i = 0; i < NumFrames; i++)
+		{
+			const VectorRegister Input = VectorLoadFloat1(&MonoBuffer[i]);
+			VectorRegister Result = VectorMultiply(Input, GainVector1);
+			VectorStoreAligned(Result, &DestinationBuffer[i * 8]);
+
+			GainVector1 = VectorAdd(GainVector1, GainDeltasVector1);
+
+			Result = VectorMultiply(Input, GainVector2);
+			VectorStoreAligned(Result, &DestinationBuffer[i * 8 + 4]);
+
+			GainVector2 = VectorAdd(GainVector2, GainDeltasVector2);
+		}
+	}
+
+	void Mix2ChannelsTo8ChannelsFast(const AlignedFloatBuffer& SourceBuffer, AlignedFloatBuffer& DestinationBuffer, const float* RESTRICT Gains)
+	{
+		Mix2ChannelsTo8ChannelsFast(SourceBuffer.GetData(), DestinationBuffer.GetData(), DestinationBuffer.Num() / 8, Gains);
+	}
+
+	/**
+	* See CHANNEL MIXING OPERATIONS above for more info.
+	* 1 frames per iteration:
+	* +--------+----------+---------+---------+---------+----------+---------+---------+---------+
+	* |        | Vector 1 |         |         |         | Vector 2 |         |         |         |
+	* +--------+----------+---------+---------+---------+----------+---------+---------+---------+
+	* |        | Index 0  | Index 1 | Index 2 | Index 3 | Index 4  | Index 5 | Index 6 | Index 7 |
+	* | Gain1  | g0       | g1      | g2      | g3      | g4       | g5      | g6      | g7      |
+	* |        | *        | *       | *       | *       | *        | *       | *       | *       |
+	* | Input1 | i0       | i0      | i0      | i0      | i0       | i0      | i0      | i0      |
+	* |        | +        | +       | +       | +       | +        | +       | +       | +       |
+	* | Gain2  | g8       | g9      | g10     | g11     | g12      | g13     | g14     | g5      |
+	* |        | *        | *       | *       | *       | *        | *       | *       | *       |
+	* | Input2 | i1       | i1      | i1      | i1      | i1       | i1      | i1      | i1      |
+	* |        | =        | =       | =       | =       | =        | =       | =       | =       |
+	* | Output | o0       | o1      | o2      | o3      | o4       | o5      | o6      | o7      |
+	* +--------+----------+---------+---------+---------+----------+---------+---------+---------+
+	*/
+
+	void Mix2ChannelsTo8ChannelsFast(const float* RESTRICT SourceBuffer, float* RESTRICT DestinationBuffer, int32 NumFrames, const float* RESTRICT Gains)
+	{
+		const VectorRegister GainVector11 = VectorLoadAligned(Gains);
+		const VectorRegister GainVector21 = VectorLoadAligned(Gains + 4);
+		const VectorRegister GainVector12 = VectorLoadAligned(Gains + 8);
+		const VectorRegister GainVector22 = VectorLoadAligned(Gains + 12);
+
+		for (int32 i = 0; i < NumFrames; i++)
+		{
+			const VectorRegister Input1 = VectorLoadFloat1(&SourceBuffer[i*2]);
+			const VectorRegister Input2 = VectorLoadFloat1(&SourceBuffer[i * 2 + 1]);
+
+			VectorRegister Result = VectorMultiply(Input1, GainVector11);
+			Result = VectorMultiplyAdd(Input2, GainVector21, Result);
+
+			VectorStoreAligned(Result, &DestinationBuffer[i * 8]);
+
+			Result = VectorMultiply(Input1, GainVector12);
+			Result = VectorMultiplyAdd(Input2, GainVector22, Result);
+			VectorStoreAligned(Result, &DestinationBuffer[i * 8 + 4]);
+		}
+	}
+
+	void Mix2ChannelsTo8ChannelsFast(const AlignedFloatBuffer& SourceBuffer, AlignedFloatBuffer& DestinationBuffer, const float* RESTRICT StartGains, const float* RESTRICT EndGains)
+	{
+		Mix2ChannelsTo8ChannelsFast(SourceBuffer.GetData(), DestinationBuffer.GetData(), DestinationBuffer.Num() / 8, StartGains, EndGains);
+	}
+
+	/**
+	* See CHANNEL MIXING OPERATIONS above for more info.
+	* 1 frames per iteration:
+	* +--------+----------+---------+---------+---------+----------+---------+---------+---------+
+	* |        | Vector 1 |         |         |         | Vector 2 |         |         |         |
+	* +--------+----------+---------+---------+---------+----------+---------+---------+---------+
+	* |        | Index 0  | Index 1 | Index 2 | Index 3 | Index 4  | Index 5 | Index 6 | Index 7 |
+	* | Gain1  | g0       | g1      | g2      | g3      | g4       | g5      | g6      | g7      |
+	* |        | *        | *       | *       | *       | *        | *       | *       | *       |
+	* | Input1 | i0       | i0      | i0      | i0      | i0       | i0      | i0      | i0      |
+	* |        | +        | +       | +       | +       | +        | +       | +       | +       |
+	* | Gain2  | g8       | g9      | g10     | g11     | g12      | g13     | g14     | g5      |
+	* |        | *        | *       | *       | *       | *        | *       | *       | *       |
+	* | Input2 | i1       | i1      | i1      | i1      | i1       | i1      | i1      | i1      |
+	* |        | =        | =       | =       | =       | =        | =       | =       | =       |
+	* | Output | o0       | o1      | o2      | o3      | o4       | o5      | o6      | o7      |
+	* +--------+----------+---------+---------+---------+----------+---------+---------+---------+
+	*/
+
+	void Mix2ChannelsTo8ChannelsFast(const float* RESTRICT SourceBuffer, float* RESTRICT DestinationBuffer, int32 NumFrames, const float* RESTRICT StartGains, const float* RESTRICT EndGains)
+	{
+		const VectorRegister NumFramesVector = VectorSetFloat1(NumFrames);
+
+		VectorRegister GainVector11 = VectorLoadAligned(StartGains);
+		const VectorRegister DestinationVector11 = VectorLoadAligned(EndGains);
+		const VectorRegister GainDeltasVector11 = VectorDivide(VectorSubtract(DestinationVector11, GainVector11), NumFramesVector);
+
+		VectorRegister GainVector21 = VectorLoadAligned(StartGains + 4);
+		const VectorRegister DestinationVector21 = VectorLoadAligned(EndGains + 4);
+		const VectorRegister GainDeltasVector21 = VectorDivide(VectorSubtract(DestinationVector21, GainVector21), NumFramesVector);
+
+		VectorRegister GainVector12 = VectorLoadAligned(StartGains + 8);
+		const VectorRegister DestinationVector12 = VectorLoadAligned(EndGains + 8);
+		const VectorRegister GainDeltasVector12 = VectorDivide(VectorSubtract(DestinationVector12, GainVector12), NumFramesVector);
+
+		VectorRegister GainVector22 = VectorLoadAligned(StartGains + 12);
+		const VectorRegister DestinationVector22 = VectorLoadAligned(EndGains + 12);
+		const VectorRegister GainDeltasVector22 = VectorDivide(VectorSubtract(DestinationVector22, GainVector22), NumFramesVector);
+
+		for (int32 i = 0; i < NumFrames; i++)
+		{
+			const VectorRegister Input1 = VectorLoadFloat1(&SourceBuffer[i*2]);
+			const VectorRegister Input2 = VectorLoadFloat1(&SourceBuffer[i * 2 + 1]);
+
+			VectorRegister Result = VectorMultiply(Input1, GainVector11);
+			Result = VectorMultiplyAdd(Input2, GainVector21, Result);
+			VectorStoreAligned(Result, &DestinationBuffer[i * 8]);
+
+			GainVector11 = VectorAdd(GainVector11, GainDeltasVector11);
+			GainVector21 = VectorAdd(GainVector21, GainDeltasVector21);
+
+			Result = VectorMultiply(Input1, GainVector12);
+			Result = VectorMultiplyAdd(Input2, GainVector22, Result);
+			VectorStoreAligned(Result, &DestinationBuffer[i * 8 + 4]);
+
+			GainVector12 = VectorAdd(GainVector12, GainDeltasVector12);
+			GainVector22 = VectorAdd(GainVector22, GainDeltasVector22);
+		}
+	}
+
+	/**
+	 * These functions are non-vectorized versions of the Mix[N]ChannelsTo[N]Channels functions above:
+	 */
+	void DownmixBuffer(int32 NumSourceChannels, int32 NumDestinationChannels, const AlignedFloatBuffer& SourceBuffer, AlignedFloatBuffer& DestinationBuffer, const float* RESTRICT Gains)
+	{
+		DownmixBuffer(NumSourceChannels, NumDestinationChannels, SourceBuffer.GetData(), DestinationBuffer.GetData(), SourceBuffer.Num() / NumSourceChannels, Gains);
+	}
+
+	void DownmixBuffer(int32 NumSourceChannels, int32 NumDestinationChannels, const float* RESTRICT SourceBuffer, float* RESTRICT DestinationBuffer, int32 NumFrames, const float* RESTRICT Gains)
+	{
+		for (int32 FrameIndex = 0; FrameIndex < NumFrames; FrameIndex++)
+		{
+			float* RESTRICT OutputFrame = &DestinationBuffer[FrameIndex * NumDestinationChannels];
+			const float* RESTRICT InputFrame = &SourceBuffer[FrameIndex * NumSourceChannels];
+
+			FMemory::Memzero(OutputFrame, NumDestinationChannels * sizeof(float));
+			for (int32 OutputChannelIndex = 0; OutputChannelIndex < NumDestinationChannels; OutputChannelIndex++)
+			{
+				for (int32 InputChannelIndex = 0; InputChannelIndex < NumSourceChannels; InputChannelIndex++)
+				{
+					OutputFrame[OutputChannelIndex] += InputFrame[InputChannelIndex] * Gains[InputChannelIndex * NumDestinationChannels + OutputChannelIndex];
+				}
+			}
+		}
+	}
+
+	void DownmixBuffer(int32 NumSourceChannels, int32 NumDestinationChannels, const AlignedFloatBuffer& SourceBuffer, AlignedFloatBuffer& DestinationBuffer, float* RESTRICT StartGains, const float* RESTRICT EndGains)
+	{
+		DownmixBuffer(NumSourceChannels, NumDestinationChannels, SourceBuffer.GetData(), DestinationBuffer.GetData(), SourceBuffer.Num() / NumSourceChannels, StartGains, EndGains);
+	}
+
+	void DownmixBuffer(int32 NumSourceChannels, int32 NumDestinationChannels, const float* RESTRICT SourceBuffer, float* RESTRICT DestinationBuffer, int32 NumFrames, float* RESTRICT StartGains, const float* RESTRICT EndGains)
+	{
+		// First, build a map of the per-frame delta that we will use to increment StartGains every frame:
+		alignas(16) float GainDeltas[AUDIO_MIXER_MAX_OUTPUT_CHANNELS * AUDIO_MIXER_MAX_OUTPUT_CHANNELS];
+		
+		for (int32 OutputChannelIndex = 0; OutputChannelIndex < NumDestinationChannels; OutputChannelIndex++)
+		{
+			for (int32 InputChannelIndex = 0; InputChannelIndex < NumSourceChannels; InputChannelIndex++)
+			{
+				const int32 GainMatrixIndex = OutputChannelIndex * NumDestinationChannels + InputChannelIndex;
+				GainDeltas[GainMatrixIndex] = (EndGains[GainMatrixIndex] - StartGains[GainMatrixIndex]) / NumFrames;
+			}
+		}
+
+		for (int32 FrameIndex = 0; FrameIndex < NumFrames; FrameIndex++)
+		{
+			float* RESTRICT OutputFrame = &DestinationBuffer[FrameIndex * NumDestinationChannels];
+			const float* RESTRICT InputFrame = &SourceBuffer[FrameIndex * NumSourceChannels];
+
+			FMemory::Memzero(OutputFrame, NumDestinationChannels * sizeof(float));
+			for (int32 OutputChannelIndex = 0; OutputChannelIndex < NumDestinationChannels; OutputChannelIndex++)
+			{
+				for (int32 InputChannelIndex = 0; InputChannelIndex < NumSourceChannels; InputChannelIndex++)
+				{
+					const int32 GainMatrixIndex = InputChannelIndex * NumDestinationChannels + OutputChannelIndex;
+					OutputFrame[OutputChannelIndex] += InputFrame[InputChannelIndex] * StartGains[GainMatrixIndex];
+					StartGains[GainMatrixIndex] += GainDeltas[GainMatrixIndex];
+				}
+			}
+		}
+	}
 }
diff --git a/Engine/Source/Runtime/AudioMixer/Public/AudioMixer.h b/Engine/Source/Runtime/AudioMixer/Public/AudioMixer.h
--- a/Engine/Source/Runtime/AudioMixer/Public/AudioMixer.h
+++ b/Engine/Source/Runtime/AudioMixer/Public/AudioMixer.h
@@ -19,7 +19,7 @@
 
 #ifndef AUDIO_MIXER_ENABLE_DEBUG_MODE
 // This define enables a bunch of more expensive debug checks and logging capabilities that are intended to be off most of the time even in debug builds of game/editor.
-#if (UE_BUILD_SHIPPING || UE_BUILD_TEST)
+#if (UE_BUILD_SHIPPING || UE_BUILD_TEST || UE_BUILD_DEVELOPMENT)
 #define AUDIO_MIXER_ENABLE_DEBUG_MODE 0
 #else
 #define AUDIO_MIXER_ENABLE_DEBUG_MODE 1
diff --git a/Engine/Source/Runtime/AudioMixer/Public/AudioMixerDevice.h b/Engine/Source/Runtime/AudioMixer/Public/AudioMixerDevice.h
--- a/Engine/Source/Runtime/AudioMixer/Public/AudioMixerDevice.h
+++ b/Engine/Source/Runtime/AudioMixer/Public/AudioMixerDevice.h
@@ -156,8 +156,12 @@ namespace Audio
 
 		int32 GetNumOutputFrames() const { return PlatformSettings.CallbackBufferFrameSize; }
 
+		// Builds a 3D channel map for a spatialized source.
 		void Get3DChannelMap(const ESubmixChannelFormat InSubmixChannelType, const FWaveInstance* InWaveInstance, const float EmitterAzimuth, const float NormalizedOmniRadius, Audio::AlignedFloatBuffer& OutChannelMap);
+
+		// Builds a channel gain matrix for a non-spatialized source. The non-static variation of this function queries AudioMixerDevice->NumOutputChannels directly which may not be thread safe.
 		void Get2DChannelMap(bool bIsVorbis, const ESubmixChannelFormat InSubmixChannelType, const int32 NumSourceChannels, const bool bIsCenterChannelOnly, Audio::AlignedFloatBuffer& OutChannelMap) const;
+		static void Get2DChannelMap(bool bIsVorbis, const int32 NumSourceChannels, const int32 NumOutputChannels, const bool bIsCenterChannelOnly, Audio::AlignedFloatBuffer& OutChannelMap);
 
 		int32 GetDeviceSampleRate() const;
 		int32 GetDeviceOutputChannels() const;
@@ -202,7 +206,7 @@ namespace Audio
 
 		void Get2DChannelMapInternal(const int32 NumSourceChannels, const int32 NumOutputChannels, const bool bIsCenterChannelOnly, TArray<float>& OutChannelMap) const;
 		void InitializeChannelMaps();
-		int32 GetChannelMapCacheId(const int32 NumSourceChannels, const int32 NumOutputChannels, const bool bIsCenterChannelOnly) const;
+		static int32 GetChannelMapCacheId(const int32 NumSourceChannels, const int32 NumOutputChannels, const bool bIsCenterChannelOnly);
 		void CacheChannelMap(const int32 NumSourceChannels, const int32 NumOutputChannels, const bool bIsCenterChannelOnly);
 		void InitializeChannelAzimuthMap(const int32 NumChannels);
 
diff --git a/Engine/Source/Runtime/AudioMixer/Public/DSP/BufferVectorOperations.h b/Engine/Source/Runtime/AudioMixer/Public/DSP/BufferVectorOperations.h
--- a/Engine/Source/Runtime/AudioMixer/Public/DSP/BufferVectorOperations.h
+++ b/Engine/Source/Runtime/AudioMixer/Public/DSP/BufferVectorOperations.h
@@ -16,37 +16,158 @@ namespace Audio
 	typedef TArray<float, TAlignedHeapAllocator<AUDIO_BUFFER_ALIGNMENT>> AlignedFloatBuffer;
 	typedef TArray<uint8, TAlignedHeapAllocator<AUDIO_BUFFER_ALIGNMENT>> AlignedByteBuffer;
 
+	/** CHANNEL-AGNOSTIC OPERATIONS */
+
 	/** Multiplies the input aligned float buffer with the given value. */
 	void BufferMultiplyByConstant(const AlignedFloatBuffer& InFloatBuffer, float InValue, AlignedFloatBuffer& OutFloatBuffer);
 
 	/** Similar to BufferMultiplyByConstant, but (a) assumes a buffer length divisible by 4 and (b) performs the multiply in place. */
 	void MultiplyBufferByConstantInPlace(AlignedFloatBuffer& InBuffer, float InGain);
-	void MultiplyBufferByConstantInPlace(float* InBuffer, int32 NumSamples, float InGain);
+	void MultiplyBufferByConstantInPlace(float* RESTRICT InBuffer, int32 NumSamples, float InGain);
 
 	/* Takes a float buffer and quickly interpolates it's gain from StartValue to EndValue. */
 	/* This operation completely ignores channel counts, so avoid using this function on buffers that are not mono, stereo or quad */
 	/* if the buffer needs to fade all channels uniformly. */
 	void FadeBufferFast(AlignedFloatBuffer& OutFloatBuffer, const float StartValue, const float EndValue);
-	void FadeBufferFast(float* OutFloatBuffer, int32 NumSamples, const float StartValue, const float EndValue);
+	void FadeBufferFast(float* RESTRICT OutFloatBuffer, int32 NumSamples, const float StartValue, const float EndValue);
 
-	// Takes buffer InFloatBuffer, optionally multiplies it by Gain, and adds it to BufferToSumTo.
+	/** Takes buffer InFloatBuffer, optionally multiplies it by Gain, and adds it to BufferToSumTo. */
 	void MixInBufferFast(const AlignedFloatBuffer& InFloatBuffer, AlignedFloatBuffer& BufferToSumTo, const float Gain);
-	void MixInBufferFast(const float* InFloatBuffer, float* BufferToSumTo, int32 NumSamples, const float Gain);
-	void MixInBufferFast(const float* InFloatBuffer, float* BufferToSumTo, int32 NumSamples);
+	void MixInBufferFast(const float* RESTRICT InFloatBuffer, float* RESTRICT BufferToSumTo, int32 NumSamples, const float Gain);
+	void MixInBufferFast(const float* RESTRICT InFloatBuffer, float* RESTRICT BufferToSumTo, int32 NumSamples);
 
-	// Sums two buffers together and places the result in the resulting buffer.
+	/** Sums two buffers together and places the result in the resulting buffer. */
 	void SumBuffers(const AlignedFloatBuffer& InFloatBuffer1, const AlignedFloatBuffer& InFloatBuffer2, AlignedFloatBuffer& OutputBuffer);
-	void SumBuffers(const float* InFloatBuffer1, const float* InFloatBuffer2, float* OutputBuffer, int32 NumSamples);
+	void SumBuffers(const float* RESTRICT InFloatBuffer1, const float* RESTRICT InFloatBuffer2, float* RESTRICT OutputBuffer, int32 NumSamples);
 
-	// Multiply the second buffer in place by the first buffer.
+	/** Multiply the second buffer in place by the first buffer. */
 	void MultiplyBuffersInPlace(const AlignedFloatBuffer& InFloatBuffer, AlignedFloatBuffer& BufferToMultiply);
-	void MultiplyBuffersInPlace(const float* InFloatBuffer, float* BufferToMultiply, int32 NumSamples);
+	void MultiplyBuffersInPlace(const float* RESTRICT InFloatBuffer, float* RESTRICT BufferToMultiply, int32 NumSamples);
+
+	/** CHANNEL-AGNOSTIC ANALYSIS OPERATIONS */
 
-	// Takes an audio buffer and returns the magnitude across that buffer.
+	/** Takes an audio buffer and returns the magnitude across that buffer. */
 	float GetMagnitude(const AlignedFloatBuffer& Buffer);
-	float GetMagnitude(const float* Buffer, int32 NumSamples);
+	float GetMagnitude(const float* RESTRICT Buffer, int32 NumSamples);
 
-	// Takes an audio buffer and gets the average absolute amplitude across that buffer.
+	/** Takes an audio buffer and gets the average absolute amplitude across that buffer. */
 	float GetAverageAmplitude(const AlignedFloatBuffer& Buffer);
-	float GetAverageAmplitude(const float* Buffer, int32 NumSamples);
+	float GetAverageAmplitude(const float* RESTRICT Buffer, int32 NumSamples);
+
+	/** CHANNEL-SPECIFIC OPERATIONS */
+
+	/** Takes a 2 channel interleaved buffer and applies Gains to it. Gains is expected to point to a 2 float long buffer.
+	 *  StereoBuffer must have an even number of frames.
+	 *  If StartGains and EndGains are provided, this function will interpolate between the two across the buffer.
+	*/
+	void Apply2ChannelGain(AlignedFloatBuffer& StereoBuffer, const float* RESTRICT Gains);
+	void Apply2ChannelGain(float* RESTRICT StereoBuffer, int32 NumSamples, const float* RESTRICT Gains);
+	void Apply2ChannelGain(AlignedFloatBuffer& StereoBuffer, const float* RESTRICT StartGains, const float* RESTRICT EndGains);
+	void Apply2ChannelGain(float* RESTRICT StereoBuffer, int32 NumSamples, const float* RESTRICT StartGains, const float* RESTRICT EndGains);
+
+	/** Takes a 1 channel buffer and mixes it to a stereo buffer using Gains. Gains is expected to point to a 2 float long buffer.
+	 *  these buffers must have an even number of frames.
+	 *  If StartGains and EndGains are provided, this function will interpolate between the two across the buffer.
+	*/
+	void MixMonoTo2ChannelsFast(const AlignedFloatBuffer& MonoBuffer, AlignedFloatBuffer& DestinationBuffer, const float* RESTRICT Gains);
+	void MixMonoTo2ChannelsFast(const float* RESTRICT MonoBuffer, float* RESTRICT DestinationBuffer, int32 NumFrames, const float* RESTRICT Gains);
+	void MixMonoTo2ChannelsFast(const AlignedFloatBuffer& MonoBuffer, AlignedFloatBuffer& DestinationBuffer, const float* RESTRICT StartGains, const float* RESTRICT EndGains);
+	void MixMonoTo2ChannelsFast(const float* RESTRICT MonoBuffer, float* RESTRICT DestinationBuffer, int32 NumFrames, const float* RESTRICT StartGains, const float* RESTRICT EndGains);
+
+	/** Takes a 2 channel buffer and mixes it to an 2 channel interleaved buffer using Gains. Gains is expected to point to a 16 float long buffer.
+	*  Output gains for the left input channel should be the first 8 values in Gains, and Output gains for the right input channel should be rest.
+	*  NumFrames must be a multiple of 4.
+	*  If StartGains and EndGains are provided, this function will interpolate between the two across the buffer.
+	*/
+	void Mix2ChannelsTo2ChannelsFast(const AlignedFloatBuffer& SourceBuffer, AlignedFloatBuffer& DestinationBuffer, const float* RESTRICT Gains);
+	void Mix2ChannelsTo2ChannelsFast(const float* RESTRICT SourceBuffer, float* RESTRICT DestinationBuffer, int32 NumFrames, const float* RESTRICT Gains);
+	void Mix2ChannelsTo2ChannelsFast(const AlignedFloatBuffer& SourceBuffer, AlignedFloatBuffer& DestinationBuffer, const float* RESTRICT StartGains, const float* RESTRICT EndGains);
+	void Mix2ChannelsTo2ChannelsFast(const float* RESTRICT SourceBuffer, float* RESTRICT DestinationBuffer, int32 NumFrames, const float* RESTRICT StartGains, const float* RESTRICT EndGains);
+
+	/** Takes a 4 channel interleaved buffer and applies Gains to it. Gains is expected to point to a 2 float long buffer.
+	 *  If StartGains and EndGains are provided, this function will interpolate between the two across the buffer.
+	*/
+	void Apply4ChannelGain(AlignedFloatBuffer& InterleavedBuffer, const float* RESTRICT Gains);
+	void Apply4ChannelGain(float* RESTRICT InterleavedBuffer, int32 NumSamples, const float* RESTRICT Gains);
+	void Apply4ChannelGain(AlignedFloatBuffer& InterleavedBuffer, const float* RESTRICT StartGains, const float* RESTRICT EndGains);
+	void Apply4ChannelGain(float* RESTRICT InterleavedBuffer, int32 NumSamples, const float* RESTRICT StartGains, const float* RESTRICT EndGains);
+
+	/** Takes a 1 channel buffer and mixes it to an 8 channel interleaved buffer using Gains. Gains is expected to point to a 8 float long buffer.
+	*  these buffers must have an even number of frames.
+	*  If StartGains and EndGains are provided, this function will interpolate between the two across the buffer.
+	*/
+	void MixMonoTo4ChannelsFast(const AlignedFloatBuffer& MonoBuffer, AlignedFloatBuffer& DestinationBuffer, const float* RESTRICT Gains);
+	void MixMonoTo4ChannelsFast(const float* RESTRICT MonoBuffer, float* RESTRICT DestinationBuffer, int32 NumFrames, const float* RESTRICT Gains);
+	void MixMonoTo4ChannelsFast(const AlignedFloatBuffer& MonoBuffer, AlignedFloatBuffer& DestinationBuffer, const float* RESTRICT StartGains, const float* RESTRICT EndGains);
+	void MixMonoTo4ChannelsFast(const float* RESTRICT MonoBuffer, float* RESTRICT DestinationBuffer, int32 NumFrames, const float* RESTRICT StartGains, const float* RESTRICT EndGains);
+
+	/** Takes a 2 channel buffer and mixes it to an 8 channel interleaved buffer using Gains. Gains is expected to point to a 16 float long buffer.
+	*  Output gains for the left input channel should be the first 8 values in Gains, and Output gains for the right input channel should be rest.
+	*  NumFrames must be a multiple of 4.
+	*  If StartGains and EndGains are provided, this function will interpolate between the two across the buffer.
+	*/
+	void Mix2ChannelsTo4ChannelsFast(const AlignedFloatBuffer& SourceBuffer, AlignedFloatBuffer& DestinationBuffer, const float* RESTRICT Gains);
+	void Mix2ChannelsTo4ChannelsFast(const float* RESTRICT SourceBuffer, float* RESTRICT DestinationBuffer, int32 NumFrames, const float* RESTRICT Gains);
+	void Mix2ChannelsTo4ChannelsFast(const AlignedFloatBuffer& SourceBuffer, AlignedFloatBuffer& DestinationBuffer, const float* RESTRICT StartGains, const float* RESTRICT EndGains);
+	void Mix2ChannelsTo4ChannelsFast(const float* RESTRICT SourceBuffer, float* RESTRICT DestinationBuffer, int32 NumFrames, const float* RESTRICT StartGains, const float* RESTRICT EndGains);
+
+	/** Takes a 6 channel interleaved buffer and applies Gains to it. Gains is expected to point to a 2 float long buffer.
+	 *  InterleavedBuffer must have an even number of frames.
+	*/
+	void Apply6ChannelGain(AlignedFloatBuffer& InterleavedBuffer, const float* RESTRICT Gains);
+	void Apply6ChannelGain(float* RESTRICT InterleavedBuffer, int32 NumSamples, const float* RESTRICT Gains);
+	void Apply6ChannelGain(AlignedFloatBuffer& InterleavedBuffer, const float* RESTRICT StartGains, const float* RESTRICT EndGains);
+	void Apply6ChannelGain(float* RESTRICT InterleavedBuffer, int32 NumSamples, const float* RESTRICT StartGains, const float* RESTRICT EndGains);
+
+	/** Takes a 1 channel buffer and mixes it to an 8 channel interleaved buffer using Gains. Gains is expected to point to a 8 float long buffer.
+	 *  these buffers must have an even number of frames.
+	 *  If StartGains and EndGains are provided, this function will interpolate between the two across the buffer.
+	*/
+	void MixMonoTo6ChannelsFast(const AlignedFloatBuffer& MonoBuffer, AlignedFloatBuffer& DestinationBuffer, const float* RESTRICT Gains);
+	void MixMonoTo6ChannelsFast(const float* RESTRICT MonoBuffer, float* RESTRICT DestinationBuffer, int32 NumFrames, const float* RESTRICT Gains);
+	void MixMonoTo6ChannelsFast(const AlignedFloatBuffer& MonoBuffer, AlignedFloatBuffer& DestinationBuffer, const float* RESTRICT StartGains, const float* RESTRICT EndGains);
+	void MixMonoTo6ChannelsFast(const float* RESTRICT MonoBuffer, float* RESTRICT DestinationBuffer, int32 NumFrames, const float* RESTRICT StartGains, const float* RESTRICT EndGains);
+
+	/** Takes a 2 channel buffer and mixes it to an 8 channel interleaved buffer using Gains. Gains is expected to point to a 16 float long buffer.
+	 *  Output gains for the left input channel should be the first 8 values in Gains, and Output gains for the right input channel should be rest.
+	 *  NumFrames must be a multiple of 4.
+	 *  If StartGains and EndGains are provided, this function will interpolate between the two across the buffer.
+	*/
+	void Mix2ChannelsTo6ChannelsFast(const AlignedFloatBuffer& SourceBuffer, AlignedFloatBuffer& DestinationBuffer, const float* RESTRICT Gains);
+	void Mix2ChannelsTo6ChannelsFast(const float* RESTRICT SourceBuffer, float* RESTRICT DestinationBuffer, int32 NumFrames, const float* RESTRICT Gains);
+	void Mix2ChannelsTo6ChannelsFast(const AlignedFloatBuffer& SourceBuffer, AlignedFloatBuffer& DestinationBuffer, const float* RESTRICT StartGains, const float* RESTRICT EndGains);
+	void Mix2ChannelsTo6ChannelsFast(const float* RESTRICT SourceBuffer, float* RESTRICT DestinationBuffer, int32 NumFrames, const float* RESTRICT StartGains, const float* RESTRICT EndGains);
+
+	/** Takes an 8 channel interleaved buffer and applies Gains to it. Gains is expected to point to an 8 float long buffer. */
+	void Apply8ChannelGain(AlignedFloatBuffer& InterleavedBuffer, const float* RESTRICT Gains);
+	void Apply8ChannelGain(float* RESTRICT InterleavedBuffer, int32 NumSamples, const float* RESTRICT Gains);
+	void Apply8ChannelGain(AlignedFloatBuffer& InterleavedBuffer, const float* RESTRICT StartGains, const float* RESTRICT EndGains);
+	void Apply8ChannelGain(float* RESTRICT InterleavedBuffer, int32 NumSamples, const float* RESTRICT StartGains, const float* RESTRICT EndGains);
+
+	/** Takes a 1 channel buffer and mixes it to an 8 channel interleaved buffer using Gains. Gains is expected to point to a 8 float long buffer.
+	 *  these buffers must have an even number of frames.
+	 *  If StartGains and EndGains are provided, this function will interpolate between the two across the buffer.
+	*/
+	void MixMonoTo8ChannelsFast(const AlignedFloatBuffer& MonoBuffer, AlignedFloatBuffer& DestinationBuffer, const float* RESTRICT Gains);
+	void MixMonoTo8ChannelsFast(const float* RESTRICT MonoBuffer, float* RESTRICT DestinationBuffer, int32 NumFrames, const float* RESTRICT Gains);
+	void MixMonoTo8ChannelsFast(const AlignedFloatBuffer& MonoBuffer, AlignedFloatBuffer& DestinationBuffer, const float* RESTRICT StartGains, const float* RESTRICT EndGains);
+	void MixMonoTo8ChannelsFast(const float* RESTRICT MonoBuffer, float* RESTRICT DestinationBuffer, int32 NumFrames, const float* RESTRICT StartGains, const float* RESTRICT EndGains);
+
+	/** Takes a 2 channel buffer and mixes it to an 8 channel interleaved buffer using Gains. Gains is expected to point to a 16 float long buffer.
+	 *  Output gains for the left input channel should be the first 8 values in Gains, and Output gains for the right input channel should be rest.
+	 *  these buffers must have an even number of frames.
+	 *  If StartGains and EndGains are provided, this function will interpolate between the two across the buffer.
+	*/
+	void Mix2ChannelsTo8ChannelsFast(const AlignedFloatBuffer& SourceBuffer, AlignedFloatBuffer& DestinationBuffer, const float* RESTRICT Gains);
+	void Mix2ChannelsTo8ChannelsFast(const float* RESTRICT SourceBuffer, float* RESTRICT DestinationBuffer, int32 NumFrames, const float* RESTRICT Gains);
+	void Mix2ChannelsTo8ChannelsFast(const AlignedFloatBuffer& SourceBuffer, AlignedFloatBuffer& DestinationBuffer, const float* RESTRICT StartGains, const float* RESTRICT EndGains);
+	void Mix2ChannelsTo8ChannelsFast(const float* RESTRICT SourceBuffer, float* RESTRICT DestinationBuffer, int32 NumFrames, const float* RESTRICT StartGains, const float* RESTRICT EndGains);
+
+	/** This is a generalized operation that uses the channel gain matrix provied in Gains to mix an interleaved source buffer to the interleaved downmix buffer.
+	 *  This operation is not explicitly vectorized and will almost always be slower than using one of the functions above.
+	*/
+	void DownmixBuffer(int32 NumSourceChannels, int32 NumDestinationChannels, const AlignedFloatBuffer& SourceBuffer, AlignedFloatBuffer& DestinationBuffer, const float* RESTRICT Gains);
+	void DownmixBuffer(int32 NumSourceChannels, int32 NumDestinationChannels, const float* RESTRICT SourceBuffer, float* RESTRICT DestinationBuffer, int32 NumFrames, const float* RESTRICT Gains);
+	void DownmixBuffer(int32 NumSourceChannels, int32 NumDestinationChannels, const AlignedFloatBuffer& SourceBuffer, AlignedFloatBuffer& DestinationBuffer, float* RESTRICT StartGains, const float* RESTRICT EndGains);
+	void DownmixBuffer(int32 NumSourceChannels, int32 NumDestinationChannels, const float* RESTRICT SourceBuffer, float* RESTRICT DestinationBuffer, int32 NumFrames, float* RESTRICT StartGains, const float* RESTRICT EndGains);
 }
