[NO NON-PII COMMIT MESSAGE]

diff --git a/Engine/Source/Runtime/AudioMixer/Private/DSP/BufferVectorOperations.cpp b/Engine/Source/Runtime/AudioMixer/Private/DSP/BufferVectorOperations.cpp
--- a/Engine/Source/Runtime/AudioMixer/Private/DSP/BufferVectorOperations.cpp
+++ b/Engine/Source/Runtime/AudioMixer/Private/DSP/BufferVectorOperations.cpp
@@ -407,7 +407,7 @@ namespace Audio
 		// Initialize GainVector at StartGains and compute GainDeltasVector:
 		VectorRegister GainVector = VectorLoadFloat2(StartGains);
 		const VectorRegister DestinationVector = VectorLoadFloat2(EndGains);
-		const VectorRegister NumFramesVector = VectorSetFloat1(NumSamples / 4);
+		const VectorRegister NumFramesVector = VectorSetFloat1(NumSamples / 4.0f);
 		const VectorRegister GainDeltasVector = VectorDivide(VectorSubtract(DestinationVector, GainVector), NumFramesVector);
 
 		for (int32 i = 0; i < NumSamples; i += 4)
@@ -472,7 +472,7 @@ namespace Audio
 		// Initialize GainVector at StartGains and compute GainDeltasVector:
 		VectorRegister GainVector = VectorLoadFloat2(StartGains);
 		const VectorRegister DestinationVector = VectorLoadFloat2(EndGains);
-		const VectorRegister NumFramesVector = VectorSetFloat1(NumFrames / 2);
+		const VectorRegister NumFramesVector = VectorSetFloat1(NumFrames / 2.0f);
 		const VectorRegister GainDeltasVector = VectorDivide(VectorSubtract(DestinationVector, GainVector), NumFramesVector);
 
 		// To help with stair stepping, we initialize the second frame in GainVector to be half a GainDeltas vector higher than the first frame.
@@ -555,7 +555,7 @@ namespace Audio
 
 	void Mix2ChannelsTo2ChannelsFast(const float* RESTRICT SourceBuffer, float* RESTRICT DestinationBuffer, int32 NumFrames, const float* RESTRICT StartGains, const float* RESTRICT EndGains)
 	{
-		const VectorRegister NumFramesVector = VectorSetFloat1(NumFrames / 2);
+		const VectorRegister NumFramesVector = VectorSetFloat1(NumFrames / 2.0f);
 
 		VectorRegister GainVector1 = VectorLoadFloat2(StartGains);
 		const VectorRegister DestinationVector1 = VectorLoadFloat2(EndGains);
@@ -616,7 +616,7 @@ namespace Audio
 		// Initialize GainVector at StartGains and compute GainDeltasVector:
 		VectorRegister GainVector = VectorLoadAligned(StartGains);
 		const VectorRegister DestinationVector = VectorLoadAligned(EndGains);
-		const VectorRegister NumFramesVector = VectorSetFloat1(NumSamples / 4);
+		const VectorRegister NumFramesVector = VectorSetFloat1(NumSamples / 4.0f);
 		const VectorRegister GainDeltasVector = VectorDivide(VectorSubtract(DestinationVector, GainVector), NumFramesVector);
 
 		for (int32 i = 0; i < NumSamples; i += 4)
@@ -683,7 +683,7 @@ namespace Audio
 	{
 		VectorRegister GainVector = VectorLoadAligned(StartGains);
 		const VectorRegister DestinationVector = VectorLoadAligned(EndGains);
-		const VectorRegister NumFramesVector = VectorSetFloat1(NumFrames);
+		const VectorRegister NumFramesVector = VectorSetFloat1((float) NumFrames);
 		const VectorRegister GainDeltasVector = VectorDivide(VectorSubtract(DestinationVector, GainVector), NumFramesVector);
 
 		for (int32 i = 0; i < NumFrames; i++)
@@ -759,7 +759,7 @@ namespace Audio
 	*/
 	void Mix2ChannelsTo4ChannelsFast(const float* RESTRICT SourceBuffer, float* RESTRICT DestinationBuffer, int32 NumFrames, const float* RESTRICT StartGains, const float* RESTRICT EndGains)
 	{
-		const VectorRegister NumFramesVector = VectorSetFloat1(NumFrames);
+		const VectorRegister NumFramesVector = VectorSetFloat1((float) NumFrames);
 
 		VectorRegister GainVector1 = VectorLoadAligned(StartGains);
 		const VectorRegister DestinationVector1 = VectorLoadAligned(EndGains);
@@ -817,7 +817,7 @@ namespace Audio
 
 	void Apply6ChannelGain(float* RESTRICT InterleavedBuffer, int32 NumSamples, const float* RESTRICT StartGains, const float* RESTRICT EndGains)
 	{
-		const VectorRegister NumFramesVector = VectorSetFloat1(NumSamples / 12);
+		const VectorRegister NumFramesVector = VectorSetFloat1(NumSamples / 12.0f);
 
 		VectorRegister GainVector1 = VectorLoadAligned(StartGains);
 		const VectorRegister DestinationVector1 = VectorLoadAligned(EndGains);
@@ -918,7 +918,7 @@ namespace Audio
 
 	void MixMonoTo6ChannelsFast(const float* RESTRICT MonoBuffer, float* RESTRICT DestinationBuffer, int32 NumFrames, const float* RESTRICT StartGains, const float* RESTRICT EndGains)
 	{
-		const VectorRegister NumFramesVector = VectorSetFloat1(NumFrames / 2);
+		const VectorRegister NumFramesVector = VectorSetFloat1(NumFrames / 2.0f);
 
 		VectorRegister GainVector1 = VectorLoadAligned(StartGains);
 		const VectorRegister DestinationVector1 = VectorLoadAligned(EndGains);
@@ -1042,7 +1042,7 @@ namespace Audio
 
 	void Mix2ChannelsTo6ChannelsFast(const float* RESTRICT SourceBuffer, float* RESTRICT DestinationBuffer, int32 NumFrames, const float* RESTRICT StartGains, const float* RESTRICT EndGains)
 	{
-		const VectorRegister NumFramesVector = VectorSetFloat1(NumFrames / 2);
+		const VectorRegister NumFramesVector = VectorSetFloat1(NumFrames / 2.0f);
 
 		VectorRegister GainVector11 = VectorLoadAligned(StartGains);
 		const VectorRegister DestinationVector11 = VectorLoadAligned(EndGains);
@@ -1151,7 +1151,7 @@ namespace Audio
 
 	void Apply8ChannelGain(float* RESTRICT InterleavedBuffer, int32 NumSamples, const float* RESTRICT StartGains, const float* RESTRICT EndGains)
 	{
-		const VectorRegister NumFramesVector = VectorSetFloat1(NumSamples / 8);
+		const VectorRegister NumFramesVector = VectorSetFloat1(NumSamples / 8.0f);
 
 		VectorRegister GainVector1 = VectorLoadAligned(StartGains);
 		const VectorRegister DestinationVector1 = VectorLoadAligned(EndGains);
@@ -1233,7 +1233,7 @@ namespace Audio
 	*/
 	void MixMonoTo8ChannelsFast(const float* RESTRICT MonoBuffer, float* RESTRICT DestinationBuffer, int32 NumFrames, const float* RESTRICT StartGains, const float* RESTRICT EndGains)
 	{
-		const VectorRegister NumFramesVector = VectorSetFloat1(NumFrames);
+		const VectorRegister NumFramesVector = VectorSetFloat1((float) NumFrames);
 
 		VectorRegister GainVector1 = VectorLoadAligned(StartGains);
 		const VectorRegister DestinationVector1 = VectorLoadAligned(EndGains);
@@ -1331,7 +1331,7 @@ namespace Audio
 
 	void Mix2ChannelsTo8ChannelsFast(const float* RESTRICT SourceBuffer, float* RESTRICT DestinationBuffer, int32 NumFrames, const float* RESTRICT StartGains, const float* RESTRICT EndGains)
 	{
-		const VectorRegister NumFramesVector = VectorSetFloat1(NumFrames);
+		const VectorRegister NumFramesVector = VectorSetFloat1((float) NumFrames);
 
 		VectorRegister GainVector11 = VectorLoadAligned(StartGains);
 		const VectorRegister DestinationVector11 = VectorLoadAligned(EndGains);
