Fixing build

diff --git a/Engine/Source/Runtime/AudioMixer/Private/AudioMixerChannelMaps.cpp b/Engine/Source/Runtime/AudioMixer/Private/AudioMixerChannelMaps.cpp
--- a/Engine/Source/Runtime/AudioMixer/Private/AudioMixerChannelMaps.cpp
+++ b/Engine/Source/Runtime/AudioMixer/Private/AudioMixerChannelMaps.cpp
@@ -204,8 +204,6 @@ namespace Audio
 
 	void FMixerDevice::Get2DChannelMap(bool bIsVorbis, const ESubmixChannelFormat InSubmixChannelType, const int32 NumSourceChannels, const bool bIsCenterChannelOnly, Audio::AlignedFloatBuffer& OutChannelMap) const
 	{
-		AUDIO_MIXER_CHECK_GAME_THREAD(MixerDevice);
-
 		int32 NumOutputChannels = GetNumChannelsForSubmixFormat(InSubmixChannelType);
 		Get2DChannelMap(bIsVorbis, NumSourceChannels, NumOutputChannels, bIsCenterChannelOnly, OutChannelMap);
 	}
