[NO NON-PII COMMIT MESSAGE]

diff --git a/Engine/Source/Runtime/AudioMixer/Private/AudioMixerSourceManager.cpp b/Engine/Source/Runtime/AudioMixer/Private/AudioMixerSourceManager.cpp
--- a/Engine/Source/Runtime/AudioMixer/Private/AudioMixerSourceManager.cpp
+++ b/Engine/Source/Runtime/AudioMixer/Private/AudioMixerSourceManager.cpp
@@ -2088,13 +2088,20 @@ namespace Audio
 				continue;
 			}
 
+			FSourceDownmixData& DownmixData = DownmixDataArray[SourceId];
+
+			if (DownmixData.PostEffectBuffers == nullptr)
+			{
+				continue;
+			}
+
 			if (SourceInfo.bIs3D)
 			{
-				ComputeDownmix3D(DownmixDataArray[SourceId]);
+				ComputeDownmix3D(DownmixData);
 			}
 			else
 			{
-				ComputeDownmix2D(DownmixDataArray[SourceId]);
+				ComputeDownmix2D(DownmixData);
 			}
 		}
 	}
diff --git a/Engine/Source/Runtime/AudioMixer/Private/AudioMixerSourceManager.h b/Engine/Source/Runtime/AudioMixer/Private/AudioMixerSourceManager.h
--- a/Engine/Source/Runtime/AudioMixer/Private/AudioMixerSourceManager.h
+++ b/Engine/Source/Runtime/AudioMixer/Private/AudioMixerSourceManager.h
@@ -156,6 +156,9 @@ namespace Audio
 
 		FORCEINLINE void Reset(int32 InNumInChannels, int32 InNumOutChannels)
 		{
+			checkSlow(InNumInChannels <= AUDIO_MIXER_MAX_OUTPUT_CHANNELS);
+			checkSlow(InNumOutChannels <= AUDIO_MIXER_MAX_OUTPUT_CHANNELS);
+
 			CopySize = InNumInChannels * InNumOutChannels * sizeof(float);
 			FMemory::Memzero(ChannelDestinationGains, CopySize);
 		}
@@ -354,7 +357,8 @@ namespace Audio
 			uint32 NumDeviceChannels;
 
 			FSourceDownmixData(uint32 SourceNumChannels, uint32 NumDeviceOutputChannels, uint32 InNumFrames)
-				: DeviceSubmixInfo(FSubmixChannelTypeInfo(SourceNumChannels, NumDeviceOutputChannels, InNumFrames))
+				: PostEffectBuffers(nullptr)
+				, DeviceSubmixInfo(FSubmixChannelTypeInfo(SourceNumChannels, NumDeviceOutputChannels, InNumFrames))
 				, StereoSubmixInfo(FSubmixChannelTypeInfo(SourceNumChannels, 2, InNumFrames))
 				, QuadSubmixInfo(FSubmixChannelTypeInfo(SourceNumChannels, 4, InNumFrames))
 				, FiveOneSubmixInfo(FSubmixChannelTypeInfo(SourceNumChannels, 6, InNumFrames))
@@ -375,13 +379,14 @@ namespace Audio
 			void ResetData(const uint32 InNumInputChannels)
 			{
 				NumInputChannels = InNumInputChannels;
+				PostEffectBuffers = nullptr;
 
 				DeviceSubmixInfo.Reset(NumInputChannels, NumDeviceChannels, NumFrames);
-				StereoSubmixInfo.Reset(NumInputChannels, NumDeviceChannels, NumFrames);
-				QuadSubmixInfo.Reset(NumInputChannels, NumDeviceChannels, NumFrames);
-				FiveOneSubmixInfo.Reset(NumInputChannels, NumDeviceChannels, NumFrames);
-				SevenOneSubmixInfo.Reset(NumInputChannels, NumDeviceChannels, NumFrames);
-				AmbisonicsSubmixInfo.Reset(NumInputChannels, NumDeviceChannels, NumFrames);
+				StereoSubmixInfo.Reset(NumInputChannels, 2, NumFrames);
+				QuadSubmixInfo.Reset(NumInputChannels, 4, NumFrames);
+				FiveOneSubmixInfo.Reset(NumInputChannels, 6, NumFrames);
+				SevenOneSubmixInfo.Reset(NumInputChannels, 8, NumFrames);
+				AmbisonicsSubmixInfo.Reset(NumInputChannels, 4, NumFrames);
 			}
 		};
 
diff --git a/Engine/Source/Runtime/AudioMixer/Private/DSP/BufferVectorOperations.cpp b/Engine/Source/Runtime/AudioMixer/Private/DSP/BufferVectorOperations.cpp
--- a/Engine/Source/Runtime/AudioMixer/Private/DSP/BufferVectorOperations.cpp
+++ b/Engine/Source/Runtime/AudioMixer/Private/DSP/BufferVectorOperations.cpp
@@ -985,7 +985,7 @@ namespace Audio
 		const VectorRegister GainVector21 = VectorSet(Gains[4], Gains[5], Gains[0], Gains[1]);
 		const VectorRegister GainVector31 = VectorLoad(&Gains[2]);
 
-		const VectorRegister GainVector12 = VectorLoadAligned(Gains + 6);
+		const VectorRegister GainVector12 = VectorLoad(Gains + 6);
 		const VectorRegister GainVector22 = VectorSet(Gains[10], Gains[11], Gains[6], Gains[7]);
 		const VectorRegister GainVector32 = VectorLoadAligned(&Gains[8]);
 
@@ -1034,7 +1034,7 @@ namespace Audio
 	* |        | +        | +       | +       | +       | +        | +       | +       | +       | +        | +       | +        | +        |
 	* | Gain2  | g6       | g7      | g8      | g9      | g10      | g11     | g6      | g7      | g8       | g9      | g10      | g11      |
 	* |        | *        | *       | *       | *       | *        | *       | *       | *       | *        | *       | *        | *        |
-	* | Input2 | i1       | i1      | i1      | i1      | i1       | i1      | i1      | i3      | i3       | i3      | i3       | i3       |
+	* | Input2 | i1       | i1      | i1      | i1      | i1       | i1      | i3      | i3      | i3       | i3      | i3       | i3       |
 	* |        | =        | =       | =       | =       | =        | =       | =       | =       | =        | =       | =        | =        |
 	* | Output | o0       | o1      | o2      | o3      | o4       | o5      | o6      | o7      | o8       | o9      | o10      | o11      |
 	* +--------+----------+---------+---------+---------+----------+---------+---------+---------+----------+---------+----------+----------+
@@ -1066,8 +1066,8 @@ namespace Audio
 		const VectorRegister InitializedDelta31 = VectorMultiply(GainDeltasVector31, DeltaHalf31);
 		GainVector31 = VectorAdd(GainVector31, InitializedDelta31);
 
-		VectorRegister GainVector12 = VectorLoadAligned(StartGains + 6);
-		const VectorRegister DestinationVector12 = VectorLoadAligned(EndGains + 6);
+		VectorRegister GainVector12 = VectorLoad(StartGains + 6);
+		const VectorRegister DestinationVector12 = VectorLoad(EndGains + 6);
 		const VectorRegister GainDeltasVector12 = VectorDivide(VectorSubtract(DestinationVector12, GainVector12), NumFramesVector);
 
 		VectorRegister GainVector22 = VectorSet(StartGains[10], StartGains[11], StartGains[6], StartGains[7]);
