Copied from dev-audio

CL4852302: Synth Sample Player is not playing pre-existing sounds that register 0 samples
CL 4852160: PR #5494: USoundNode Delay and Concatenator fix when not on first child node
CL 4852119: PR #5492 Distance crossfade SoundCue node returns incorrect volume when in range FadeInDistanceMax <= 0.f and CrossFadeInput[ChildNodeIndex].Volume != 1.f

diff --git a/Engine/Plugins/Runtime/Synthesis/Source/Synthesis/Private/SynthComponentWaveTable.cpp b/Engine/Plugins/Runtime/Synthesis/Source/Synthesis/Private/SynthComponentWaveTable.cpp
--- a/Engine/Plugins/Runtime/Synthesis/Source/Synthesis/Private/SynthComponentWaveTable.cpp
+++ b/Engine/Plugins/Runtime/Synthesis/Source/Synthesis/Private/SynthComponentWaveTable.cpp
@@ -10,7 +10,7 @@ USynthSamplePlayer::USynthSamplePlayer(const FObjectInitializer& ObjInitializer)
 	: Super(ObjInitializer)
 	, SoundWave(nullptr)
 	, SampleDurationSec(0.0f)
-	, SamplePlaybackProgressSec(0.0F)
+	, SamplePlaybackProgressSec(0.0f)
 	, bIsLoaded(false)
 {
 	PrimaryComponentTick.bCanEverTick = true;
@@ -134,7 +134,6 @@ void USynthSamplePlayer::LoadSoundWaveInternal()
 	}
 }
 
-
 void USynthSamplePlayer::SetSoundWave(USoundWave* InSoundWave)
 {
 	if (SoundWave != InSoundWave)
@@ -157,8 +156,10 @@ void USynthSamplePlayer::OnUnregister()
 	Super::OnUnregister();
 }
 
-void USynthSamplePlayer::TickComponent(float DeltaTime, enum ELevelTick TickType, FActorComponentTickFunction *ThisTickFunction)
+void USynthSamplePlayer::TickComponent(float DeltaTime, enum ELevelTick TickType, FActorComponentTickFunction* ThisTickFunction)
 {
+	Super::TickComponent(DeltaTime, TickType, ThisTickFunction);
+
 	SoundWaveLoader.Update();
 
 	OnSamplePlaybackProgress.Broadcast(GetCurrentPlaybackProgressTime(), GetCurrentPlaybackProgressPercent());
@@ -168,10 +169,11 @@ int32 USynthSamplePlayer::OnGenerateAudio(float* OutAudio, int32 NumSamples)
 {
 	if (SampleBuffer.GetData() && !SampleBufferReader.HasBuffer())
 	{
-		const int16* BufferData = SampleBuffer.GetData();
-		const int32 BufferNumSamples = SampleBuffer.GetNumSamples();
+		const int16* BufferData       = SampleBuffer.GetData();
+		const int32 BufferNumSamples  = SampleBuffer.GetNumSamples();
 		const int32 BufferNumChannels = SampleBuffer.GetNumChannels();
-		const int32 BufferSampleRate = SampleBuffer.GetSampleRate();
+		const int32 BufferSampleRate  = SampleBuffer.GetSampleRate();
+
 		SampleBufferReader.SetBuffer(BufferData, BufferNumSamples, BufferNumChannels, BufferSampleRate);
 		SampleDurationSec = BufferNumSamples / (BufferSampleRate * BufferNumChannels);
 	}
diff --git a/Engine/Source/Runtime/AudioMixer/Private/Components/SynthComponent.cpp b/Engine/Source/Runtime/AudioMixer/Private/Components/SynthComponent.cpp
--- a/Engine/Source/Runtime/AudioMixer/Private/Components/SynthComponent.cpp
+++ b/Engine/Source/Runtime/AudioMixer/Private/Components/SynthComponent.cpp
@@ -132,7 +132,7 @@ void USynthComponent::OnAudioComponentEnvelopeValue(const UAudioComponent* InAud
 
 void USynthComponent::Activate(bool bReset)
 {
-	if (bReset || ShouldActivate() == true)
+	if (bReset || ShouldActivate())
 	{
 		Start();
 		if (bIsActive)
@@ -172,7 +172,7 @@ void USynthComponent::Initialize(int32 SampleRateOverride)
 		}
 	}
 
-	// Only allow initialization if we've gota  proper sample rate
+	// Only allow initialization if we have a proper sample rate
 	if (SampleRate != INDEX_NONE)
 	{
 #if SYNTH_GENERATOR_TEST_TONE
@@ -181,7 +181,7 @@ void USynthComponent::Initialize(int32 SampleRateOverride)
 		TestSineRight.Init(SampleRate, 220.0f, 0.5f);
 #else	
 		// Initialize the synth component
-		this->Init(SampleRate);
+		Init(SampleRate);
 
 		if (NumChannels < 0 || NumChannels > 2)
 		{
@@ -191,7 +191,7 @@ void USynthComponent::Initialize(int32 SampleRateOverride)
 		NumChannels = FMath::Clamp(NumChannels, 1, 2);
 #endif
 
-		if (nullptr == Synth)
+		if (!Synth)
 		{
 			Synth = NewObject<USynthSound>(this, TEXT("Synth"));
 		}
@@ -237,9 +237,9 @@ void USynthComponent::CreateAudioComponent()
 #if WITH_EDITORONLY_DATA
 		AudioComponent->bVisualizeComponent = false;
 #endif
-		if (AudioComponent->GetAttachParent() == nullptr && !AudioComponent->IsAttachedTo(this))
+		if (!AudioComponent->GetAttachParent() && !AudioComponent->IsAttachedTo(this))
 		{
-			if (GetOwner() == nullptr || GetOwner()->GetWorld() == nullptr)
+			if (!GetOwner() || !GetOwner()->GetWorld())
 			{
 				AudioComponent->SetupAttachment(this);
 			}
@@ -258,7 +258,6 @@ void USynthComponent::CreateAudioComponent()
 	}
 }
 
-
 void USynthComponent::OnRegister()
 {
 	CreateAudioComponent();
@@ -348,12 +347,12 @@ void USynthComponent::PumpPendingMessages()
 		{
 			case ESynthEvent::Start:
 				bIsSynthPlaying = true;
-				this->OnStart();
+				OnStart();
 				break;
 
 			case ESynthEvent::Stop:
 				bIsSynthPlaying = false;
-				this->OnStop();
+				OnStop();
 				break;
 
 			default:
@@ -371,7 +370,7 @@ int32 USynthComponent::OnGeneratePCMAudio(float* GeneratedPCMData, int32 NumSamp
 	// Only call into the synth if we're actually playing, otherwise, we'll write out zero's
 	if (bIsSynthPlaying)
 	{
-		return this->OnGenerateAudio(GeneratedPCMData, NumSamples);
+		return OnGenerateAudio(GeneratedPCMData, NumSamples);
 	}
 	return NumSamples;
 }
@@ -390,7 +389,7 @@ void USynthComponent::Start()
 	// If there is no Synth USoundBase, we can't start. This can happen if start is called in a cook, a server, or 
 	// if the audio engine is set to "noaudio".
 	// TODO: investigate if this should be handled elsewhere before this point
-	if (Synth == nullptr)
+	if (!Synth)
 	{
 		return;
 	}
diff --git a/Engine/Source/Runtime/Engine/Classes/Sound/SampleBuffer.h b/Engine/Source/Runtime/Engine/Classes/Sound/SampleBuffer.h
--- a/Engine/Source/Runtime/Engine/Classes/Sound/SampleBuffer.h
+++ b/Engine/Source/Runtime/Engine/Classes/Sound/SampleBuffer.h
@@ -381,11 +381,25 @@ namespace Audio
 			// The lambda function to call when t he sound wave finishes loading
 			TFunction<void(const USoundWave* SoundWave, const Audio::FSampleBuffer& LoadedSampleBuffer)> OnLoaded;
 
-			// Whether the sound wave load/decode is in-flight
-			bool bIsLoading;
+			enum class LoadStatus : uint8
+			{
+				// No request to load has been issued (default)
+				None = 0,
+
+				// The sound wave load/decode is in-flight
+				Loading,
+
+				// The sound wave has already been loaded
+				Loaded,
+			};
 
-			// Whether or not the sound wave has already been loaded
-			bool bIsLoaded;
+			LoadStatus Status;
+
+			FLoadingSoundWaveInfo()
+				: SoundWave(nullptr)
+				, Status(LoadStatus::None)
+			{
+			}
 		};
 
 		// Reference to current loading sound wave
diff --git a/Engine/Source/Runtime/Engine/Private/SampleBuffer.cpp b/Engine/Source/Runtime/Engine/Private/SampleBuffer.cpp
--- a/Engine/Source/Runtime/Engine/Private/SampleBuffer.cpp
+++ b/Engine/Source/Runtime/Engine/Private/SampleBuffer.cpp
@@ -25,22 +25,25 @@ namespace Audio
 		}
 
 		FLoadingSoundWaveInfo LoadingSoundWaveInfo;
-		LoadingSoundWaveInfo.SoundWave = InSoundWave;
-		LoadingSoundWaveInfo.OnLoaded = MoveTemp(OnLoaded);
-		LoadingSoundWaveInfo.bIsLoading = true;
 
-		if (LoadingSoundWaveInfo.SoundWave->GetPrecacheState() != ESoundWavePrecacheState::Done)
+		const bool bRequestPrecache = InSoundWave->GetPrecacheState() != ESoundWavePrecacheState::Done ||
+			!InSoundWave->RawPCMData ||
+			InSoundWave->RawPCMDataSize == 0;
+		if (bRequestPrecache)
 		{
-			LoadingSoundWaveInfo.bIsLoaded = false;
+			LoadingSoundWaveInfo.Status = FLoadingSoundWaveInfo::LoadStatus::Loading;
 
 			// Kick off a decompression/precache of the sound wave
 			AudioDevice->Precache(InSoundWave, false, true, true);
 		}
 		else
 		{
-			LoadingSoundWaveInfo.bIsLoaded = true;
+			LoadingSoundWaveInfo.Status = FLoadingSoundWaveInfo::LoadStatus::Loaded;
 		}
 
+		LoadingSoundWaveInfo.SoundWave  = InSoundWave;
+		LoadingSoundWaveInfo.OnLoaded   = MoveTemp(OnLoaded);
+
 		LoadingSoundWaves.Add(LoadingSoundWaveInfo);
 	}
 
@@ -49,25 +52,20 @@ namespace Audio
 		for (int32 i = LoadingSoundWaves.Num() - 1; i >= 0; --i)
 		{
 			FLoadingSoundWaveInfo& LoadingSoundWaveInfo = LoadingSoundWaves[i];
-
-			if (LoadingSoundWaveInfo.bIsLoading || LoadingSoundWaveInfo.bIsLoaded)
+			if (USoundWave* SoundWave = LoadingSoundWaveInfo.SoundWave)
 			{
-				check(LoadingSoundWaveInfo.SoundWave);
-
-				if (LoadingSoundWaveInfo.bIsLoaded || LoadingSoundWaveInfo.SoundWave->GetPrecacheState() == ESoundWavePrecacheState::Done)
+				if (SoundWave->GetPrecacheState() == ESoundWavePrecacheState::Done)
 				{
-					LoadingSoundWaveInfo.bIsLoading = false;
-					LoadingSoundWaveInfo.bIsLoaded = true;
-
-					USoundWave* SoundWave = LoadingSoundWaveInfo.SoundWave;
+					LoadingSoundWaveInfo.Status = FLoadingSoundWaveInfo::LoadStatus::Loaded;
+				}
 
+				if (LoadingSoundWaveInfo.Status == FLoadingSoundWaveInfo::LoadStatus::Loaded)
+				{
 					const Audio::DefaultUSoundWaveSampleType* RawPCMData = reinterpret_cast<const Audio::DefaultUSoundWaveSampleType*>(SoundWave->RawPCMData);
 					const int32 NumSamples = SoundWave->RawPCMDataSize / sizeof(Audio::DefaultUSoundWaveSampleType);
 
 					TSampleBuffer<> SampleBuffer(RawPCMData, NumSamples, SoundWave->NumChannels, SoundWave->GetSampleRateForCurrentPlatform());
-
 					LoadingSoundWaveInfo.OnLoaded(SoundWave, SampleBuffer);
-
 					LoadingSoundWaves.RemoveAtSwap(i, 1, false);
 				}
 			}
diff --git a/Engine/Source/Runtime/Engine/Private/SoundNode.cpp b/Engine/Source/Runtime/Engine/Private/SoundNode.cpp
--- a/Engine/Source/Runtime/Engine/Private/SoundNode.cpp
+++ b/Engine/Source/Runtime/Engine/Private/SoundNode.cpp
@@ -179,7 +179,10 @@ bool USoundNode::HasDelayNode() const
 		if (ChildNode)
 		{
 			ChildNode->ConditionalPostLoad();
-			return ChildNode->HasDelayNode();
+			if (ChildNode->HasDelayNode())
+			{
+				return true;
+			}
 		}
 	}
 	return false;
@@ -192,7 +195,10 @@ bool USoundNode::HasConcatenatorNode() const
 		if (ChildNode)
 		{
 			ChildNode->ConditionalPostLoad();
-			return ChildNode->HasConcatenatorNode();
+			if (ChildNode->HasConcatenatorNode())
+			{
+				return true;
+			}
 		}
 	}
 	return false;
@@ -205,7 +211,10 @@ bool USoundNode::IsVirtualizeWhenSilent() const
 		if (ChildNode)
 		{
 			ChildNode->ConditionalPostLoad();
-			return ChildNode->IsVirtualizeWhenSilent();
+			if (ChildNode->IsVirtualizeWhenSilent())
+			{
+				return true;
+			}
 		}
 	}
 	return false;
diff --git a/Engine/Source/Runtime/Engine/Private/SoundNodeDistanceCrossFade.cpp b/Engine/Source/Runtime/Engine/Private/SoundNodeDistanceCrossFade.cpp
--- a/Engine/Source/Runtime/Engine/Private/SoundNodeDistanceCrossFade.cpp
+++ b/Engine/Source/Runtime/Engine/Private/SoundNodeDistanceCrossFade.cpp
@@ -76,7 +76,7 @@ void USoundNodeDistanceCrossFade::ParseNodes( FAudioDevice* AudioDevice, const U
 			}
 			else if( ( Distance >= FadeInDistanceMin ) && ( Distance <= FadeInDistanceMax ) )
 			{
-				VolumeToSet = (FadeInDistanceMax > 0.f ? CrossFadeInput[ ChildNodeIndex ].Volume * ( 0.0f + ( Distance - FadeInDistanceMin ) / ( FadeInDistanceMax - FadeInDistanceMin ) ) : 1.f);
+				VolumeToSet = (FadeInDistanceMax > 0.f ? CrossFadeInput[ ChildNodeIndex ].Volume * ( 0.0f + ( Distance - FadeInDistanceMin ) / ( FadeInDistanceMax - FadeInDistanceMin ) ) : CrossFadeInput[ ChildNodeIndex ].Volume);
 				//UE_LOG(LogAudio, Log,  TEXT("     FadeIn.  Distance: %f,  VolumeToSet: %f"), Distance, VolumeToSet );
 			}
 			// else if we are inside the FadeOut edge
