Copy from Dev-Audio:

CL 4884792: Remove invalid check.  PCMBuffer is appropriately null on first read.

diff --git a/Engine/Source/Runtime/AudioMixer/Private/DSP/SoundWaveDecoder.cpp b/Engine/Source/Runtime/AudioMixer/Private/DSP/SoundWaveDecoder.cpp
--- a/Engine/Source/Runtime/AudioMixer/Private/DSP/SoundWaveDecoder.cpp
+++ b/Engine/Source/Runtime/AudioMixer/Private/DSP/SoundWaveDecoder.cpp
@@ -146,7 +146,6 @@ namespace Audio
 
 			if (MixerSourceBuffer.GetNumBuffersQueued() > 0)
 			{
-				check(SourceInfo.CurrentPCMBuffer->AudioData.Num() > 0);
 				SourceInfo.CurrentPCMBuffer = MixerSourceBuffer.GetNextBuffer();
 				SourceInfo.CurrentAudioChunkNumFrames = SourceInfo.CurrentPCMBuffer->AudioData.Num() / SourceInfo.NumSourceChannels;
 
