Fix compile error in Debug build.

diff --git a/Engine/Source/Runtime/AudioMixer/Private/AudioMixerDevice.cpp b/Engine/Source/Runtime/AudioMixer/Private/AudioMixerDevice.cpp
--- a/Engine/Source/Runtime/AudioMixer/Private/AudioMixerDevice.cpp
+++ b/Engine/Source/Runtime/AudioMixer/Private/AudioMixerDevice.cpp
@@ -65,7 +65,7 @@ namespace Audio
 		}
 	}
 
-	void FMixerDevice::CheckAudioThread()
+	void FMixerDevice::CheckAudioThread() const
 	{
 #if AUDIO_MIXER_ENABLE_DEBUG_MODE
 		// "Audio Thread" is the game/audio thread ID used above audio rendering thread.
@@ -81,7 +81,7 @@ namespace Audio
 #endif
 	}
 
-	void FMixerDevice::CheckAudioRenderingThread()
+	void FMixerDevice::CheckAudioRenderingThread() const
 	{
 #if AUDIO_MIXER_ENABLE_DEBUG_MODE
 		if (AudioPlatformThreadId == INDEX_NONE)
@@ -93,7 +93,7 @@ namespace Audio
 #endif
 	}
 
-	bool FMixerDevice::IsAudioRenderingThread()
+	bool FMixerDevice::IsAudioRenderingThread() const
 	{
 		int32 CurrentThreadId = FPlatformTLS::GetCurrentThreadId();
 		return CurrentThreadId == AudioPlatformThreadId;
diff --git a/Engine/Source/Runtime/AudioMixer/Public/AudioMixerDevice.h b/Engine/Source/Runtime/AudioMixer/Public/AudioMixerDevice.h
--- a/Engine/Source/Runtime/AudioMixer/Public/AudioMixerDevice.h
+++ b/Engine/Source/Runtime/AudioMixer/Public/AudioMixerDevice.h
@@ -141,9 +141,9 @@ namespace Audio
 		FMixerSubmixWeakPtr GetSubmixInstance(USoundSubmix* SoundSubmix);
 
 		// Functions which check the thread it's called on and helps make sure functions are called from correct threads
-		void CheckAudioThread();
-		void CheckAudioRenderingThread();
-		bool IsAudioRenderingThread();
+		void CheckAudioThread() const;
+		void CheckAudioRenderingThread() const;
+		bool IsAudioRenderingThread() const;
 
 		// Public Functions
 		FMixerSourceVoice* GetMixerSourceVoice();
@@ -278,10 +278,10 @@ namespace Audio
 		FMixerSourceManager SourceManager;
 
 		/** ThreadId for the game thread (or if audio is running a seperate thread, that ID) */
-		int32 GameOrAudioThreadId;
+		mutable int32 GameOrAudioThreadId;
 
 		/** ThreadId for the low-level platform audio mixer. */
-		int32 AudioPlatformThreadId;
+		mutable int32 AudioPlatformThreadId;
 
 		/** Command queue to send commands to audio render thread from game thread or audio thread. */
 		TQueue<TFunction<void()>> CommandQueue;
