Fix for pausing sounds w/ audio mixer.

PR #4994: Fix: Stopping FMixerSource can be paused (Contributed by [REDACTED_PERSON])

diff --git a/Engine/Source/Runtime/AudioMixer/Private/AudioMixerSource.cpp b/Engine/Source/Runtime/AudioMixer/Private/AudioMixerSource.cpp
--- a/Engine/Source/Runtime/AudioMixer/Private/AudioMixerSource.cpp
+++ b/Engine/Source/Runtime/AudioMixer/Private/AudioMixerSource.cpp
@@ -558,6 +558,11 @@ namespace Audio
 			return;
 		}
 
+		if (bIsStopping)
+		{
+			return;
+		}
+
 		if (MixerSourceVoice)
 		{
 			MixerSourceVoice->Pause();
