[NO NON-PII COMMIT MESSAGE]

diff --git a/Engine/Source/Runtime/AudioMixer/Private/DSP/BufferVectorOperations.cpp b/Engine/Source/Runtime/AudioMixer/Private/DSP/BufferVectorOperations.cpp
--- a/Engine/Source/Runtime/AudioMixer/Private/DSP/BufferVectorOperations.cpp
+++ b/Engine/Source/Runtime/AudioMixer/Private/DSP/BufferVectorOperations.cpp
@@ -1295,11 +1295,11 @@ namespace Audio
 			const VectorRegister Input2 = VectorLoadFloat1(&SourceBuffer[i * 2 + 1]);
 
 			VectorRegister Result = VectorMultiply(Input1, GainVector11);
-			Result = VectorMultiplyAdd(Input2, GainVector21, Result);
+			Result = VectorMultiplyAdd(Input2, GainVector12, Result);
 
 			VectorStoreAligned(Result, &DestinationBuffer[i * 8]);
 
-			Result = VectorMultiply(Input1, GainVector12);
+			Result = VectorMultiply(Input1, GainVector21);
 			Result = VectorMultiplyAdd(Input2, GainVector22, Result);
 			VectorStoreAligned(Result, &DestinationBuffer[i * 8 + 4]);
 		}
@@ -1321,7 +1321,7 @@ namespace Audio
 	* |        | *        | *       | *       | *       | *        | *       | *       | *       |
 	* | Input1 | i0       | i0      | i0      | i0      | i0       | i0      | i0      | i0      |
 	* |        | +        | +       | +       | +       | +        | +       | +       | +       |
-	* | Gain2  | g8       | g9      | g10     | g11     | g12      | g13     | g14     | g5      |
+	* | Gain2  | g8       | g9      | g10     | g11     | g12      | g13     | g14     | g15     |
 	* |        | *        | *       | *       | *       | *        | *       | *       | *       |
 	* | Input2 | i1       | i1      | i1      | i1      | i1       | i1      | i1      | i1      |
 	* |        | =        | =       | =       | =       | =        | =       | =       | =       |
@@ -1355,17 +1355,17 @@ namespace Audio
 			const VectorRegister Input2 = VectorLoadFloat1(&SourceBuffer[i * 2 + 1]);
 
 			VectorRegister Result = VectorMultiply(Input1, GainVector11);
-			Result = VectorMultiplyAdd(Input2, GainVector21, Result);
+			Result = VectorMultiplyAdd(Input2, GainVector12, Result);
 			VectorStoreAligned(Result, &DestinationBuffer[i * 8]);
 
 			GainVector11 = VectorAdd(GainVector11, GainDeltasVector11);
-			GainVector21 = VectorAdd(GainVector21, GainDeltasVector21);
+			GainVector12 = VectorAdd(GainVector12, GainDeltasVector12);
 
-			Result = VectorMultiply(Input1, GainVector12);
+			Result = VectorMultiply(Input1, GainVector21);
 			Result = VectorMultiplyAdd(Input2, GainVector22, Result);
 			VectorStoreAligned(Result, &DestinationBuffer[i * 8 + 4]);
 
-			GainVector12 = VectorAdd(GainVector12, GainDeltasVector12);
+			GainVector21 = VectorAdd(GainVector21, GainDeltasVector21);
 			GainVector22 = VectorAdd(GainVector22, GainDeltasVector22);
 		}
 	}
