Fixing crash with listener transform update

diff --git a/Engine/Source/Runtime/AudioMixer/Private/AudioMixerDevice.cpp b/Engine/Source/Runtime/AudioMixer/Private/AudioMixerDevice.cpp
--- a/Engine/Source/Runtime/AudioMixer/Private/AudioMixerDevice.cpp
+++ b/Engine/Source/Runtime/AudioMixer/Private/AudioMixerDevice.cpp
@@ -347,6 +347,7 @@ namespace Audio
 			AudioMixerPlatform->ResumePlaybackOnNewDevice();
 		}
 
+#if 0 // Disable touching the listener transforms
 		ListenerTransforms.Reset();
 		for (FListener& Listener : Listeners)
 		{
@@ -355,6 +356,7 @@ namespace Audio
 
 		// Update listener transforms, some effects use the listener transform data
 		SourceManager.SetListenerTransforms(ListenerTransforms);
+#endif
 
 		// Loop through any envelope-following submixes and perform any broadcasting of envelope data if needed
 		TArray<float> SubmixEnvelopeData;
