Fix for distortion and zippering during panning.

- Fix is to make sure the starting speaker map is set to the destination after doing the interpolation.

diff --git a/Engine/Source/Runtime/AudioMixer/Private/AudioMixerSourceDecode.cpp b/Engine/Source/Runtime/AudioMixer/Private/AudioMixerSourceDecode.cpp
--- a/Engine/Source/Runtime/AudioMixer/Private/AudioMixerSourceDecode.cpp
+++ b/Engine/Source/Runtime/AudioMixer/Private/AudioMixerSourceDecode.cpp
@@ -50,6 +50,8 @@ public:
 				// Make sure we've been flagged as active
 				if (!ProceduralTaskData.ProceduralSoundWave->IsGenerating())
 				{
+					// Act as if we generated audio.
+					ProceduralResult.NumSamplesWritten = ProceduralTaskData.NumSamples;
 					return;
 				}
 
diff --git a/Engine/Source/Runtime/AudioMixer/Private/AudioMixerSourceManager.cpp b/Engine/Source/Runtime/AudioMixer/Private/AudioMixerSourceManager.cpp
--- a/Engine/Source/Runtime/AudioMixer/Private/AudioMixerSourceManager.cpp
+++ b/Engine/Source/Runtime/AudioMixer/Private/AudioMixerSourceManager.cpp
@@ -1668,6 +1668,8 @@ namespace Audio
 				Audio::DownmixBuffer(DownmixData.NumInputChannels, 8, *DownmixData.PostEffectBuffers, DownmixData.SevenOneSubmixInfo.OutputBuffer, DownmixData.SevenOneSubmixInfo.ChannelMap.ChannelStartGains, DownmixData.SevenOneSubmixInfo.ChannelMap.ChannelDestinationGains);
 			}
 		}
+
+		DownmixData.DeviceSubmixInfo.ChannelMap.CopyDestinationToStart();
 	}
 
 	void FMixerSourceManager::ComputeDownmix2D(FSourceDownmixData& DownmixData)
diff --git a/Engine/Source/Runtime/AudioMixer/Private/AudioMixerSourceManager.h b/Engine/Source/Runtime/AudioMixer/Private/AudioMixerSourceManager.h
--- a/Engine/Source/Runtime/AudioMixer/Private/AudioMixerSourceManager.h
+++ b/Engine/Source/Runtime/AudioMixer/Private/AudioMixerSourceManager.h
@@ -163,6 +163,11 @@ namespace Audio
 			FMemory::Memzero(ChannelDestinationGains, CopySize);
 		}
 
+		FORCEINLINE void CopyDestinationToStart()
+		{
+			FMemory::Memcpy(ChannelStartGains, ChannelDestinationGains, CopySize);
+		}
+
 		FORCEINLINE void SetChannelMap(const float* RESTRICT InChannelGains)
 		{
 			// TODO: See if you can find a way to assign this without going back to memory.
