Fix for procedural sources getting precached & static analysis warnings

diff --git a/Engine/Source/Runtime/AudioMixer/Private/AudioMixerSource.cpp b/Engine/Source/Runtime/AudioMixer/Private/AudioMixerSource.cpp
--- a/Engine/Source/Runtime/AudioMixer/Private/AudioMixerSource.cpp
+++ b/Engine/Source/Runtime/AudioMixer/Private/AudioMixerSource.cpp
@@ -298,7 +298,7 @@ namespace Audio
 				MixerSourceBuffer->SetPCMData(RawPCMDataBuffer);
 			}
 #if PLATFORM_NUM_AUDIODECOMPRESSION_PRECACHE_BUFFERS > 0
-			else if (BufferType == EBufferType::PCMRealTime || EBufferType::Streaming)
+			else if (BufferType == EBufferType::PCMRealTime || BufferType == EBufferType::Streaming)
 			{
 				USoundWave* WaveData = WaveInstance->WaveData;
 				if (WaveData->CachedRealtimeFirstBuffer)
diff --git a/Engine/Source/Runtime/AudioMixer/Private/AudioMixerSourceBuffer.cpp b/Engine/Source/Runtime/AudioMixer/Private/AudioMixerSourceBuffer.cpp
--- a/Engine/Source/Runtime/AudioMixer/Private/AudioMixerSourceBuffer.cpp
+++ b/Engine/Source/Runtime/AudioMixer/Private/AudioMixerSourceBuffer.cpp
@@ -153,7 +153,7 @@ namespace Audio
 
 	void FMixerSourceBuffer::SetPCMData(const FRawPCMDataBuffer& InPCMDataBuffer)
 	{
-		check(BufferType == EBufferType::PCM || EBufferType::PCMPreview);
+		check(BufferType == EBufferType::PCM || BufferType == EBufferType::PCMPreview);
 		RawPCMDataBuffer = InPCMDataBuffer;
 	}
 
