[NO NON-PII COMMIT MESSAGE]

diff --git a/Engine/Source/Runtime/AudioMixer/Private/AudioMixerSource.cpp b/Engine/Source/Runtime/AudioMixer/Private/AudioMixerSource.cpp
--- a/Engine/Source/Runtime/AudioMixer/Private/AudioMixerSource.cpp
+++ b/Engine/Source/Runtime/AudioMixer/Private/AudioMixerSource.cpp
@@ -61,7 +61,7 @@ namespace Audio
 		check(WaveInstance->WaveData);
 
 		// Prevent double-triggering procedural soundwaves
-		if (WaveInstance->WaveData->bProcedural && WaveInstance->WaveData->IsGenerating())
+		if (WaveInstance->WaveData->bProcedural && WaveInstance->WaveData->IsGeneratingAudio())
 		{
 			UE_LOG(LogAudioMixer, Warning, TEXT("Procedural sound wave is reinitializing even though it is currently actively generating audio. Please stop sound before trying to play it again."));
 			return false;
diff --git a/Engine/Source/Runtime/AudioMixer/Private/AudioMixerSourceBuffer.cpp b/Engine/Source/Runtime/AudioMixer/Private/AudioMixerSourceBuffer.cpp
--- a/Engine/Source/Runtime/AudioMixer/Private/AudioMixerSourceBuffer.cpp
+++ b/Engine/Source/Runtime/AudioMixer/Private/AudioMixerSourceBuffer.cpp
@@ -97,6 +97,9 @@ namespace Audio
 			delete DecompressionState;
 			DecompressionState = nullptr;
 		}
+
+		check(SoundWave);
+		SoundWave->RemovePlayingSource();
 	}
 
 	bool FMixerSourceBuffer::PreInit(FMixerBuffer* InBuffer, USoundWave* InWave, ELoopingMode InLoopingMode, bool bInIsSeeking)
@@ -115,10 +118,9 @@ namespace Audio
 		
 		check(SoundWave == nullptr);
 		// Only need to have a handle to a USoundWave for procedural sound waves
-		if (bProcedural)
-		{
-			SoundWave = InWave;
-		}
+		SoundWave = InWave;
+		check(SoundWave);
+		SoundWave->AddPlayingSource();
 
 		LoopingMode = InLoopingMode;
 		bIsSeeking = bInIsSeeking;
@@ -482,12 +484,8 @@ namespace Audio
 
 	bool FMixerSourceBuffer::IsBeginDestroy()
 	{
-		if (bProcedural)
-		{
-			check(SoundWave && SoundWave->bProcedural);
-			return SoundWave->bIsBeginDestroy;
-		}
-		return false;
+		check(SoundWave);
+		return SoundWave->bIsBeginDestroy;
 	}
 
 	void FMixerSourceBuffer::ClearSoundWave()
@@ -501,7 +499,6 @@ namespace Audio
 		if (bProcedural)
 		{
 			check(SoundWave && SoundWave->bProcedural);
-			SoundWave->SetGenerating(true);
 			SoundWave->OnBeginGenerate();
 		}
 	}
@@ -516,7 +513,6 @@ namespace Audio
 		{
 			check(SoundWave && SoundWave->bProcedural);
 			SoundWave->OnEndGenerate();
-			SoundWave->SetGenerating(false);
 			SoundWave = nullptr;
 		}
 	}
diff --git a/Engine/Source/Runtime/AudioMixer/Private/AudioMixerSourceDecode.cpp b/Engine/Source/Runtime/AudioMixer/Private/AudioMixerSourceDecode.cpp
--- a/Engine/Source/Runtime/AudioMixer/Private/AudioMixerSourceDecode.cpp
+++ b/Engine/Source/Runtime/AudioMixer/Private/AudioMixerSourceDecode.cpp
@@ -48,9 +48,10 @@ public:
 			case EAudioTaskType::Procedural:
 			{
 				// Make sure we've been flagged as active
-				if (!ProceduralTaskData.ProceduralSoundWave->IsGenerating())
+				if (!ProceduralTaskData.ProceduralSoundWave->IsGeneratingAudio())
 				{
-					// Act as if we generated audio.
+					// Act as if we generated audio, but return silence.
+					FMemory::Memzero(ProceduralTaskData.AudioData, ProceduralTaskData.NumSamples * sizeof(float));
 					ProceduralResult.NumSamplesWritten = ProceduralTaskData.NumSamples;
 					return;
 				}
diff --git a/Engine/Source/Runtime/AudioMixer/Private/Components/SynthComponent.cpp b/Engine/Source/Runtime/AudioMixer/Private/Components/SynthComponent.cpp
--- a/Engine/Source/Runtime/AudioMixer/Private/Components/SynthComponent.cpp
+++ b/Engine/Source/Runtime/AudioMixer/Private/Components/SynthComponent.cpp
@@ -295,7 +295,7 @@ void USynthComponent::OnUnregister()
 bool USynthComponent::IsReadyForOwnerToAutoDestroy() const
 {
 	const bool bIsAudioComponentReadyForDestroy = !AudioComponent || (AudioComponent && !AudioComponent->IsPlaying());
-	const bool bIsSynthSoundReadyForDestroy = !Synth || !Synth->IsGenerating();
+	const bool bIsSynthSoundReadyForDestroy = !Synth || !Synth->IsGeneratingAudio();
 	return bIsAudioComponentReadyForDestroy && bIsSynthSoundReadyForDestroy;
 }
 
diff --git a/Engine/Source/Runtime/Engine/Classes/Sound/SoundWave.h b/Engine/Source/Runtime/Engine/Classes/Sound/SoundWave.h
--- a/Engine/Source/Runtime/Engine/Classes/Sound/SoundWave.h
+++ b/Engine/Source/Runtime/Engine/Classes/Sound/SoundWave.h
@@ -426,7 +426,9 @@ private:
 	/** What state the precache decompressor is in. */
 	FThreadSafeCounter PrecacheState;
 
-	FThreadSafeBool bGenerating;
+	/** the number of sounds currently playing this sound wave. */
+	FThreadSafeCounter NumSourcesPlaying;
+
 #if !WITH_EDITOR
 	// This is the sample rate gotten from platform settings.
 	float CachedSampleRateOverride;
@@ -599,8 +601,22 @@ public:
 	// Called when the procedural sound wave is done generating on the render thread. Only used in the audio mixer and when bProcedural is true..
 	virtual void OnEndGenerate() {};
 
-	bool IsGenerating() const { return bGenerating; }
-	void SetGenerating(bool bInGenerating) { bGenerating = bInGenerating; }
+	void AddPlayingSource()
+	{
+		NumSourcesPlaying.Increment();
+	}
+
+	void RemovePlayingSource()
+	{
+		check(NumSourcesPlaying.GetValue() > 0);
+		NumSourcesPlaying.Decrement();
+	}
+
+	bool IsGeneratingAudio() const
+	{
+		return NumSourcesPlaying.GetValue() > 0;
+	}
+
 	/**
 	* Overwrite sample rate. Used for procedural soundwaves, as well as sound waves that are resampled on compress/decompress.
 	*/
diff --git a/Engine/Source/Runtime/Engine/Private/AudioDevice.cpp b/Engine/Source/Runtime/Engine/Private/AudioDevice.cpp
--- a/Engine/Source/Runtime/Engine/Private/AudioDevice.cpp
+++ b/Engine/Source/Runtime/Engine/Private/AudioDevice.cpp
@@ -4141,7 +4141,7 @@ void FAudioDevice::AddNewActiveSound(const FActiveSound& NewActiveSound)
 	}
 
 	USoundWave* SoundWave = Cast<USoundWave>(Sound);
-	if (SoundWave && SoundWave->bProcedural && SoundWave->IsGenerating())
+	if (SoundWave && SoundWave->bProcedural && SoundWave->IsGeneratingAudio())
 	{
 		FString SoundWaveName;
 		SoundWave->GetName(SoundWaveName);
diff --git a/Engine/Source/Runtime/Engine/Private/SoundWave.cpp b/Engine/Source/Runtime/Engine/Private/SoundWave.cpp
--- a/Engine/Source/Runtime/Engine/Private/SoundWave.cpp
+++ b/Engine/Source/Runtime/Engine/Private/SoundWave.cpp
@@ -770,8 +770,7 @@ void USoundWave::BeginDestroy()
 {
 	Super::BeginDestroy();
 
-	// Flag that this sound wave is beginning destroying. For procedural sound waves, this will ensure
-	// the audio render thread stops the sound before GC hits.
+	// Flag that this sound wave is beginning destroying. This will ensure that all sounds using this in the audio renderer are stopped before GC finishes.
 	bIsBeginDestroy = true;
 	
 #if WITH_EDITOR
@@ -1427,7 +1426,7 @@ bool USoundWave::IsReadyForFinishDestroy()
 			}, GET_STATID(STAT_AudioFreeResources));
 		}
 	
-	return !bGenerating && ResourceState == ESoundWaveResourceState::Freed;
+	return NumSourcesPlaying.GetValue() == 0 && ResourceState == ESoundWaveResourceState::Freed;
 }
 
 
