Fix for handling failed initialization.

- Clean up the resources and handle edge cases.
- Tested by forcing various failure points.

diff --git a/Engine/Source/Runtime/AudioMixer/Private/AudioMixerSource.cpp b/Engine/Source/Runtime/AudioMixer/Private/AudioMixerSource.cpp
--- a/Engine/Source/Runtime/AudioMixer/Private/AudioMixerSource.cpp
+++ b/Engine/Source/Runtime/AudioMixer/Private/AudioMixerSource.cpp
@@ -64,6 +64,7 @@ namespace Audio
 		if (WaveInstance->WaveData->bProcedural && WaveInstance->WaveData->IsGeneratingAudio())
 		{
 			UE_LOG(LogAudioMixer, Warning, TEXT("Procedural sound wave is reinitializing even though it is currently actively generating audio. Please stop sound before trying to play it again."));
+			FreeResources();
 			return false;
 		}
 
@@ -95,6 +96,8 @@ namespace Audio
 			MixerSourceVoice = MixerDevice->GetMixerSourceVoice();
 			if (!MixerSourceVoice)
 			{
+				FreeResources();
+				UE_LOG(LogAudioMixer, Warning, TEXT("Failed to get a mixer source voice for sound %s."), *InWaveInstance->GetName());
 				return false;
 			}
 
@@ -335,8 +338,15 @@ namespace Audio
 			else
 			{
 				InitializationState = EMixerSourceInitializationState::NotInitialized;
+				UE_LOG(LogAudioMixer, Warning, TEXT("Failed to initialize mixer source voice '%s'."), *InWaveInstance->GetName());
 			}
 		}
+		else
+		{
+			UE_LOG(LogAudioMixer, Warning, TEXT("Num channels was 0 for sound buffer '%s'."), *InWaveInstance->GetName());
+		}
+
+		FreeResources();
 		return false;
 	}
 
@@ -758,6 +768,8 @@ namespace Audio
 			ChannelMaps[i].bUsed = false;
 			ChannelMaps[i].ChannelMap.Reset();
 		}
+
+		InitializationState = EMixerSourceInitializationState::NotInitialized;
 	}
 
 	void FMixerSource::UpdatePitch()
diff --git a/Engine/Source/Runtime/Engine/Private/AudioDevice.cpp b/Engine/Source/Runtime/Engine/Private/AudioDevice.cpp
--- a/Engine/Source/Runtime/Engine/Private/AudioDevice.cpp
+++ b/Engine/Source/Runtime/Engine/Private/AudioDevice.cpp
@@ -3634,7 +3634,9 @@ void FAudioDevice::StartSources(TArray<FWaveInstance*>& WaveInstances, int32 Fir
 					// This can happen if e.g. the USoundWave pointed to by the WaveInstance is not a valid sound file.
 					// If we don't stop the wave file, it will continue to try initializing the file every frame, which is a perf hit
 					UE_LOG(LogAudio, Log, TEXT("Failed to start sound source for %s"), (WaveInstance->ActiveSound && WaveInstance->ActiveSound->Sound) ? *WaveInstance->ActiveSound->Sound->GetName() : TEXT("UNKNOWN") );
-					Source->Stop();
+					WaveInstance->StopWithoutNotification();
+					Source->WaveInstance = nullptr;
+					FreeSources.Add(Source);
 				}
 			}
 			else if (Source)
