[NO NON-PII COMMIT MESSAGE]

diff --git a/Engine/Source/Runtime/AudioMixer/Private/AudioMixer.cpp b/Engine/Source/Runtime/AudioMixer/Private/AudioMixer.cpp
--- a/Engine/Source/Runtime/AudioMixer/Private/AudioMixer.cpp
+++ b/Engine/Source/Runtime/AudioMixer/Private/AudioMixer.cpp
@@ -569,7 +569,7 @@ namespace Audio
 			OverrunTimeoutCVar = FMath::Clamp(OverrunTimeoutCVar, 500, 5000);
 
 			// Now wait for a buffer to be consumed, which will bump up the read index.
-			if (!AudioRenderEvent->Wait(static_cast<uint32>(OverrunTimeoutCVar)))
+			if (AudioRenderEvent && !AudioRenderEvent->Wait(static_cast<uint32>(OverrunTimeoutCVar)))
 			{
 				// if we reached this block, we timed out, and should attempt to
 				// bail on our current device.
