Fix for media audio and audio mixer

diff --git a/Engine/Source/Runtime/AudioMixer/Private/AudioMixerSource.cpp b/Engine/Source/Runtime/AudioMixer/Private/AudioMixerSource.cpp
--- a/Engine/Source/Runtime/AudioMixer/Private/AudioMixerSource.cpp
+++ b/Engine/Source/Runtime/AudioMixer/Private/AudioMixerSource.cpp
@@ -60,19 +60,10 @@ namespace Audio
 
 		check(WaveInstance->WaveData);
 
-		// Prevent double-triggering procedural soundwaves
-		if (WaveInstance->WaveData->bProcedural && WaveInstance->WaveData->IsGeneratingAudio())
-		{
-			UE_LOG(LogAudioMixer, Warning, TEXT("Procedural sound wave is reinitializing even though it is currently actively generating audio. Please stop sound before trying to play it again."));
-			FreeResources();
-			return false;
-		}
-
 		// Get the number of frames before creating the buffer
 		int32 NumFrames = INDEX_NONE;
 		if (WaveInstance->WaveData->DecompressionType != DTYPE_Procedural)
 		{
-			//NumFrames = MixerBuffer->GetNumFrames();
 			check(!InWaveInstance->WaveData->RawPCMData || InWaveInstance->WaveData->RawPCMDataSize);
 			const int32 NumBytes = WaveInstance->WaveData->RawPCMDataSize;
 			NumFrames = NumBytes / (WaveInstance->WaveData->NumChannels * sizeof(int16));
diff --git a/Engine/Source/Runtime/AudioMixer/Private/AudioMixerSourceBuffer.cpp b/Engine/Source/Runtime/AudioMixer/Private/AudioMixerSourceBuffer.cpp
--- a/Engine/Source/Runtime/AudioMixer/Private/AudioMixerSourceBuffer.cpp
+++ b/Engine/Source/Runtime/AudioMixer/Private/AudioMixerSourceBuffer.cpp
@@ -98,8 +98,10 @@ namespace Audio
 			DecompressionState = nullptr;
 		}
 
-		check(SoundWave);
-		SoundWave->RemovePlayingSource();
+		if (SoundWave)
+		{
+			SoundWave->RemovePlayingSource();
+		}
 	}
 
 	bool FMixerSourceBuffer::PreInit(FMixerBuffer* InBuffer, USoundWave* InWave, ELoopingMode InLoopingMode, bool bInIsSeeking)
@@ -116,6 +118,13 @@ namespace Audio
 		bProcedural = InWave->bProcedural;
 		NumPrecacheFrames = InWave->NumPrecacheFrames;
 		
+		// Prevent double-triggering procedural soundwaves
+		if (bProcedural && InWave->IsGeneratingAudio())
+		{
+			UE_LOG(LogAudioMixer, Warning, TEXT("Procedural sound wave is reinitializing even though it is currently actively generating audio. Please stop sound before trying to play it again."));
+			return false;
+		}
+
 		check(SoundWave == nullptr);
 		// Only need to have a handle to a USoundWave for procedural sound waves
 		SoundWave = InWave;
@@ -513,7 +522,6 @@ namespace Audio
 		{
 			check(SoundWave && SoundWave->bProcedural);
 			SoundWave->OnEndGenerate();
-			SoundWave = nullptr;
 		}
 	}
 
