Unshelved from pending changelist '4759891':

PR #5474: Resetting destination values when removing patches from Modulation Matrix (Contributed by [REDACTED_PERSON])

diff --git a/Engine/Source/Runtime/AudioMixer/Private/DSP/ModulationMatrix.cpp b/Engine/Source/Runtime/AudioMixer/Private/DSP/ModulationMatrix.cpp
--- a/Engine/Source/Runtime/AudioMixer/Private/DSP/ModulationMatrix.cpp
+++ b/Engine/Source/Runtime/AudioMixer/Private/DSP/ModulationMatrix.cpp
@@ -89,6 +89,14 @@ namespace Audio
 			return false;
 		}
 
+		//reset destination values
+		for (int32 PatchDest = 0; PatchDest < InPatch->Destinations.Num(); ++PatchDest)
+		{
+			const FPatchDestination& Destination = InPatch->Destinations[PatchDest];
+			Destinations[VoiceId][Destination.Id].bDirty = false;
+			Destinations[VoiceId][Destination.Id].Value = 0.0f;
+		}
+
 		Patches[VoiceId].Remove(InPatch);
 		return true;
 	}
@@ -107,6 +115,14 @@ namespace Audio
 	void FModulationMatrix::ClearPatches(const int32 VoiceId)
 	{
 		check(VoiceId < NumVoices);
+
+		for (FDestData& Destination : Destinations[VoiceId]) 
+		{
+
+			Destination.bDirty = false;
+			Destination.Value  = 0.0f;
+		}
+
 		Patches[VoiceId].Reset();
 	}
 
