Fixes for hang for audio mixer

diff --git a/Engine/Source/Runtime/AudioMixer/Private/AudioMixerSourceManager.cpp b/Engine/Source/Runtime/AudioMixer/Private/AudioMixerSourceManager.cpp
--- a/Engine/Source/Runtime/AudioMixer/Private/AudioMixerSourceManager.cpp
+++ b/Engine/Source/Runtime/AudioMixer/Private/AudioMixerSourceManager.cpp
@@ -1,4 +1,4 @@
-// Copyright 1998-2019 Epic Games, Inc. All Rights Reserved.
+// Copyright 1998-2018 Epic Games, Inc. All Rights Reserved.
 
 #include "AudioMixerSourceManager.h"
 #include "AudioMixerSource.h"
@@ -106,7 +106,6 @@ namespace Audio
 		, NumSourceWorkers(4)
 		, bInitialized(false)
 		, bUsingSpatializationPlugin(false)
-		, MaxChannelsSupportedBySpatializationPlugin(1)
 	{
 		// Get a manual resetable event
 		const bool bIsManualReset = true;
@@ -262,7 +261,6 @@ namespace Audio
 		if (SpatializationPlugin.IsValid())
 		{
 			bUsingSpatializationPlugin = true;
-			MaxChannelsSupportedBySpatializationPlugin = MixerDevice->MaxChannelsSupportedBySpatializationPlugin;
 		}
 
 		bInitialized = true;
@@ -1517,7 +1515,7 @@ namespace Audio
 			SCOPE_CYCLE_COUNTER(STAT_AudioMixerHRTF);
 
 			AUDIO_MIXER_CHECK(SpatializationPlugin.IsValid());
-			AUDIO_MIXER_CHECK(SourceInfo.NumInputChannels <= MaxChannelsSupportedBySpatializationPlugin);
+			AUDIO_MIXER_CHECK(SourceInfo.NumInputChannels == 1);
 
 			FAudioPluginSourceInputData AudioPluginInputData;
 			AudioPluginInputData.AudioBuffer = &SourceInfo.SourceBuffer;
diff --git a/Engine/Source/Runtime/AudioMixer/Private/AudioMixerSourceManager.h b/Engine/Source/Runtime/AudioMixer/Private/AudioMixerSourceManager.h
--- a/Engine/Source/Runtime/AudioMixer/Private/AudioMixerSourceManager.h
+++ b/Engine/Source/Runtime/AudioMixer/Private/AudioMixerSourceManager.h
@@ -1,4 +1,4 @@
-// Copyright 1998-2019 Epic Games, Inc. All Rights Reserved.
+// Copyright 1998-2018 Epic Games, Inc. All Rights Reserved.
 
 #pragma once
 
@@ -170,8 +170,6 @@ namespace Audio
 
 		FORCEINLINE void SetChannelMap(const float* RESTRICT InChannelGains)
 		{
-			// TODO: See if you can find a way to assign this without going back to memory.
-			FMemory::Memcpy(ChannelStartGains, ChannelDestinationGains, CopySize);
 			FMemory::Memcpy(ChannelDestinationGains, InChannelGains, CopySize);
 		}
 
@@ -550,7 +548,6 @@ namespace Audio
 
 		uint8 bInitialized : 1;
 		uint8 bUsingSpatializationPlugin : 1;
-		int32 MaxChannelsSupportedBySpatializationPlugin;
 
 		friend class FMixerSourceVoice;
 	};
diff --git a/Engine/Source/Runtime/AudioMixer/Private/DSP/BufferVectorOperations.cpp b/Engine/Source/Runtime/AudioMixer/Private/DSP/BufferVectorOperations.cpp
--- a/Engine/Source/Runtime/AudioMixer/Private/DSP/BufferVectorOperations.cpp
+++ b/Engine/Source/Runtime/AudioMixer/Private/DSP/BufferVectorOperations.cpp
@@ -462,7 +462,7 @@ namespace Audio
 	* +------------+---------+---------+---------+---------+
 	* | Gain       | g0      | g1      | g0      | g1      |
 	* |            | *       | *       | *       | *       |
-	* | Input      | i0      | i0      | i0      | i0      |
+	* | Input      | i0      | i0      | i1      | i1      |
 	* |            | =       | =       | =       | =       |
 	* | Output     | o0      | o1      | o2      | o3      |
 	* +------------+---------+---------+---------+---------+
@@ -476,7 +476,7 @@ namespace Audio
 		const VectorRegister GainDeltasVector = VectorDivide(VectorSubtract(DestinationVector, GainVector), NumFramesVector);
 
 		// To help with stair stepping, we initialize the second frame in GainVector to be half a GainDeltas vector higher than the first frame.
-		const VectorRegister VectorOfHalf = VectorSet(0.0f, 0.0f, 0.5f, 0.5f);
+		const VectorRegister VectorOfHalf = VectorSet(0.5f, 0.5f, 1.0f, 1.0f);
 		const VectorRegister HalfOfDeltaVector = VectorMultiply(GainDeltasVector, VectorOfHalf);
 		GainVector = VectorAdd(GainVector, HalfOfDeltaVector);
 
@@ -562,7 +562,7 @@ namespace Audio
 		const VectorRegister GainDeltasVector1 = VectorDivide(VectorSubtract(DestinationVector1, GainVector1), NumFramesVector);
 
 		// To help with stair stepping, we initialize the second frame in GainVector to be half a GainDeltas vector higher than the first frame.
-		const VectorRegister VectorOfHalf = VectorSet(0.0f, 0.0f, 0.5f, 0.5f);
+		const VectorRegister VectorOfHalf = VectorSet(0.5f, 0.5f, 1.0f, 1.0f);
 
 		const VectorRegister HalfOfDeltaVector1 = VectorMultiply(GainDeltasVector1, VectorOfHalf);
 		GainVector1 = VectorAdd(GainVector1, HalfOfDeltaVector1);
@@ -1410,7 +1410,7 @@ namespace Audio
 		{
 			for (int32 InputChannelIndex = 0; InputChannelIndex < NumSourceChannels; InputChannelIndex++)
 			{
-				const int32 GainMatrixIndex = OutputChannelIndex * NumDestinationChannels + InputChannelIndex;
+				const int32 GainMatrixIndex = InputChannelIndex * NumDestinationChannels + OutputChannelIndex;
 				GainDeltas[GainMatrixIndex] = (EndGains[GainMatrixIndex] - StartGains[GainMatrixIndex]) / NumFrames;
 			}
 		}
