Fix for ambisonics decoding

diff --git a/Engine/Plugins/Runtime/ResonanceAudio/Source/ResonanceAudio/Private/ResonanceAudioReverb.cpp b/Engine/Plugins/Runtime/ResonanceAudio/Source/ResonanceAudio/Private/ResonanceAudioReverb.cpp
--- a/Engine/Plugins/Runtime/ResonanceAudio/Source/ResonanceAudio/Private/ResonanceAudioReverb.cpp
+++ b/Engine/Plugins/Runtime/ResonanceAudio/Source/ResonanceAudio/Private/ResonanceAudioReverb.cpp
@@ -26,8 +26,12 @@ namespace ResonanceAudio
 	{
 		if (GlobalReverbPluginPreset)
 		{
+#if !WITH_EDITOR
 			GlobalReverbPluginPreset->RemoveFromRoot();
+#endif
+			GlobalReverbPluginPreset = nullptr;
 		}
+		ReverbPluginPreset = nullptr;
 	}
 
 	void FResonanceAudioReverb::Initialize(const FAudioPluginInitializationParams InitializationParams)
@@ -130,6 +134,7 @@ namespace ResonanceAudio
 		}
 		else
 		{
+			ReverbPluginPreset = nullptr;
 			return nullptr;
 		}
 	}
@@ -283,7 +288,10 @@ void FResonanceAudioReverbPlugin::Init(const FSoundEffectSubmixInitData& InData)
 void FResonanceAudioReverbPlugin::OnPresetChanged()
 {
 	GET_EFFECT_SETTINGS(ResonanceAudioReverbPlugin);
-	ResonanceAudioReverb->UpdateRoomEffects();
+	if (ResonanceAudioReverb)
+	{
+		ResonanceAudioReverb->UpdateRoomEffects();
+	}
 }
 
 uint32 FResonanceAudioReverbPlugin::GetDesiredInputChannelCountOverride() const
@@ -293,7 +301,10 @@ uint32 FResonanceAudioReverbPlugin::GetDesiredInputChannelCountOverride() const
 
 void FResonanceAudioReverbPlugin::OnProcessAudio(const FSoundEffectSubmixInputData& InData, FSoundEffectSubmixOutputData& OutData)
 {
-	ResonanceAudioReverb->ProcessMixedAudio(InData, OutData);
+	if (ResonanceAudioReverb)
+	{
+		ResonanceAudioReverb->ProcessMixedAudio(InData, OutData);
+	}
 }
 
 void FResonanceAudioReverbPlugin::SetResonanceAudioReverbPlugin(ResonanceAudio::FResonanceAudioReverb* InResonanceAudioReverbPlugin)
diff --git a/Engine/Source/Runtime/AudioMixer/Private/AudioMixerChannelMaps.cpp b/Engine/Source/Runtime/AudioMixer/Private/AudioMixerChannelMaps.cpp
--- a/Engine/Source/Runtime/AudioMixer/Private/AudioMixerChannelMaps.cpp
+++ b/Engine/Source/Runtime/AudioMixer/Private/AudioMixerChannelMaps.cpp
@@ -511,7 +511,10 @@ namespace Audio
 
 		// ambisonics is special cased and uses a plugin.
 		TArray<FChannelPositionInfo> FirstOrderAmbisonicsPositions;
-		FirstOrderAmbisonicsPositions.Add(DefaultChannelAzimuthPosition[EAudioMixerChannel::FrontCenter]);
+		FirstOrderAmbisonicsPositions.Add(DefaultChannelAzimuthPosition[EAudioMixerChannel::FrontLeft]);
+		FirstOrderAmbisonicsPositions.Add(DefaultChannelAzimuthPosition[EAudioMixerChannel::FrontRight]);
+		FirstOrderAmbisonicsPositions.Add(DefaultChannelAzimuthPosition[EAudioMixerChannel::SideLeft]);
+		FirstOrderAmbisonicsPositions.Add(DefaultChannelAzimuthPosition[EAudioMixerChannel::SideRight]);
 		FirstOrderAmbisonicsPositions.Sort(FCompareByAzimuth());
 		ChannelAzimuthPositions.Add(ESubmixChannelFormat::Ambisonics, FirstOrderAmbisonicsPositions);
 		OutputChannels[int32(ESubmixChannelFormat::Ambisonics)] = FirstOrderAmbisonicsPositions.Num();
diff --git a/Engine/Source/Runtime/AudioMixer/Private/AudioMixerSource.cpp b/Engine/Source/Runtime/AudioMixer/Private/AudioMixerSource.cpp
--- a/Engine/Source/Runtime/AudioMixer/Private/AudioMixerSource.cpp
+++ b/Engine/Source/Runtime/AudioMixer/Private/AudioMixerSource.cpp
@@ -1037,15 +1037,11 @@ namespace Audio
 			{
 				ESubmixChannelFormat ChannelType = (ESubmixChannelFormat)i;
 
-				// We don't need to compute speaker maps for ambisonics channel maps since we're not doing downmixing on ambisonics sources
-				if (ChannelType != ESubmixChannelFormat::Ambisonics)
+				check(Buffer);
+				const uint32 NumChannels = Buffer->NumChannels;
+				if (ComputeChannelMap(ChannelType, Buffer->NumChannels, ChannelMapInfo.ChannelMap))
 				{
-					check(Buffer);
-					const uint32 NumChannels = Buffer->NumChannels;
-					if (ComputeChannelMap(ChannelType, Buffer->NumChannels, ChannelMapInfo.ChannelMap))
-					{
-						MixerSourceVoice->SetChannelMap(ChannelType, NumChannels, ChannelMapInfo.ChannelMap, bIs3D, WaveInstance->bCenterChannelOnly);
-					}
+					MixerSourceVoice->SetChannelMap(ChannelType, NumChannels, ChannelMapInfo.ChannelMap, bIs3D, WaveInstance->bCenterChannelOnly);
 				}
 			}
 		}
diff --git a/Engine/Source/Runtime/AudioMixer/Private/AudioMixerSourceManager.cpp b/Engine/Source/Runtime/AudioMixer/Private/AudioMixerSourceManager.cpp
--- a/Engine/Source/Runtime/AudioMixer/Private/AudioMixerSourceManager.cpp
+++ b/Engine/Source/Runtime/AudioMixer/Private/AudioMixerSourceManager.cpp
@@ -7,6 +7,7 @@
 #include "AudioMixerSubmix.h"
 #include "IAudioExtensionPlugin.h"
 #include "AudioMixer.h"
+#include "Sound/SampleBuffer.h"
 
 DEFINE_STAT(STAT_AudioMixerHRTF);
 
@@ -729,6 +730,8 @@ namespace Audio
 				}
 			}
 
+			DebugAudioBuffer.Reset();
+
 			SourceInfo.CurrentFrameValues.Init(0.0f, InitParams.NumInputChannels);
 			SourceInfo.NextFrameValues.Init(0.0f, InitParams.NumInputChannels);
 
@@ -1460,7 +1463,9 @@ namespace Audio
 						const float NextFrameValue = SourceInfo.NextFrameValues[Channel];
 						const float CurrentAlpha = SourceInfo.CurrentFrameAlpha;
 
-						PreDistanceAttenBufferPtr[SampleIndex++] = FMath::Lerp(CurrFrameValue, NextFrameValue, CurrentAlpha);
+						float PreDistSampleValue = FMath::Lerp(CurrFrameValue, NextFrameValue, CurrentAlpha);
+						DebugAudioBuffer.Add(PreDistSampleValue);
+						PreDistanceAttenBufferPtr[SampleIndex++] = PreDistSampleValue;
 					}
 					const float CurrentPitchScale = SourceInfo.PitchSourceParam.Update();
 
@@ -1857,6 +1862,11 @@ namespace Audio
 			{
 				Audio::DownmixBuffer(DownmixData.NumInputChannels, 8, *DownmixData.PostEffectBuffers, DownmixData.SevenOneSubmixInfo.OutputBuffer, DownmixData.SevenOneSubmixInfo.ChannelMap.ChannelDestinationGains);
 			}
+
+			if (DownmixData.AmbisonicsSubmixInfo.bInUse)
+			{
+				Audio::DownmixBuffer(DownmixData.NumInputChannels, 4, *DownmixData.PostEffectBuffers, DownmixData.AmbisonicsSubmixInfo.OutputBuffer, DownmixData.AmbisonicsSubmixInfo.ChannelMap.ChannelDestinationGains);
+			}
 		}
 	}
 
@@ -2065,15 +2075,19 @@ namespace Audio
 			{
 				float* PostDistanceAttenBufferPtr = SourceInfo.SourceBuffer.GetData();
 
-				// Process the filters after the source effects (Using optimized filters)
-				for (int32 SampleIndex = 0; SampleIndex < NumOutputFrames * SourceInfo.NumInputChannels; SampleIndex += SourceInfo.NumInputChannels)
+				if (SourceInfo.LowPassFilter.GetCutoffFrequency() < (MAX_FILTER_FREQUENCY - KINDA_SMALL_NUMBER))
+				{
+					SourceInfo.LowPassFilter.ProcessAudio(PreDistanceAttenBufferPtr, NumSamples, PostDistanceAttenBufferPtr);
+				
+				}
+				else
 				{
-					SourceInfo.LowPassFilter.ProcessAudioFrame(&PreDistanceAttenBufferPtr[SampleIndex], &PostDistanceAttenBufferPtr[SampleIndex]);
+					FMemory::Memcpy(PostDistanceAttenBufferPtr, PreDistanceAttenBufferPtr, NumSamples * sizeof(float));
+				}
 
-					if (!DisableHPFilteringCvar)
-					{
-						SourceInfo.HighPassFilter.ProcessAudioFrame(&PostDistanceAttenBufferPtr[SampleIndex], &PostDistanceAttenBufferPtr[SampleIndex]);
-					}
+				if (!DisableHPFilteringCvar)
+				{
+					SourceInfo.HighPassFilter.ProcessAudio(PostDistanceAttenBufferPtr, NumSamples, PostDistanceAttenBufferPtr);
 				}
 
 				// We manually reset interpolation to avoid branches in filter code
@@ -2101,6 +2115,15 @@ namespace Audio
 			// Check the source effect tails condition
 			if (SourceInfo.bIsLastBuffer && SourceInfo.bEffectTailsDone)
 			{
+				TSampleBuffer<> DebugSampleBuffer(DebugAudioBuffer.GetData(), DebugAudioBuffer.Num(), SourceInfo.NumInputChannels, MixerDevice->GetDeviceSampleRate());
+
+				static int32 DebugWrites = 0;
+				FString DebugName = FString::Printf(TEXT("PreAttenDebug_%d"), DebugWrites++);
+				FString Path = FPaths::ProjectLogDir();
+
+				Audio::FSoundWavePCMWriter PCMWriter;
+				PCMWriter.BeginWriteToWavFile(DebugSampleBuffer, DebugName, Path);
+
 				// If we're done and our tails our done, clear everything out
 				SourceInfo.CurrentFrameValues.Reset();
 				SourceInfo.NextFrameValues.Reset();
