Back out changelist 6578557. CIS errors

diff --git a/Engine/Plugins/Runtime/ResonanceAudio/Source/ResonanceAudio/Private/ResonanceAudioReverb.cpp b/Engine/Plugins/Runtime/ResonanceAudio/Source/ResonanceAudio/Private/ResonanceAudioReverb.cpp
--- a/Engine/Plugins/Runtime/ResonanceAudio/Source/ResonanceAudio/Private/ResonanceAudioReverb.cpp
+++ b/Engine/Plugins/Runtime/ResonanceAudio/Source/ResonanceAudio/Private/ResonanceAudioReverb.cpp
@@ -26,12 +26,8 @@ namespace ResonanceAudio
 	{
 		if (GlobalReverbPluginPreset)
 		{
-#if !WITH_EDITOR
 			GlobalReverbPluginPreset->RemoveFromRoot();
-#endif
-			GlobalReverbPluginPreset = nullptr;
 		}
-		ReverbPluginPreset = nullptr;
 	}
 
 	void FResonanceAudioReverb::Initialize(const FAudioPluginInitializationParams InitializationParams)
@@ -134,7 +130,6 @@ namespace ResonanceAudio
 		}
 		else
 		{
-			ReverbPluginPreset = nullptr;
 			return nullptr;
 		}
 	}
@@ -288,10 +283,7 @@ void FResonanceAudioReverbPlugin::Init(const FSoundEffectSubmixInitData& InData)
 void FResonanceAudioReverbPlugin::OnPresetChanged()
 {
 	GET_EFFECT_SETTINGS(ResonanceAudioReverbPlugin);
-	if (ResonanceAudioReverb)
-	{
-		ResonanceAudioReverb->UpdateRoomEffects();
-	}
+	ResonanceAudioReverb->UpdateRoomEffects();
 }
 
 uint32 FResonanceAudioReverbPlugin::GetDesiredInputChannelCountOverride() const
@@ -301,10 +293,7 @@ uint32 FResonanceAudioReverbPlugin::GetDesiredInputChannelCountOverride() const
 
 void FResonanceAudioReverbPlugin::OnProcessAudio(const FSoundEffectSubmixInputData& InData, FSoundEffectSubmixOutputData& OutData)
 {
-	if (ResonanceAudioReverb)
-	{
-		ResonanceAudioReverb->ProcessMixedAudio(InData, OutData);
-	}
+	ResonanceAudioReverb->ProcessMixedAudio(InData, OutData);
 }
 
 void FResonanceAudioReverbPlugin::SetResonanceAudioReverbPlugin(ResonanceAudio::FResonanceAudioReverb* InResonanceAudioReverbPlugin)
diff --git a/Engine/Source/Runtime/AudioMixer/Private/AudioMixerChannelMaps.cpp b/Engine/Source/Runtime/AudioMixer/Private/AudioMixerChannelMaps.cpp
--- a/Engine/Source/Runtime/AudioMixer/Private/AudioMixerChannelMaps.cpp
+++ b/Engine/Source/Runtime/AudioMixer/Private/AudioMixerChannelMaps.cpp
@@ -511,10 +511,7 @@ namespace Audio
 
 		// ambisonics is special cased and uses a plugin.
 		TArray<FChannelPositionInfo> FirstOrderAmbisonicsPositions;
-		FirstOrderAmbisonicsPositions.Add(DefaultChannelAzimuthPosition[EAudioMixerChannel::FrontLeft]);
-		FirstOrderAmbisonicsPositions.Add(DefaultChannelAzimuthPosition[EAudioMixerChannel::FrontRight]);
-		FirstOrderAmbisonicsPositions.Add(DefaultChannelAzimuthPosition[EAudioMixerChannel::SideLeft]);
-		FirstOrderAmbisonicsPositions.Add(DefaultChannelAzimuthPosition[EAudioMixerChannel::SideRight]);
+		FirstOrderAmbisonicsPositions.Add(DefaultChannelAzimuthPosition[EAudioMixerChannel::FrontCenter]);
 		FirstOrderAmbisonicsPositions.Sort(FCompareByAzimuth());
 		ChannelAzimuthPositions.Add(ESubmixChannelFormat::Ambisonics, FirstOrderAmbisonicsPositions);
 		OutputChannels[int32(ESubmixChannelFormat::Ambisonics)] = FirstOrderAmbisonicsPositions.Num();
diff --git a/Engine/Source/Runtime/AudioMixer/Private/AudioMixerSource.cpp b/Engine/Source/Runtime/AudioMixer/Private/AudioMixerSource.cpp
--- a/Engine/Source/Runtime/AudioMixer/Private/AudioMixerSource.cpp
+++ b/Engine/Source/Runtime/AudioMixer/Private/AudioMixerSource.cpp
@@ -1037,11 +1037,15 @@ namespace Audio
 			{
 				ESubmixChannelFormat ChannelType = (ESubmixChannelFormat)i;
 
-				check(Buffer);
-				const uint32 NumChannels = Buffer->NumChannels;
-				if (ComputeChannelMap(ChannelType, Buffer->NumChannels, ChannelMapInfo.ChannelMap))
+				// We don't need to compute speaker maps for ambisonics channel maps since we're not doing downmixing on ambisonics sources
+				if (ChannelType != ESubmixChannelFormat::Ambisonics)
 				{
-					MixerSourceVoice->SetChannelMap(ChannelType, NumChannels, ChannelMapInfo.ChannelMap, bIs3D, WaveInstance->bCenterChannelOnly);
+					check(Buffer);
+					const uint32 NumChannels = Buffer->NumChannels;
+					if (ComputeChannelMap(ChannelType, Buffer->NumChannels, ChannelMapInfo.ChannelMap))
+					{
+						MixerSourceVoice->SetChannelMap(ChannelType, NumChannels, ChannelMapInfo.ChannelMap, bIs3D, WaveInstance->bCenterChannelOnly);
+					}
 				}
 			}
 		}
diff --git a/Engine/Source/Runtime/AudioMixer/Private/AudioMixerSourceManager.cpp b/Engine/Source/Runtime/AudioMixer/Private/AudioMixerSourceManager.cpp
--- a/Engine/Source/Runtime/AudioMixer/Private/AudioMixerSourceManager.cpp
+++ b/Engine/Source/Runtime/AudioMixer/Private/AudioMixerSourceManager.cpp
@@ -7,7 +7,6 @@
 #include "AudioMixerSubmix.h"
 #include "IAudioExtensionPlugin.h"
 #include "AudioMixer.h"
-#include "Sound/SampleBuffer.h"
 
 DEFINE_STAT(STAT_AudioMixerHRTF);
 
@@ -730,8 +729,6 @@ namespace Audio
 				}
 			}
 
-			DebugAudioBuffer.Reset();
-
 			SourceInfo.CurrentFrameValues.Init(0.0f, InitParams.NumInputChannels);
 			SourceInfo.NextFrameValues.Init(0.0f, InitParams.NumInputChannels);
 
@@ -1463,9 +1460,7 @@ namespace Audio
 						const float NextFrameValue = SourceInfo.NextFrameValues[Channel];
 						const float CurrentAlpha = SourceInfo.CurrentFrameAlpha;
 
-						float PreDistSampleValue = FMath::Lerp(CurrFrameValue, NextFrameValue, CurrentAlpha);
-						DebugAudioBuffer.Add(PreDistSampleValue);
-						PreDistanceAttenBufferPtr[SampleIndex++] = PreDistSampleValue;
+						PreDistanceAttenBufferPtr[SampleIndex++] = FMath::Lerp(CurrFrameValue, NextFrameValue, CurrentAlpha);
 					}
 					const float CurrentPitchScale = SourceInfo.PitchSourceParam.Update();
 
@@ -1862,11 +1857,6 @@ namespace Audio
 			{
 				Audio::DownmixBuffer(DownmixData.NumInputChannels, 8, *DownmixData.PostEffectBuffers, DownmixData.SevenOneSubmixInfo.OutputBuffer, DownmixData.SevenOneSubmixInfo.ChannelMap.ChannelDestinationGains);
 			}
-
-			if (DownmixData.AmbisonicsSubmixInfo.bInUse)
-			{
-				Audio::DownmixBuffer(DownmixData.NumInputChannels, 4, *DownmixData.PostEffectBuffers, DownmixData.AmbisonicsSubmixInfo.OutputBuffer, DownmixData.AmbisonicsSubmixInfo.ChannelMap.ChannelDestinationGains);
-			}
 		}
 	}
 
@@ -2075,19 +2065,15 @@ namespace Audio
 			{
 				float* PostDistanceAttenBufferPtr = SourceInfo.SourceBuffer.GetData();
 
-				if (SourceInfo.LowPassFilter.GetCutoffFrequency() < (MAX_FILTER_FREQUENCY - KINDA_SMALL_NUMBER))
-				{
-					SourceInfo.LowPassFilter.ProcessAudio(PreDistanceAttenBufferPtr, NumSamples, PostDistanceAttenBufferPtr);
-				
-				}
-				else
+				// Process the filters after the source effects (Using optimized filters)
+				for (int32 SampleIndex = 0; SampleIndex < NumOutputFrames * SourceInfo.NumInputChannels; SampleIndex += SourceInfo.NumInputChannels)
 				{
-					FMemory::Memcpy(PostDistanceAttenBufferPtr, PreDistanceAttenBufferPtr, NumSamples * sizeof(float));
-				}
+					SourceInfo.LowPassFilter.ProcessAudioFrame(&PreDistanceAttenBufferPtr[SampleIndex], &PostDistanceAttenBufferPtr[SampleIndex]);
 
-				if (!DisableHPFilteringCvar)
-				{
-					SourceInfo.HighPassFilter.ProcessAudio(PostDistanceAttenBufferPtr, NumSamples, PostDistanceAttenBufferPtr);
+					if (!DisableHPFilteringCvar)
+					{
+						SourceInfo.HighPassFilter.ProcessAudioFrame(&PostDistanceAttenBufferPtr[SampleIndex], &PostDistanceAttenBufferPtr[SampleIndex]);
+					}
 				}
 
 				// We manually reset interpolation to avoid branches in filter code
@@ -2115,15 +2101,6 @@ namespace Audio
 			// Check the source effect tails condition
 			if (SourceInfo.bIsLastBuffer && SourceInfo.bEffectTailsDone)
 			{
-				TSampleBuffer<> DebugSampleBuffer(DebugAudioBuffer.GetData(), DebugAudioBuffer.Num(), SourceInfo.NumInputChannels, MixerDevice->GetDeviceSampleRate());
-
-				static int32 DebugWrites = 0;
-				FString DebugName = FString::Printf(TEXT("PreAttenDebug_%d"), DebugWrites++);
-				FString Path = FPaths::ProjectLogDir();
-
-				Audio::FSoundWavePCMWriter PCMWriter;
-				PCMWriter.BeginWriteToWavFile(DebugSampleBuffer, DebugName, Path);
-
 				// If we're done and our tails our done, clear everything out
 				SourceInfo.CurrentFrameValues.Reset();
 				SourceInfo.NextFrameValues.Reset();
