Added backwards compatibility for PassiveFilterParams change.  Fixed up QAGame code using struct.

diff --git a/Engine/Source/Runtime/AudioMixer/Public/DSP/PassiveFilter.h b/Engine/Source/Runtime/AudioMixer/Public/DSP/PassiveFilter.h
--- a/Engine/Source/Runtime/AudioMixer/Public/DSP/PassiveFilter.h
+++ b/Engine/Source/Runtime/AudioMixer/Public/DSP/PassiveFilter.h
@@ -9,6 +9,7 @@
 
 namespace Audio
 {
+
 	struct FPassiveFilterParams
 	{
 		enum class EClass
@@ -43,6 +44,9 @@ namespace Audio
 		}
 	};
 
+	UE_DEPRECATED(4.23, "PassiveFilterParams has been renamed to FPassiveFilterParams")
+	typedef FPassiveFilterParams PassiveFilterParams;
+
 	static float EvaluateChebyshevPolynomial(float FrequencyRatio, int32 Order)
 	{
 		if (FMath::IsNearlyEqual(FrequencyRatio, 1.0f))
