Max's fix for initial fade in issue with bug fix for bIsInit not being used properly as a gating mechanism

diff --git a/Engine/Source/Runtime/AudioMixer/Private/AudioMixerSourceManager.h b/Engine/Source/Runtime/AudioMixer/Private/AudioMixerSourceManager.h
--- a/Engine/Source/Runtime/AudioMixer/Private/AudioMixerSourceManager.h
+++ b/Engine/Source/Runtime/AudioMixer/Private/AudioMixerSourceManager.h
@@ -169,7 +169,7 @@ namespace Audio
 			CopySize = InNumInChannels * InNumOutChannels * sizeof(float);
 			FMemory::Memzero(ChannelStartGains, CopySize);
 			FMemory::Memzero(ChannelDestinationGains, CopySize);
-			bIsInit = true;
+			bIsInit = false;
 		}
 
 		FORCEINLINE void CopyDestinationToStart()
@@ -180,16 +180,17 @@ namespace Audio
 		FORCEINLINE void SetChannelMap(const float* RESTRICT InChannelGains)
 		{
 			FMemory::Memcpy(ChannelDestinationGains, InChannelGains, CopySize);
-			if (bIsInit)
+			if (!bIsInit)
 			{
 				FMemory::Memcpy(ChannelStartGains, InChannelGains, CopySize);
+				bIsInit = true;
 			}
 		}
 
 	private:
 		FSourceChannelMap()
 			: CopySize(0)
-			, bIsInit(true)
+			, bIsInit(false)
 		{
 		}
 	};
