Fix for XMA2 popping.

- Partial fix only includes changes to xma2 cooking format while we're recooking.

diff --git a/Engine/Source/Runtime/AudioMixer/Private/AudioMixerSourceDecode.cpp b/Engine/Source/Runtime/AudioMixer/Private/AudioMixerSourceDecode.cpp
--- a/Engine/Source/Runtime/AudioMixer/Private/AudioMixerSourceDecode.cpp
+++ b/Engine/Source/Runtime/AudioMixer/Private/AudioMixerSourceDecode.cpp
@@ -102,7 +102,7 @@ public:
 
 				// Create a buffer to decode into that's of the appropriate size
 				TArray<uint8> DecodeBuffer;
-				DecodeBuffer.AddUninitialized(ByteSize);
+				DecodeBuffer.AddZeroed(ByteSize);
 
 				// skip the first buffers if we've already decoded them during Precache:
 				if (DecodeTaskData.bSkipFirstBuffer)
diff --git a/Engine/Source/Runtime/AudioMixer/Public/DSP/Dsp.h b/Engine/Source/Runtime/AudioMixer/Public/DSP/Dsp.h
--- a/Engine/Source/Runtime/AudioMixer/Public/DSP/Dsp.h
+++ b/Engine/Source/Runtime/AudioMixer/Public/DSP/Dsp.h
@@ -512,7 +512,7 @@ namespace Audio
 			const int32 NumRead = FMath::Min<int32>(NumToCopy, Capacity - ReadIndex);
 			FMemory::Memcpy(OutBuffer, &SrcBuffer[ReadIndex], NumRead * sizeof(SampleType));
 				
-			FMemory::Memcpy(&OutBuffer[NumRead], &SrcBuffer[0], (NumSamples - NumRead) * sizeof(SampleType));
+			FMemory::Memcpy(&OutBuffer[NumRead], &SrcBuffer[0], (NumToCopy - NumRead) * sizeof(SampleType));
 
 			check(NumSamples < ((uint32)TNumericLimits<int32>::Max()));
 
