[NO NON-PII COMMIT MESSAGE]

diff --git a/Engine/Source/Runtime/AudioMixer/Private/DSP/AudioFFT.cpp b/Engine/Source/Runtime/AudioMixer/Private/DSP/AudioFFT.cpp
--- a/Engine/Source/Runtime/AudioMixer/Private/DSP/AudioFFT.cpp
+++ b/Engine/Source/Runtime/AudioMixer/Private/DSP/AudioFFT.cpp
@@ -156,9 +156,11 @@ namespace Audio
 		// Separates InBuffer (assumed to be mono here)
 		void SeperateInPlace(float* InBuffer, uint32 NumSamples)
 		{
+			check(FMath::CountBits(NumSamples) == 1)
+			const uint32 NumBits = FMath::CountTrailingZeros(NumSamples);
+			
 			for (uint32 Index = 0; Index < NumSamples; Index++)
 			{
-				const uint32 NumBits = FMath::Log2(NumSamples);
 				uint32 SwappedIndex = SlowBitReversal(Index, NumBits);
 				if (Index < SwappedIndex)
 				{
@@ -169,9 +171,10 @@ namespace Audio
 
 		void SeparateIntoCopy(float* InBuffer, float* OutBuffer, uint32 NumSamples)
 		{
+			check(FMath::CountBits(NumSamples) == 1)
+			const uint32 NumBits = FMath::CountTrailingZeros(NumSamples);
 			for (uint32 Index = 0; Index < NumSamples; Index++)
 			{
-				const uint32 NumBits = FMath::Log2(NumSamples);
 				const uint32 ReversedIndex = SlowBitReversal2(Index, NumBits);
 				OutBuffer[ReversedIndex] = InBuffer[Index];
 			}
@@ -179,8 +182,9 @@ namespace Audio
 
 		void ComputeButterfliesInPlace(float* OutReal, float* OutImag, uint32 NumSamples)
 		{
-			const uint32 LogNumSamples = FMath::Log2(NumSamples);
-
+			check(FMath::CountBits(NumSamples) == 1)
+			const uint32 LogNumSamples = FMath::CountTrailingZeros(NumSamples);
+			
 			for (uint32 S = 1; S <= LogNumSamples; S++)
 			{
 				const uint32 M = (1u << S);
