Changed AudioMixerSource::GetPlaybackPercent() to return the last playback percentage if a sound is stopping (instead of 0.0f)

diff --git a/Engine/Source/Runtime/AudioMixer/Private/AudioMixerSource.cpp b/Engine/Source/Runtime/AudioMixer/Private/AudioMixerSource.cpp
--- a/Engine/Source/Runtime/AudioMixer/Private/AudioMixerSource.cpp
+++ b/Engine/Source/Runtime/AudioMixer/Private/AudioMixerSource.cpp
@@ -30,6 +30,7 @@ namespace Audio
 		, MixerSourceBuffer(nullptr)
 		, MixerSourceVoice(nullptr)
 		, PreviousAzimuth(-1.0f)
+		, PreviousPlaybackPercent(0.0f)
 		, InitializationState(EMixerSourceInitializationState::NotInitialized)
 		, bPlayedCachedBuffer(false)
 		, bPlaying(false)
@@ -719,19 +720,19 @@ namespace Audio
 	{
 		if (InitializationState != EMixerSourceInitializationState::Initialized)
 		{
-			return 0.f;
+			return PreviousPlaybackPercent;
 		}
 
 		if (MixerSourceVoice && NumTotalFrames > 0)
 		{
 			int64 NumFrames = MixerSourceVoice->GetNumFramesPlayed();
 			AUDIO_MIXER_CHECK(NumTotalFrames > 0);
-			float PlaybackPercent = (float)NumFrames / NumTotalFrames;
+			PreviousPlaybackPercent = (float)NumFrames / NumTotalFrames;
 			if (WaveInstance->LoopingMode == LOOP_Never)
 			{
-				PlaybackPercent = FMath::Min(PlaybackPercent, 1.0f);
+				PreviousPlaybackPercent = FMath::Min(PreviousPlaybackPercent, 1.0f);
 			}
-			return PlaybackPercent;
+			return PreviousPlaybackPercent;
 		}
 		else
 		{
diff --git a/Engine/Source/Runtime/AudioMixer/Private/AudioMixerSource.h b/Engine/Source/Runtime/AudioMixer/Private/AudioMixerSource.h
--- a/Engine/Source/Runtime/AudioMixer/Private/AudioMixerSource.h
+++ b/Engine/Source/Runtime/AudioMixer/Private/AudioMixerSource.h
@@ -157,6 +157,7 @@ namespace Audio
 		FChannelMapInfo ChannelMaps[(int32) ESubmixChannelFormat::Count];
 
 		float PreviousAzimuth;
+		float PreviousPlaybackPercent;
 
 		FSpatializationParams SpatializationParams;
 
