Fix for GetPlaybackPerncent() modifying a member var in a const function

diff --git a/Engine/Source/Runtime/AudioMixer/Private/AudioMixerSource.h b/Engine/Source/Runtime/AudioMixer/Private/AudioMixerSource.h
--- a/Engine/Source/Runtime/AudioMixer/Private/AudioMixerSource.h
+++ b/Engine/Source/Runtime/AudioMixer/Private/AudioMixerSource.h
@@ -157,7 +157,7 @@ namespace Audio
 		FChannelMapInfo ChannelMaps[(int32) ESubmixChannelFormat::Count];
 
 		float PreviousAzimuth;
-		float PreviousPlaybackPercent;
+		mutable float PreviousPlaybackPercent;
 
 		FSpatializationParams SpatializationParams;
 
