[NO NON-PII COMMIT MESSAGE]

diff --git a/Engine/Source/Runtime/AudioMixer/Private/AudioMixerDevice.cpp b/Engine/Source/Runtime/AudioMixer/Private/AudioMixerDevice.cpp
--- a/Engine/Source/Runtime/AudioMixer/Private/AudioMixerDevice.cpp
+++ b/Engine/Source/Runtime/AudioMixer/Private/AudioMixerDevice.cpp
@@ -685,7 +685,7 @@ namespace Audio
 				}
 				else
 				{
-					ReverbPreset = NewObject<USubmixEffectReverbFastPreset>(MasterReverbSoundSubmix, TEXT("Master Reverb Effect Preset"));
+					ReverbPreset = NewObject<USubmixEffectReverbFastPreset>(MasterReverbSoundSubmix, TEXT("Master Reverb Effect Fast Preset"));
 				}
 				
 				ReverbPreset->AddToRoot();
diff --git a/Engine/Source/Runtime/AudioMixer/Private/AudioMixerEffectsManager.cpp b/Engine/Source/Runtime/AudioMixer/Private/AudioMixerEffectsManager.cpp
--- a/Engine/Source/Runtime/AudioMixer/Private/AudioMixerEffectsManager.cpp
+++ b/Engine/Source/Runtime/AudioMixer/Private/AudioMixerEffectsManager.cpp
@@ -97,7 +97,9 @@ namespace Audio
 
 	FAudioMixerEffectsManager::FAudioMixerEffectsManager(FAudioDevice* InDevice)
 		: FAudioEffectsManager(InDevice)
-	{}
+	{
+		bUseLegacyReverb = GetDefault<UAudioSettings>()->bEnableLegacyReverb;
+	}
 
 	FAudioMixerEffectsManager::~FAudioMixerEffectsManager()
 	{}
@@ -115,7 +117,7 @@ namespace Audio
 			if (SoundEffectSubmix)
 			{
 				// Choose correct reverb based upon ini settings.
-				if (GetDefault<UAudioSettings>()->bEnableLegacyReverb)
+				if (bUseLegacyReverb)
 				{
 					FSubmixEffectReverb* SoundEffectReverb = static_cast<FSubmixEffectReverb*>(SoundEffectSubmix);
 					SoundEffectReverb->SetEffectParameters(ReverbEffectParameters);
diff --git a/Engine/Source/Runtime/Engine/Private/AudioEffect.cpp b/Engine/Source/Runtime/Engine/Private/AudioEffect.cpp
--- a/Engine/Source/Runtime/Engine/Private/AudioEffect.cpp
+++ b/Engine/Source/Runtime/Engine/Private/AudioEffect.cpp
@@ -252,6 +252,7 @@ FAudioEffectsManager::FAudioEffectsManager( FAudioDevice* InDevice )
 	, bEQActive(false)
 	, bReverbChanged(true) // Setting to true to catch the first default reverb setting
 	, bEQChanged(false)
+	, bUseLegacyReverb(false)
 {
 	InitAudioEffects();
 }
diff --git a/Engine/Source/Runtime/Engine/Public/AudioEffect.h b/Engine/Source/Runtime/Engine/Public/AudioEffect.h
--- a/Engine/Source/Runtime/Engine/Public/AudioEffect.h
+++ b/Engine/Source/Runtime/Engine/Public/AudioEffect.h
@@ -220,6 +220,7 @@ protected:
 	bool bEQActive;
 	bool bReverbChanged;
 	bool bEQChanged;
+	bool bUseLegacyReverb;
 };
 
 // end 
