Fix for reverb not responding well to rapid settings changes. (Smoothed all-pass filter parameter and applied wet level gain to input instead of output)

diff --git a/Engine/Source/Runtime/AudioMixer/Private/DSP/DynamicDelayAPF.cpp b/Engine/Source/Runtime/AudioMixer/Private/DSP/DynamicDelayAPF.cpp
--- a/Engine/Source/Runtime/AudioMixer/Private/DSP/DynamicDelayAPF.cpp
+++ b/Engine/Source/Runtime/AudioMixer/Private/DSP/DynamicDelayAPF.cpp
@@ -6,12 +6,14 @@
 using namespace Audio;
 
 FDynamicDelayAPF::FDynamicDelayAPF(float InG, int32 InMinDelay, int32 InMaxDelay, int32 InMaxNumInternalBufferSamples)
-:	G(InG)
-,	MinDelay(InMinDelay)
+:	MinDelay(InMinDelay)
 ,	MaxDelay(InMaxDelay)
 ,	NumDelaySamples(InMinDelay - 1)
 ,	NumInternalBufferSamples(InMaxNumInternalBufferSamples)
 {
+	G.SetValueInterrupt(0.0f);
+	G.Init(48000); // TODO: use actual sample rate (will warp ease time slightly right now)
+
 	checkf(NumDelaySamples >= 0, TEXT("Minimum delay must be atleast 1"));
 	// NumInternalBufferSamples must be less than the length of the delay to support buffer indexing logic.
 	int32 MaxBufferSamples = FMath::Min(NumDelaySamples, InMaxNumInternalBufferSamples);
@@ -102,13 +104,24 @@ void FDynamicDelayAPF::ProcessAudioBlock(const float* InSamples, const AlignedFl
 	// w[n] = x[n] + G * w[n - d]
 	DelayLineInput.Reset(InNum);
 	DelayLineInput.AddUninitialized(InNum);
-	BufferWeightedSumFast(WorkBufferB.GetData(), G, InSamples, DelayLineInput.GetData(), InNum);
-	
-	BufferUnderflowClampFast(DelayLineInput);
+
+
+//	BufferWeightedSumFast(WorkBufferB.GetData(), G, InSamples, DelayLineInput.GetData(), InNum);
+	const float LastG = G.GetValue();
+	const float CurrG = G.GetValue(InNum);
+
+	// WorkBufferA = G * WorkBufferB
+	FMemory::Memcpy(WorkBufferA.GetData(), WorkBufferB.GetData(), InNum * sizeof(float));
+	FadeBufferFast(WorkBufferA.GetData(), InNum, LastG, CurrG);
+	SumBuffers(InSamples, WorkBufferA.GetData(), DelayLineInput.GetData(), InNum);
+
+	BufferUnderflowClampFast(WorkBufferB.GetData(), InNum);
 
 	// y[n] = -G w[n] + w[n - d]
-	BufferWeightedSumFast(DelayLineInput.GetData(), -G, WorkBufferB.GetData(), OutSamples, InNum);
-	
+//	BufferWeightedSumFast(DelayLineInput.GetData(), -G, WorkBufferB.GetData(), OutSamples, InNum);
+	FadeBufferFast(DelayLineInput.GetData(), InNum, -LastG, -CurrG);
+	SumBuffers(DelayLineInput.GetData(), WorkBufferB.GetData(), OutSamples, InNum);
+
 	// Update delay line
 	IntegerDelayLine->RemoveSamples(InNum);
 	IntegerDelayLine->AddSamples(DelayLineInput.GetData(), InNum);
diff --git a/Engine/Source/Runtime/AudioMixer/Private/DSP/ReverbFast.cpp b/Engine/Source/Runtime/AudioMixer/Private/DSP/ReverbFast.cpp
--- a/Engine/Source/Runtime/AudioMixer/Private/DSP/ReverbFast.cpp
+++ b/Engine/Source/Runtime/AudioMixer/Private/DSP/ReverbFast.cpp
@@ -34,6 +34,8 @@ namespace Audio {
 
 	FPlateReverbFast::FPlateReverbFast(float InSampleRate, int32 InMaxInternalBufferSamples, const FPlateReverbFastSettings& InSettings)
 		: SampleRate(InSampleRate)
+		, LastWetness(0.0f)
+		, bProcessCallSinceWetnessChanged(false)
 		, EarlyReflections(InSampleRate, InMaxInternalBufferSamples)
 		, LateReflections(InSampleRate, InMaxInternalBufferSamples, InSettings.LateReflections)
 		, bEnableEarlyReflections(true)
@@ -48,6 +50,12 @@ namespace Audio {
 	void FPlateReverbFast::SetSettings(const FPlateReverbFastSettings& InSettings)
 	{
 		// Copy, clamp and apply settings
+		if (bProcessCallSinceWetnessChanged)
+		{
+			LastWetness = Settings.Wetness;
+			bProcessCallSinceWetnessChanged = false;
+		}
+
 		Settings = InSettings;
 		ClampSettings(Settings);
 		ApplySettings();
@@ -73,6 +81,24 @@ namespace Audio {
 	// Process a buffer of input audio samples.
 	void FPlateReverbFast::ProcessAudio(const AlignedFloatBuffer& InSamples, const int32 InNumChannels, AlignedFloatBuffer& OutSamples, const int32 OutNumChannels)
 	{
+		ScaledInputBuffer.Reset(InSamples.Num());
+		ScaledInputBuffer.AddUninitialized(InSamples.Num());
+		check(ScaledInputBuffer.Num() == InSamples.Num());
+
+		FMemory::Memcpy(ScaledInputBuffer.GetData(), InSamples.GetData(), InSamples.Num() * sizeof(float));
+
+		// Scale input by wetness (or fade to new wetness)
+		if (FMath::IsNearlyEqual(LastWetness, Settings.Wetness))
+		{
+			MultiplyBufferByConstantInPlace(ScaledInputBuffer.GetData(), InSamples.Num(), Settings.Wetness);
+		}
+		else
+		{
+			FadeBufferFast(ScaledInputBuffer, LastWetness, Settings.Wetness);
+			LastWetness = Settings.Wetness;
+		}
+
+
 		checkf((1 == InNumChannels) || (2 == InNumChannels), TEXT("FPlateReverbFast only supports 1 or 2 channel inputs."))
 		checkf(OutNumChannels >= 2, TEXT("FPlateReverbFast requires at least 2 output channels."))
 
@@ -102,12 +128,12 @@ namespace Audio {
 		if (bEnableEarlyReflections && !bEnableLateReflections)
 		{
 			// Only generate early reflections.
-			EarlyReflections.ProcessAudio(InSamples, InNumChannels, FrontLeftReverbSamples, FrontRightReverbSamples);
+			EarlyReflections.ProcessAudio(ScaledInputBuffer, InNumChannels, FrontLeftReverbSamples, FrontRightReverbSamples);
 		}
 		else if (!bEnableEarlyReflections && bEnableLateReflections)
 		{
 			// Only generate late reflections.
-			LateReflections.ProcessAudio(InSamples, InNumChannels, FrontLeftReverbSamples, FrontRightReverbSamples);
+			LateReflections.ProcessAudio(ScaledInputBuffer, InNumChannels, FrontLeftReverbSamples, FrontRightReverbSamples);
 		}
 		else if (bEnableEarlyReflections && bEnableLateReflections)
 		{
@@ -123,19 +149,17 @@ namespace Audio {
 			FrontRightEarlyReflectionsSamples.AddUninitialized(InNumFrames);
 
 			// Generate both early reflections and late reflections 
-			EarlyReflections.ProcessAudio(InSamples, InNumChannels, FrontLeftEarlyReflectionsSamples, FrontRightEarlyReflectionsSamples);
-			LateReflections.ProcessAudio(InSamples, InNumChannels, FrontLeftLateReflectionsSamples, FrontRightLateReflectionsSamples);
+			EarlyReflections.ProcessAudio(ScaledInputBuffer, InNumChannels, FrontLeftEarlyReflectionsSamples, FrontRightEarlyReflectionsSamples);
+			LateReflections.ProcessAudio(ScaledInputBuffer, InNumChannels, FrontLeftLateReflectionsSamples, FrontRightLateReflectionsSamples);
 			// Add early and late reflections together.
 			SumBuffers(FrontLeftEarlyReflectionsSamples, FrontLeftLateReflectionsSamples, FrontLeftReverbSamples);
 			SumBuffers(FrontRightEarlyReflectionsSamples, FrontRightLateReflectionsSamples, FrontRightReverbSamples);
 		}
 
-		// Scale output by wetness
-		MultiplyBufferByConstantInPlace(FrontLeftReverbSamples, Settings.Wetness);
-		MultiplyBufferByConstantInPlace(FrontRightReverbSamples, Settings.Wetness);
 
 		// Interleave and upmix
 		InterleaveAndMixOutput(FrontLeftReverbSamples, FrontRightReverbSamples, OutSamples, OutNumChannels);
+		bProcessCallSinceWetnessChanged = true;
 	}
 
 	void FPlateReverbFast::ClampSettings(FPlateReverbFastSettings& InOutSettings)
diff --git a/Engine/Source/Runtime/AudioMixer/Public/DSP/Dsp.h b/Engine/Source/Runtime/AudioMixer/Public/DSP/Dsp.h
--- a/Engine/Source/Runtime/AudioMixer/Public/DSP/Dsp.h
+++ b/Engine/Source/Runtime/AudioMixer/Public/DSP/Dsp.h
@@ -357,6 +357,22 @@ namespace Audio
 			++CurrentTick;
 			return CurrentValue;
 		}
+
+		// same as GetValue(), but overloaded to increment Current Tick by NumTicksToJumpAhead
+		// (before getting the value);
+		float GetValue(int32 NumTicksToJumpAhead)
+		{
+			if (IsDone())
+			{
+				return CurrentValue;
+			}
+
+			CurrentTick += NumTicksToJumpAhead;
+			CurrentValue = DeltaValue * (float)CurrentTick / DurationTicks + StartValue;
+
+			return CurrentValue;
+		}
+
 		 
 		// Updates the target value without changing the duration or tick data.
 		// Sets the state as if the new value was the target value all along
diff --git a/Engine/Source/Runtime/AudioMixer/Public/DSP/DynamicDelayAPF.h b/Engine/Source/Runtime/AudioMixer/Public/DSP/DynamicDelayAPF.h
--- a/Engine/Source/Runtime/AudioMixer/Public/DSP/DynamicDelayAPF.h
+++ b/Engine/Source/Runtime/AudioMixer/Public/DSP/DynamicDelayAPF.h
@@ -25,7 +25,10 @@ namespace Audio
 		~FDynamicDelayAPF();
 
 		// Set the APF feedback/feedforward gain coefficient
-		void SetG(float InG) { G = InG; }
+		void SetG(float InG)
+		{
+			G.SetValue(InG, EaseTimeInSec);
+		}
 
 		// Processes InSamples through the all pass filter and populates OutSamples with the filter output.
 		// InDelays denotes the per-sample delay of the allpass. It must have an equal number of elements as InSamples
@@ -34,6 +37,8 @@ namespace Audio
 		// Zeros the internal delay line.
 		void Reset();
 
+		void SetEaseTimeInSec(float InEaseTimeInSec) { EaseTimeInSec = InEaseTimeInSec; }
+
 	protected:
 
 		// Process one block of audio.
@@ -41,7 +46,8 @@ namespace Audio
 
 	private:
 		// Feedback/Feedforward gain coefficient
-		float G;
+		FLinearEase G;
+		float EaseTimeInSec{ 2.0 };
 		int32 MinDelay;
 		int32 MaxDelay;
 		int32 NumDelaySamples;
@@ -56,7 +62,6 @@ namespace Audio
 		// Delay line memory.
 		TUniquePtr<FAlignedBlockBuffer> IntegerDelayLine;
 		TUniquePtr<FLinearInterpFractionalDelay> FractionalDelayLine;
-
 	};
 
 }
diff --git a/Engine/Source/Runtime/AudioMixer/Public/DSP/ReverbFast.h b/Engine/Source/Runtime/AudioMixer/Public/DSP/ReverbFast.h
--- a/Engine/Source/Runtime/AudioMixer/Public/DSP/ReverbFast.h
+++ b/Engine/Source/Runtime/AudioMixer/Public/DSP/ReverbFast.h
@@ -94,6 +94,8 @@ namespace Audio
 
 
 			float SampleRate;
+			float LastWetness;
+			bool bProcessCallSinceWetnessChanged;
 
 			FPlateReverbFastSettings Settings;
 			FEarlyReflectionsFast EarlyReflections;
@@ -107,6 +109,7 @@ namespace Audio
 			AlignedFloatBuffer FrontRightReverbSamples;
 			AlignedFloatBuffer LeftAttenuatedSamples;
 			AlignedFloatBuffer RightAttenuatedSamples;
+			AlignedFloatBuffer ScaledInputBuffer;
 
 			bool bEnableEarlyReflections;
 			bool bEnableLateReflections;
