Changing assert to log

diff --git a/Engine/Source/Runtime/AudioMixer/Private/AudioMixerSourceManager.cpp b/Engine/Source/Runtime/AudioMixer/Private/AudioMixerSourceManager.cpp
--- a/Engine/Source/Runtime/AudioMixer/Private/AudioMixerSourceManager.cpp
+++ b/Engine/Source/Runtime/AudioMixer/Private/AudioMixerSourceManager.cpp
@@ -347,7 +347,10 @@ namespace Audio
 				const int32 NextIndex = !CurrentGameIndex;
 
 				// Make sure we've actually emptied the command queue from the render thread before writing to it
-				check(CommandBuffers[NextIndex].SourceCommandQueue.Num() == 0);
+				if (CommandBuffers[NextIndex].SourceCommandQueue.Num() != 0)
+				{
+					UE_LOG(LogAudioMixer, Warning, TEXT("Source command queue not empty: %d"), CommandBuffers[NextIndex].SourceCommandQueue.Num());
+				}
 				bPumpQueue = true;
 			}
 		}
