Fix for crashing with 2 channels of audio for HRTF spatialization.

diff --git a/Engine/Source/Runtime/AudioMixer/Private/AudioMixerSourceManager.cpp b/Engine/Source/Runtime/AudioMixer/Private/AudioMixerSourceManager.cpp
--- a/Engine/Source/Runtime/AudioMixer/Private/AudioMixerSourceManager.cpp
+++ b/Engine/Source/Runtime/AudioMixer/Private/AudioMixerSourceManager.cpp
@@ -116,6 +116,7 @@ namespace Audio
 		, NumSourceWorkers(4)
 		, bInitialized(false)
 		, bUsingSpatializationPlugin(false)
+		, MaxChannelsSupportedBySpatializationPlugin(1)
 	{
 		// Get a manual resetable event
 		const bool bIsManualReset = true;
@@ -275,6 +276,7 @@ namespace Audio
 		if (SpatializationPlugin.IsValid())
 		{
 			bUsingSpatializationPlugin = true;
+			MaxChannelsSupportedBySpatializationPlugin = MixerDevice->MaxChannelsSupportedBySpatializationPlugin;
 		}
 
 		bInitialized = true;
@@ -1633,7 +1635,7 @@ namespace Audio
 			CSV_SCOPED_TIMING_STAT(Audio, HRTF);
 
 			AUDIO_MIXER_CHECK(SpatializationPlugin.IsValid());
-			AUDIO_MIXER_CHECK(SourceInfo.NumInputChannels == 1);
+			AUDIO_MIXER_CHECK(SourceInfo.NumInputChannels <= MaxChannelsSupportedBySpatializationPlugin);
 
 			FAudioPluginSourceInputData AudioPluginInputData;
 			AudioPluginInputData.AudioBuffer = &SourceInfo.SourceBuffer;
diff --git a/Engine/Source/Runtime/AudioMixer/Private/AudioMixerSourceManager.h b/Engine/Source/Runtime/AudioMixer/Private/AudioMixerSourceManager.h
--- a/Engine/Source/Runtime/AudioMixer/Private/AudioMixerSourceManager.h
+++ b/Engine/Source/Runtime/AudioMixer/Private/AudioMixerSourceManager.h
@@ -570,6 +570,7 @@ namespace Audio
 
 		uint8 bInitialized : 1;
 		uint8 bUsingSpatializationPlugin : 1;
+		int32 MaxChannelsSupportedBySpatializationPlugin;
 
 		// Set to true when the audio source manager should pump the command queue
 		FThreadSafeBool bPumpQueue;
