Fix for quad speaker mapping

diff --git a/Engine/Source/Runtime/AudioMixer/Private/AudioMixerChannelMaps.cpp b/Engine/Source/Runtime/AudioMixer/Private/AudioMixerChannelMaps.cpp
--- a/Engine/Source/Runtime/AudioMixer/Private/AudioMixerChannelMaps.cpp
+++ b/Engine/Source/Runtime/AudioMixer/Private/AudioMixerChannelMaps.cpp
@@ -294,6 +294,37 @@ namespace Audio
 				}
 			}
 		}
+		// Quad has a special-case to map input channels 0 1 2 3 to 0 1 4 5: 
+		else if (NumSourceChannels == 4)
+		{
+			const int32 FrontLeftChannel = 0;
+			for (int32 OutputChannel = 0; OutputChannel < NumOutputChannels; ++OutputChannel)
+			{
+				const int32 Index = OutputChannel * AUDIO_MIXER_MAX_OUTPUT_CHANNELS + FrontLeftChannel;
+				OutChannelMap.Add(Matrix[Index]);
+			}
+
+			const int32 FrontRightChannel = 1;
+			for (int32 OutputChannel = 0; OutputChannel < NumOutputChannels; ++OutputChannel)
+			{
+				const int32 Index = OutputChannel * AUDIO_MIXER_MAX_OUTPUT_CHANNELS + FrontRightChannel;
+				OutChannelMap.Add(Matrix[Index]);
+			}
+
+			const int32 SurroundLeftChannel = 4;
+			for (int32 OutputChannel = 0; OutputChannel < NumOutputChannels; ++OutputChannel)
+			{
+				const int32 Index = OutputChannel * AUDIO_MIXER_MAX_OUTPUT_CHANNELS + SurroundLeftChannel;
+				OutChannelMap.Add(Matrix[Index]);
+			}
+
+			const int32 SurroundRightChannel = 5;
+			for (int32 OutputChannel = 0; OutputChannel < NumOutputChannels; ++OutputChannel)
+			{
+				const int32 Index = OutputChannel * AUDIO_MIXER_MAX_OUTPUT_CHANNELS + SurroundRightChannel;
+				OutChannelMap.Add(Matrix[Index]);
+			}
+		}
 		else
 		{
 			// Compute a vorbis channel map only for 5.1 source files
