Fixing issue with submixes volumes

- Adding ability to change output volume of submixes

diff --git a/Engine/Source/Runtime/AudioMixer/Private/AudioMixerDevice.cpp b/Engine/Source/Runtime/AudioMixer/Private/AudioMixerDevice.cpp
--- a/Engine/Source/Runtime/AudioMixer/Private/AudioMixerDevice.cpp
+++ b/Engine/Source/Runtime/AudioMixer/Private/AudioMixerDevice.cpp
@@ -912,6 +912,50 @@ namespace Audio
 		SourceManager.UpdateSourceEffectChain(SourceEffectChainId, SourceEffectChain, bPlayEffectChainTails);
 	}
 
+	void FMixerDevice::UpdateSubmixProperties(USoundSubmix* InSoundSubmix)
+	{
+		check(InSoundSubmix);
+
+#if WITH_EDITOR
+		check(IsInAudioThread());
+
+		FMixerSubmixPtr* MixerSubmix = Submixes.Find(InSoundSubmix);
+		if (MixerSubmix)
+		{
+			float NewVolume = InSoundSubmix->OutputVolume;
+			AudioRenderThreadCommand([MixerSubmix, NewVolume]()
+			{
+				(*MixerSubmix)->SetOutputVolume(NewVolume);
+			});
+		}
+#endif // WITH_EDITOR
+	}
+
+	void FMixerDevice::SetSubmixOutputVolume(USoundSubmix* InSoundSubmix, float NewVolume)
+	{
+		if (!IsInAudioThread())
+		{
+			DECLARE_CYCLE_STAT(TEXT("FAudioThreadTask.SetSubmixOutputVolume"), STAT_AudioSetSubmixOutputVolume, STATGROUP_AudioThreadCommands);
+
+			FMixerDevice* MixerDevice = this;
+			FAudioThread::RunCommandOnAudioThread([MixerDevice, InSoundSubmix, NewVolume]()
+			{
+				CSV_SCOPED_TIMING_STAT(Audio, SetSubmixOutputVolume);
+				MixerDevice->SetSubmixOutputVolume(InSoundSubmix, NewVolume);
+			}, GET_STATID(STAT_AudioSetSubmixOutputVolume));
+			return;
+		}
+
+		FMixerSubmixPtr* MixerSubmix = Submixes.Find(InSoundSubmix);
+		if (MixerSubmix)
+		{
+			AudioRenderThreadCommand([MixerSubmix, NewVolume]()
+			{
+				(*MixerSubmix)->SetDynamicOutputVolume(NewVolume);
+			});
+		}
+	}
+
 	bool FMixerDevice::GetCurrentSourceEffectChain(const uint32 SourceEffectChainId, TArray<FSourceEffectChainEntry>& OutCurrentSourceEffectChainEntries)
 	{
 		TArray<FSourceEffectChainEntry>* ExistingOverride = SourceEffectChainOverrides.Find(SourceEffectChainId);
diff --git a/Engine/Source/Runtime/AudioMixer/Private/AudioMixerSubmix.cpp b/Engine/Source/Runtime/AudioMixer/Private/AudioMixerSubmix.cpp
--- a/Engine/Source/Runtime/AudioMixer/Private/AudioMixerSubmix.cpp
+++ b/Engine/Source/Runtime/AudioMixer/Private/AudioMixerSubmix.cpp
@@ -25,9 +25,13 @@ namespace Audio
 		, NumSamples(0)
 		, SubmixAmbisonicsEncoderID(INDEX_NONE)
 		, SubmixAmbisonicsDecoderID(INDEX_NONE)
+		, InitializedOutputVolume(1.0f)
+		, OutputVolume(1.0f)
+		, TargetOutputVolume(1.0f)
 		, EnvelopeNumChannels(0)
 		, bIsRecording(false)
 		, bIsBackgroundMuted(false)
+		, bApplyOutputVolumeScale(false)
 		, OwningSubmixObject(nullptr)
 	{
 	}
@@ -59,6 +63,14 @@ namespace Audio
 		{
 			OwningSubmixObject = InSoundSubmix;
 
+			// Set the initialized output volume
+			InitializedOutputVolume = FMath::Clamp(InSoundSubmix->OutputVolume, 0.0f, 1.0f);
+
+			if (!FMath::IsNearlyEqual(InitializedOutputVolume, 1.0f))
+			{
+				bApplyOutputVolumeScale = true;
+			}
+
 			// Loop through the submix's presets and make new instances of effects in the same order as the presets
 			ClearSoundEffectSubmixes();
 
@@ -769,6 +781,30 @@ namespace Audio
 			FMemory::Memcpy(EnvelopeValues, TempEnvelopeValues, sizeof(float)*AUDIO_MIXER_MAX_OUTPUT_CHANNELS);
 		}
 
+		// Don't necessarily need to do this if the user isn't using this feature
+		if (bApplyOutputVolumeScale)
+		{
+			const float TargetVolumeProduct = TargetOutputVolume * InitializedOutputVolume;
+			const float OutputVolumeProduct = OutputVolume * InitializedOutputVolume;
+
+			// If we've already set the volume, only need to multiply by constant
+			if (FMath::IsNearlyEqual(TargetVolumeProduct, OutputVolumeProduct))
+			{
+				Audio::MultiplyBufferByConstantInPlace(OutAudioBuffer, OutputVolumeProduct);
+			}
+			else
+			{
+				// To avoid popping, we do a fade on the buffer to the target volume
+				Audio::FadeBufferFast(OutAudioBuffer, OutputVolumeProduct, TargetVolumeProduct);
+				OutputVolume = TargetOutputVolume;
+
+				// No longer need to multiply the output buffer if we're now at 1.0
+				if (FMath::IsNearlyEqual(OutputVolume * InitializedOutputVolume, 1.0f))
+				{
+					bApplyOutputVolumeScale = false;
+				}
+			}
+		}
 
 		// Now loop through any buffer listeners and feed the listeners the result of this audio callback
 		{
@@ -952,9 +988,28 @@ namespace Audio
 		}
 	}
 
-	void FMixerSubmix::BroadcastEnvelope()
+	void FMixerSubmix::SetDynamicOutputVolume(float InVolume)
 	{
+		InVolume = FMath::Clamp(InVolume, 0.0f, 1.0f);
+		if (!FMath::IsNearlyEqual(InVolume, TargetOutputVolume))
+		{
+			TargetOutputVolume = InVolume;
+			bApplyOutputVolumeScale = true;
+		}
+	}
 
+	void FMixerSubmix::SetOutputVolume(float InVolume)
+	{
+		InVolume = FMath::Clamp(InVolume, 0.0f, 1.0f);
+		if (FMath::IsNearlyEqual(InitializedOutputVolume, InVolume))
+		{
+			InitializedOutputVolume = InVolume;
+			bApplyOutputVolumeScale = true;
+		}
+	}
+
+	void FMixerSubmix::BroadcastEnvelope()
+	{
 		if (bIsEnvelopeFollowing)
 		{
 			// Get the envelope data
diff --git a/Engine/Source/Runtime/AudioMixer/Public/AudioMixerDevice.h b/Engine/Source/Runtime/AudioMixer/Public/AudioMixerDevice.h
--- a/Engine/Source/Runtime/AudioMixer/Public/AudioMixerDevice.h
+++ b/Engine/Source/Runtime/AudioMixer/Public/AudioMixerDevice.h
@@ -108,7 +108,13 @@ namespace Audio
 		// Updates the source effect chain (using unique object id). 
 		virtual void UpdateSourceEffectChain(const uint32 SourceEffectChainId, const TArray<FSourceEffectChainEntry>& SourceEffectChain, const bool bPlayEffectChainTails) override;
 		virtual bool GetCurrentSourceEffectChain(const uint32 SourceEffectChainId, TArray<FSourceEffectChainEntry>& OutCurrentSourceEffectChainEntries) override;
+
+		// Updates submix instances with new properties
+		virtual void UpdateSubmixProperties(USoundSubmix* InSubmix) override;
 		
+		// Sets the submix output volume dynamically
+		virtual void SetSubmixOutputVolume(USoundSubmix* InSubmix, float NewVolume) override;
+
 		// Submix recording callbacks:
 		virtual void StartRecording(USoundSubmix* InSubmix, float ExpectedRecordingDuration) override;
 		virtual Audio::AlignedFloatBuffer& StopRecording(USoundSubmix* InSubmix, float& OutNumChannels, float& OutSampleRate) override;
diff --git a/Engine/Source/Runtime/AudioMixer/Public/AudioMixerSubmix.h b/Engine/Source/Runtime/AudioMixer/Public/AudioMixerSubmix.h
--- a/Engine/Source/Runtime/AudioMixer/Public/AudioMixerSubmix.h
+++ b/Engine/Source/Runtime/AudioMixer/Public/AudioMixerSubmix.h
@@ -60,6 +60,12 @@ namespace Audio
 		// Adds the given submix to this submix's children
 		void AddChildSubmix(TWeakPtr<FMixerSubmix, ESPMode::ThreadSafe> Submix);
 
+		// Sets the static output volume of the submix
+		void SetOutputVolume(float InVolume);
+
+		// Sets the dynamic output volume
+		void SetDynamicOutputVolume(float InVolume);
+
 		// Gets the submix channels channels
 		ESubmixChannelFormat GetSubmixChannels() const;
 
@@ -249,6 +255,15 @@ namespace Audio
 		// Decoder ID set up with Ambisonics Mixer. Set to INDEX_NONE if there is no decoder stream open.
 		uint32 SubmixAmbisonicsDecoderID;
 
+		// The output volume of the submix set via the USoundSubmix property. Can be set in the editor.
+		float InitializedOutputVolume;
+		
+		// The current dynamic output volume
+		float OutputVolume;
+		
+		// The target dynamic output volume
+		float TargetOutputVolume;
+
 		// Envelope following data
 		float EnvelopeValues[AUDIO_MIXER_MAX_OUTPUT_CHANNELS];
 		Audio::FEnvelopeFollower EnvelopeFollowers[AUDIO_MIXER_MAX_OUTPUT_CHANNELS];
@@ -285,6 +300,9 @@ namespace Audio
 		// Whether or not this submix is muted.
 		uint8 bIsBackgroundMuted : 1;
 
+		// Whether or not to apply a volume scale to output
+		uint8 bApplyOutputVolumeScale : 1;
+
 		// Bool set to true when envelope following is enabled
 		FThreadSafeBool bIsEnvelopeFollowing;
 
diff --git a/Engine/Source/Runtime/Engine/Classes/Sound/SoundSubmix.h b/Engine/Source/Runtime/Engine/Classes/Sound/SoundSubmix.h
--- a/Engine/Source/Runtime/Engine/Classes/Sound/SoundSubmix.h
+++ b/Engine/Source/Runtime/Engine/Classes/Sound/SoundSubmix.h
@@ -191,6 +191,10 @@ class ENGINE_API USoundSubmix : public UObject
 	UPROPERTY(EditAnywhere, BlueprintReadWrite, Category = EnvelopeFollower, meta = (ClampMin = "0", UIMin = "0"))
 	int32 EnvelopeFollowerReleaseTime;
 
+	/** The output volume of the submix. Applied after submix effects and analysis are performed. */
+	UPROPERTY(EditAnywhere, BlueprintReadWrite, Category = SoundSubmix, meta = (ClampMin = "0.0", ClampMax = "1.0", UIMin = "0.0", UIMax = "1.0"))
+	float OutputVolume;
+
 	// Blueprint delegate for when a recorded file is finished exporting.
 	UPROPERTY(BlueprintAssignable)
 	FOnSubmixRecordedFileDone OnSubmixRecordedFileDone;
@@ -222,6 +226,10 @@ class ENGINE_API USoundSubmix : public UObject
 	UFUNCTION(BlueprintCallable, Category = "Audio|EnvelopeFollowing", meta = (WorldContext = "WorldContextObject"))
 	void AddEnvelopeFollowerDelegate(const UObject* WorldContextObject, const FOnSubmixEnvelopeBP& OnSubmixEnvelopeBP);
 
+	/** Sets the output volume of the submix. This dynamic volume scales with the OutputVolume property of this submix. */
+	UFUNCTION(BlueprintCallable, Category = "Audio", meta = (WorldContext = "WorldContextObject"))
+	void SetSubmixOutputVolume(const UObject* WorldContextObject, float InOutputVolume);
+
 	// Registers and unregisters buffer listeners with the submix
 	void RegisterSubmixBufferListener(ISubmixBufferListener* InBufferListener);
 	void UnregisterSubmixBufferListener(ISubmixBufferListener* InBufferListener);
diff --git a/Engine/Source/Runtime/Engine/Private/AudioDeviceManager.cpp b/Engine/Source/Runtime/Engine/Private/AudioDeviceManager.cpp
--- a/Engine/Source/Runtime/Engine/Private/AudioDeviceManager.cpp
+++ b/Engine/Source/Runtime/Engine/Private/AudioDeviceManager.cpp
@@ -751,6 +751,17 @@ void FAudioDeviceManager::UpdateSourceEffectChain(const uint32 SourceEffectChain
 	}
 }
 
+void FAudioDeviceManager::UpdateSubmix(USoundSubmix* SoundSubmix)
+{
+	for (FAudioDevice* AudioDevice : Devices)
+	{
+		if (AudioDevice)
+		{
+			AudioDevice->UpdateSubmixProperties(SoundSubmix);
+		}
+	}
+}
+
 void FAudioDeviceManager::SetActiveDevice(uint32 InAudioDeviceHandle)
 {
 	// Only change the active device if there are no solo'd audio devices
diff --git a/Engine/Source/Runtime/Engine/Private/SoundSubmix.cpp b/Engine/Source/Runtime/Engine/Private/SoundSubmix.cpp
--- a/Engine/Source/Runtime/Engine/Private/SoundSubmix.cpp
+++ b/Engine/Source/Runtime/Engine/Private/SoundSubmix.cpp
@@ -22,6 +22,7 @@ USoundSubmix::USoundSubmix(const FObjectInitializer& ObjectInitializer)
 {
 	EnvelopeFollowerAttackTime = 10;
 	EnvelopeFollowerReleaseTime = 500;
+	OutputVolume = 1.0f;
 }
 
 void USoundSubmix::StartRecordingOutput(const UObject* WorldContextObject, float ExpectedDuration)
@@ -186,6 +187,21 @@ void USoundSubmix::AddEnvelopeFollowerDelegate(const UObject* WorldContextObject
 	}
 }
 
+void USoundSubmix::SetSubmixOutputVolume(const UObject* WorldContextObject, float InOutputVolume)
+{
+	if (!GEngine)
+	{
+		return;
+	}
+
+	UWorld* ThisWorld = GEngine->GetWorldFromContextObject(WorldContextObject, EGetWorldErrorMode::LogAndReturnNull);
+	FAudioDevice* AudioDevice = ThisWorld->GetAudioDevice();
+	if (AudioDevice)
+	{
+		AudioDevice->SetSubmixOutputVolume(this, InOutputVolume);
+	}
+}
+
 FString USoundSubmix::GetDesc()
 {
 	return FString(TEXT("Sound submix"));
@@ -240,6 +256,7 @@ void USoundSubmix::PostEditChangeProperty(struct FPropertyChangedEvent& Property
 	{
 		static const FName NAME_ChildSubmixes(TEXT("ChildSubmixes"));
 		static const FName NAME_ParentSubmix(TEXT("ParentSubmix"));
+		static const FName NAME_OutputVolume(TEXT("OutputVolume"));
 
 		if (PropertyChangedEvent.Property->GetFName() == NAME_ChildSubmixes)
 		{
@@ -306,6 +323,14 @@ void USoundSubmix::PostEditChangeProperty(struct FPropertyChangedEvent& Property
 			Modify();
 			RefreshAllGraphs(false);
 		}
+		else if (PropertyChangedEvent.Property->GetFName() == NAME_OutputVolume)
+		{
+			FAudioDeviceManager* AudioDeviceManager = (GEngine ? GEngine->GetAudioDeviceManager() : nullptr);
+			if (AudioDeviceManager)
+			{
+				AudioDeviceManager->UpdateSubmix(this);
+			}
+		}
 	}
 
 	// Use the main/default audio device for storing and retrieving sound class properties
diff --git a/Engine/Source/Runtime/Engine/Public/AudioDevice.h b/Engine/Source/Runtime/Engine/Public/AudioDevice.h
--- a/Engine/Source/Runtime/Engine/Public/AudioDevice.h
+++ b/Engine/Source/Runtime/Engine/Public/AudioDevice.h
@@ -1143,6 +1143,18 @@ public:
 	/** Returns the current source effect chain entries set dynamically from BP or elsewhere. */
 	virtual bool GetCurrentSourceEffectChain(const uint32 SourceEffectChainId, TArray<FSourceEffectChainEntry>& OutCurrentSourceEffectChainEntries) { return false; }
 
+	/** Updates the submix properties of any playing submix instances. Allows editor to make changes to submix properties and hear them propagate live.*/
+	virtual void UpdateSubmixProperties(USoundSubmix* InSubmix)
+	{
+		UE_LOG(LogAudio, Error, TEXT("Submixes are only supported in audio mixer."));
+	}
+
+	/** Sets the dynamic volume of the given submix. */
+	virtual void SetSubmixOutputVolume(USoundSubmix* InSubmix, float NewVolume)
+	{
+		UE_LOG(LogAudio, Error, TEXT("Submixes are only supported in audio mixer."));
+	}
+
 	/** This is called by a USoundSubmix to start recording a submix instance on this device. */
 	virtual void StartRecording(USoundSubmix* InSubmix, float ExpectedRecordingDuration)
 	{
diff --git a/Engine/Source/Runtime/Engine/Public/AudioDeviceManager.h b/Engine/Source/Runtime/Engine/Public/AudioDeviceManager.h
--- a/Engine/Source/Runtime/Engine/Public/AudioDeviceManager.h
+++ b/Engine/Source/Runtime/Engine/Public/AudioDeviceManager.h
@@ -142,6 +142,9 @@ public:
 	/** Updates source effect chain on all sources currently using the source effect chain. */
 	void UpdateSourceEffectChain(const uint32 SourceEffectChainId, const TArray<FSourceEffectChainEntry>& SourceEffectChain, const bool bPlayEffectChainTails);
 
+	/** Updates this submix for any changes made. Broadcasts to all submix instances. */
+	void UpdateSubmix(USoundSubmix* SoundSubmix);
+
 	/** Sets which audio device is the active audio device. */
 	void SetActiveDevice(uint32 InAudioDeviceHandle);
 
