Fix for submix envelope BP not working in packaged builds

diff --git a/Engine/Source/Runtime/AudioMixer/Private/AudioMixerDevice.cpp b/Engine/Source/Runtime/AudioMixer/Private/AudioMixerDevice.cpp
--- a/Engine/Source/Runtime/AudioMixer/Private/AudioMixerDevice.cpp
+++ b/Engine/Source/Runtime/AudioMixer/Private/AudioMixerDevice.cpp
@@ -1334,6 +1334,18 @@ namespace Audio
 
 	void FMixerDevice::StartRecording(USoundSubmix* InSubmix, float ExpectedRecordingDuration)
 	{
+		if (!IsInAudioThread())
+		{
+			DECLARE_CYCLE_STAT(TEXT("FAudioThreadTask.PauseRecording"), STAT_StartRecording, STATGROUP_AudioThreadCommands);
+
+			FAudioThread::RunCommandOnAudioThread([this, InSubmix, ExpectedRecordingDuration]()
+			{
+				CSV_SCOPED_TIMING_STAT(Audio, StartRecording);
+				StartRecording(InSubmix, ExpectedRecordingDuration);
+			}, GET_STATID(STAT_StartRecording));
+			return;
+		}
+
 		// if we can find the submix here, record that submix. Otherwise, just record the master submix.
 		Audio::FMixerSubmixPtr* FoundSubmix = Submixes.Find(InSubmix);
 		if (FoundSubmix)
@@ -1370,6 +1382,18 @@ namespace Audio
 
 	void FMixerDevice::PauseRecording(USoundSubmix* InSubmix)
 	{
+		if (!IsInAudioThread())
+		{
+			DECLARE_CYCLE_STAT(TEXT("FAudioThreadTask.PauseRecording"), STAT_PauseRecording, STATGROUP_AudioThreadCommands);
+
+			FAudioThread::RunCommandOnAudioThread([this, InSubmix]()
+			{
+				CSV_SCOPED_TIMING_STAT(Audio, PauseRecording);
+				PauseRecording(InSubmix);
+			}, GET_STATID(STAT_PauseRecording));
+			return;
+		}
+
 		// if we can find the submix here, pause that submix. Otherwise, just pause the master submix.
 		Audio::FMixerSubmixPtr* FoundSubmix = Submixes.Find(InSubmix);
 		if (FoundSubmix)
@@ -1388,6 +1412,18 @@ namespace Audio
 
 	void FMixerDevice::ResumeRecording(USoundSubmix* InSubmix)
 	{
+		if (!IsInAudioThread())
+		{
+			DECLARE_CYCLE_STAT(TEXT("FAudioThreadTask.ResumeRecording"), STAT_ResumeRecording, STATGROUP_AudioThreadCommands);
+
+			FAudioThread::RunCommandOnAudioThread([this, InSubmix]()
+			{
+				CSV_SCOPED_TIMING_STAT(Audio, ResumeRecording);
+				ResumeRecording(InSubmix);
+			}, GET_STATID(STAT_ResumeRecording));
+			return;
+		}
+
 		// if we can find the submix here, resume that submix. Otherwise, just resume the master submix.
 		Audio::FMixerSubmixPtr* FoundSubmix = Submixes.Find(InSubmix);
 		if (FoundSubmix)
@@ -1406,6 +1442,18 @@ namespace Audio
 
 	void FMixerDevice::StartEnvelopeFollowing(USoundSubmix* InSubmix)
 	{
+		if (!IsInAudioThread())
+		{
+			DECLARE_CYCLE_STAT(TEXT("FAudioThreadTask.StartEnvelopeFollowing"), STAT_StartEnvelopeFollowing, STATGROUP_AudioThreadCommands);
+
+			FAudioThread::RunCommandOnAudioThread([this, InSubmix]()
+			{
+				CSV_SCOPED_TIMING_STAT(Audio, StartEnvelopeFollowing);
+				StartEnvelopeFollowing(InSubmix);
+			}, GET_STATID(STAT_StartEnvelopeFollowing));
+			return;
+		}
+
 		// if we can find the submix here, record that submix. Otherwise, just record the master submix.
 		Audio::FMixerSubmixPtr* FoundSubmix = Submixes.Find(InSubmix);
 		if (FoundSubmix)
@@ -1426,6 +1474,18 @@ namespace Audio
 
 	void FMixerDevice::StopEnvelopeFollowing(USoundSubmix* InSubmix)
 	{
+		if (!IsInAudioThread())
+		{
+			DECLARE_CYCLE_STAT(TEXT("FAudioThreadTask.StopEnvelopeFollowing"), STAT_StopEnvelopeFollowing, STATGROUP_AudioThreadCommands);
+
+			FAudioThread::RunCommandOnAudioThread([this, InSubmix]()
+			{
+				CSV_SCOPED_TIMING_STAT(Audio, StopEnvelopeFollowing);
+				StopEnvelopeFollowing(InSubmix);
+			}, GET_STATID(STAT_StopEnvelopeFollowing));
+			return;
+		}
+
 		// if we can find the submix here, record that submix. Otherwise, just record the master submix.
 		Audio::FMixerSubmixPtr* FoundSubmix = Submixes.Find(InSubmix);
 		if (FoundSubmix)
@@ -1446,6 +1506,18 @@ namespace Audio
 
 	void FMixerDevice::AddEnvelopeFollowerDelegate(USoundSubmix* InSubmix, const FOnSubmixEnvelopeBP& OnSubmixEnvelopeBP)
 	{
+		if (!IsInAudioThread())
+		{
+			DECLARE_CYCLE_STAT(TEXT("FAudioThreadTask.AddEnvelopeFollowerDelegate"), STAT_AddEnvelopeFollowerDelegate, STATGROUP_AudioThreadCommands);
+
+			FAudioThread::RunCommandOnAudioThread([this, InSubmix, OnSubmixEnvelopeBP]()
+			{
+				CSV_SCOPED_TIMING_STAT(Audio, AddEnvelopeFollowerDelegate);
+				AddEnvelopeFollowerDelegate(InSubmix, OnSubmixEnvelopeBP);
+			}, GET_STATID(STAT_AddEnvelopeFollowerDelegate));
+			return;
+		}
+
 		// if we can find the submix here, record that submix. Otherwise, just record the master submix.
 		Audio::FMixerSubmixPtr* FoundSubmix = Submixes.Find(InSubmix);
 		if (FoundSubmix)
@@ -1465,6 +1537,18 @@ namespace Audio
 
 	void FMixerDevice::StartSpectrumAnalysis(USoundSubmix* InSubmix, const Audio::FSpectrumAnalyzerSettings& InSettings)
 	{
+		if (!IsInAudioThread())
+		{
+			DECLARE_CYCLE_STAT(TEXT("FAudioThreadTask.StartSpectrumAnalysis"), STAT_StartSpectrumAnalysis, STATGROUP_AudioThreadCommands);
+
+			FAudioThread::RunCommandOnAudioThread([this, InSubmix, InSettings]()
+			{
+				CSV_SCOPED_TIMING_STAT(Audio, StartSpectrumAnalysis);
+				StartSpectrumAnalysis(InSubmix, InSettings);
+			}, GET_STATID(STAT_StartSpectrumAnalysis));
+			return;
+		}
+
 		Audio::FMixerSubmixPtr* FoundSubmix = Submixes.Find(InSubmix);
 		if (FoundSubmix)
 		{
@@ -1482,6 +1566,18 @@ namespace Audio
 
 	void FMixerDevice::StopSpectrumAnalysis(USoundSubmix* InSubmix)
 	{
+		if (!IsInAudioThread())
+		{
+			DECLARE_CYCLE_STAT(TEXT("FAudioThreadTask.StopSpectrumAnalysis"), STAT_StopSpectrumAnalysis, STATGROUP_AudioThreadCommands);
+
+			FAudioThread::RunCommandOnAudioThread([this, InSubmix]()
+			{
+				CSV_SCOPED_TIMING_STAT(Audio, StopSpectrumAnalysis);
+				StopSpectrumAnalysis(InSubmix);
+			}, GET_STATID(STAT_StopSpectrumAnalysis));
+			return;
+		}
+
 		Audio::FMixerSubmixPtr* FoundSubmix = Submixes.Find(InSubmix);
 		if (FoundSubmix)
 		{
@@ -1533,6 +1629,18 @@ namespace Audio
 
 	void FMixerDevice::RegisterSubmixBufferListener(ISubmixBufferListener* InSubmixBufferListener, USoundSubmix* InSubmix)
 	{
+		if (!IsInAudioThread())
+		{
+			DECLARE_CYCLE_STAT(TEXT("FAudioThreadTask.RegisterSubmixBufferListener"), STAT_RegisterSubmixBufferListener, STATGROUP_AudioThreadCommands);
+
+			FAudioThread::RunCommandOnAudioThread([this, InSubmixBufferListener, InSubmix]()
+			{
+				CSV_SCOPED_TIMING_STAT(Audio, StopSpectrumAnalysis);
+				RegisterSubmixBufferListener(InSubmixBufferListener, InSubmix);
+			}, GET_STATID(STAT_RegisterSubmixBufferListener));
+			return;
+		}
+
 		Audio::FMixerSubmixPtr* FoundSubmix = Submixes.Find(InSubmix);
 		if (FoundSubmix)
 		{
@@ -1550,6 +1658,18 @@ namespace Audio
 
 	void FMixerDevice::UnregisterSubmixBufferListener(ISubmixBufferListener* InSubmixBufferListener, USoundSubmix* InSubmix)
 	{
+		if (!IsInAudioThread())
+		{
+			DECLARE_CYCLE_STAT(TEXT("FAudioThreadTask.UnregisterSubmixBufferListener"), STAT_UnregisterSubmixBufferListener, STATGROUP_AudioThreadCommands);
+
+			FAudioThread::RunCommandOnAudioThread([this, InSubmixBufferListener, InSubmix]()
+			{
+				CSV_SCOPED_TIMING_STAT(Audio, UnregisterSubmixBufferListener);
+				UnregisterSubmixBufferListener(InSubmixBufferListener, InSubmix);
+			}, GET_STATID(STAT_UnregisterSubmixBufferListener));
+			return;
+		}
+
 		Audio::FMixerSubmixPtr* FoundSubmix = Submixes.Find(InSubmix);
 		if (FoundSubmix)
 		{
