Adding audio math routines and sliding window utility for use in audio analyzer

diff --git a/Engine/Source/Runtime/AudioMixer/Private/DSP/FloatArrayMath.cpp b/Engine/Source/Runtime/AudioMixer/Private/DSP/FloatArrayMath.cpp
new file mode 100644
--- /dev/null
+++ b/Engine/Source/Runtime/AudioMixer/Private/DSP/FloatArrayMath.cpp
@@ -0,0 +1,205 @@
+// Copyright 1998-2019 Epic Games, Inc. All Rights Reserved.
+
+#include "DSP/FloatArrayMath.h"
+#include "CoreMinimal.h"
+
+namespace Audio
+{
+	void ArrayCumulativeSum(const TArray<float>& InData, TArray<float>& OutData)
+	{
+		// Initialize output data
+		int32 Num = InData.Num();
+		OutData.Reset();
+		OutData.AddUninitialized(Num);
+
+		if (Num < 1)
+		{
+			return;
+		}
+
+		float* OutDataPtr = OutData.GetData();
+		const float* InDataPtr = InData.GetData();
+
+		// Start summing
+		*OutDataPtr = *InDataPtr++;
+
+		for (int32 i = 1; i < Num; i++)
+		{
+			float Temp = *OutDataPtr++ + *InDataPtr++;
+			//OutDataPtr++;
+			*OutDataPtr = Temp;
+			//InDataPtr++;
+		}
+	}
+
+	void ArrayMeanFilter(const TArray<float>& InData, int32 WindowSize, int32 WindowOrigin, TArray<float>& OutData)
+	{
+		// a quick but sinful implementation of a mean filter. encourages floating point rounding errors. 
+		check(WindowOrigin < WindowSize);
+		check(WindowOrigin >= 0);
+		check(WindowSize > 0);
+
+		// Initialize output data
+		const int32 Num = InData.Num();
+		OutData.Reset();
+		OutData.AddUninitialized(Num);
+
+		if (Num < 1)
+		{
+			return;
+		}
+		
+		// Use cumulative sum to avoid multiple summations 
+		// Instead of summing over InData[StartIndex:EndIndex], avoid all that
+		// calculation by taking difference of cumulative sum at those two points:
+		//  cumsum(X[0:b]) - cumsum(X[0:a]) = sum(X[a:b])
+		TArray<float> SummedData;
+		ArrayCumulativeSum(InData, SummedData);
+		const float LastSummedData = SummedData.Last();
+		
+		
+		const int32 LastIndexBeforeEndBoundaryCondition = FMath::Max(WindowOrigin, Num - WindowSize + WindowOrigin + 1);
+		const int32 StartOffset = -WindowOrigin - 1;
+		const int32 EndOffset = WindowSize - WindowOrigin - 1;
+		const int32 WindowTail = WindowSize - WindowOrigin;
+
+		float* OutDataPtr = OutData.GetData();
+		const float* SummedDataPtr = SummedData.GetData();
+
+		if ((WindowSize - WindowOrigin) < Num)
+		{
+			// Handle boundary condition where analysis window precedes beginning of array.
+			for (int32 i = 0; i < (WindowOrigin + 1); i++)
+			{
+				OutDataPtr[i] = SummedDataPtr[i + EndOffset] / FMath::Max(1.f, static_cast<float>(WindowTail + i));
+			}
+
+			// No boundary conditions to handle here.	
+			const float MeanDivisor = static_cast<float>(WindowSize);
+			for (int32 i = WindowOrigin + 1; i < LastIndexBeforeEndBoundaryCondition; i++)
+			{
+				OutDataPtr[i] = (SummedDataPtr[i + EndOffset] - SummedDataPtr[i + StartOffset]) / MeanDivisor;
+			}
+		}
+		else
+		{
+			// Handle boundary condition where window precedes beginning and goes past end of array
+			const float ArrayMean = LastSummedData / static_cast<float>(Num);
+			for (int32 i = 0; i < LastIndexBeforeEndBoundaryCondition; i++)
+			{
+				OutDataPtr[i] = ArrayMean;
+			}
+		}
+
+		// Handle boundary condition where analysis window goes past end of array.
+		for (int32 i = LastIndexBeforeEndBoundaryCondition; i < Num; i++)
+		{
+			OutDataPtr[i] = (LastSummedData - SummedDataPtr[i + StartOffset]) / static_cast<float>(Num - i + WindowOrigin);
+		}
+	}
+
+	void ArrayMaxFilter(const TArray<float>& InData, int32 WindowSize, int32 WindowOrigin, TArray<float>& OutData)
+	{
+		// A reasonable implementation of a max filter for the data we're interested in, though surely not the fastest.
+		check(WindowOrigin < WindowSize);
+		check(WindowOrigin >= 0);
+		check(WindowSize > 0);
+		
+		int32 StartIndex = -WindowOrigin;
+		int32 EndIndex = StartIndex + WindowSize;
+
+		// Initialize output
+		int32 Num = InData.Num();
+		OutData.Reset();
+		OutData.AddUninitialized(Num);
+
+		if (Num < 1)
+		{
+			return;
+		}
+
+		// Get max in first window
+		int32 ActualStartIndex = 0;
+		int32 ActualEndIndex = FMath::Min(EndIndex, Num);
+
+		const float* InDataPtr = InData.GetData();
+		float* OutDataPtr = OutData.GetData();
+		int32 MaxIndex = 0;
+		float MaxValue = InData[0];
+
+		for (int32 i = ActualStartIndex; i < ActualEndIndex; i++)
+		{
+			if (InDataPtr[i] > MaxValue)
+			{
+				MaxValue = InDataPtr[i];
+				MaxIndex = i;
+			}		
+		}
+		OutDataPtr[0] = MaxValue;
+
+		StartIndex++;
+		EndIndex++;
+
+		// Get max in remaining windows
+		for (int32 i = 1; i < Num; i++)
+		{
+			ActualStartIndex = FMath::Max(StartIndex, 0);
+			ActualEndIndex = FMath::Min(EndIndex, Num);
+
+			if (MaxIndex < StartIndex)
+			{
+				// We need to evaluate the entire window because the previous maximum value was not in this window.
+				MaxIndex = ActualStartIndex;
+				MaxValue = InDataPtr[MaxIndex];
+				for (int32 j = ActualStartIndex + 1; j < ActualEndIndex; j++)
+				{
+					if (InDataPtr[j] > MaxValue)
+					{
+						MaxIndex = j;
+						MaxValue = InDataPtr[MaxIndex];
+					}
+				}
+			}
+			else
+			{
+				// We only need to inspect the newest sample because the previous maximum value was in this window.
+				if (InDataPtr[ActualEndIndex - 1] > MaxValue)
+				{
+					MaxIndex = ActualEndIndex - 1;
+					MaxValue = InDataPtr[MaxIndex];
+				}
+			}
+
+			OutDataPtr[i] = MaxValue;
+
+			StartIndex++;
+			EndIndex++;
+		}
+	}
+
+	void ArrayGetEuclideanNorm(const TArray<float>& InArray, float& OutEuclideanNorm)
+	{
+		// Initialize output.
+		OutEuclideanNorm = 0.0f;
+		const int32 Num = InArray.Num();
+		const float* InArrayData = InArray.GetData();
+		
+		// Sum it up.
+		for (int32 i = 0; i < Num; i++)
+		{
+			OutEuclideanNorm += InArrayData[i] * InArrayData[i];
+		}
+
+		OutEuclideanNorm = FMath::Sqrt(OutEuclideanNorm);
+	}
+
+	void ArrayMultiplyByConstantInPlace(TArray<float>& InArray, const float InMultiplier)
+	{
+		const int32 Num = InArray.Num();
+		float* InArrayData = InArray.GetData();
+		for (int32 i = 0; i < Num; i++)
+		{
+			InArrayData[i] *= InMultiplier;
+		}
+	}
+}
diff --git a/Engine/Source/Runtime/AudioMixer/Public/DSP/FloatArrayMath.h b/Engine/Source/Runtime/AudioMixer/Public/DSP/FloatArrayMath.h
new file mode 100644
--- /dev/null
+++ b/Engine/Source/Runtime/AudioMixer/Public/DSP/FloatArrayMath.h
@@ -0,0 +1,44 @@
+// Copyright 1998-2019 Epic Games, Inc. All Rights Reserved.
+
+#pragma once
+
+#include "CoreMinimal.h"
+
+namespace Audio
+{
+	/** Cumulative sum of array.
+	 *
+	 *  InData contains data to be cumulatively summed.
+	 *  OutData contains sum and is same size as InData.
+	 */
+	void ArrayCumulativeSum(const TArray<float>& InData, TArray<float>& OutData);
+
+	/** Mean filter of array.
+	 *
+	 *  Note: Uses standard biased mean estimator of Sum(x) / Count(x).
+	 *  Note: At array boundaries, this algorithm truncates windows where no valid array data exists. Values calculated with truncated windows have corresponding increased variances. 
+	 *
+	 *  InData contains data to be filtered.
+	 *  WindowSize determines the number of samples from InData analyzed to produce a value in OutData.
+	 *  WindowOrigin describes the offset from the windows first sample to the index of OutData. For example, if WindowOrigin = WindowSize/4, then OutData[i] = Mean(InData[i - Window/4 : i + 3 * Window / 4]).
+	 *  OutData contains the produceds data.
+	 */
+	void ArrayMeanFilter(const TArray<float>& InData, int32 WindowSize, int32 WindowOrigin, TArray<float>& OutData);
+
+	/** Max filter of array.
+	 *
+	 *  Note: At array boundaries, this algorithm truncates windows where no valid array data exists. 
+	 *
+	 *  InData contains data to be filtered.
+	 *  WindowSize determines the number of samples from InData analyzed to produce a value in OutData.
+	 *  WindowOrigin describes the offset from the windows first sample to the index of OutData. For example, if WindowOrigin = WindowSize/4, then OutData[i] = Max(InData[i - Window/4 : i + 3 * Window / 4]).
+	 *  OutData contains the produceds data.
+	 */
+	void ArrayMaxFilter(const TArray<float>& InData, int32 WindowSize, int32 WindowOrigin, TArray<float>& OutData);
+
+	/** Computes the EuclideanNorm of the InArray. Same as calculating the energy in window. */
+	void ArrayGetEuclideanNorm(const TArray<float>& InArray, float& OutEuclideanNorm);
+
+	/** Multiplies each element in InArray by InMultiplier */
+	void ArrayMultiplyByConstantInPlace(TArray<float>& InArray, const float InMultiplier);
+}
diff --git a/Engine/Source/Runtime/AudioMixer/Public/DSP/SlidingWindow.h b/Engine/Source/Runtime/AudioMixer/Public/DSP/SlidingWindow.h
new file mode 100644
--- /dev/null
+++ b/Engine/Source/Runtime/AudioMixer/Public/DSP/SlidingWindow.h
@@ -0,0 +1,410 @@
+// Copyright 1998-2019 Epic Games, Inc. All Rights Reserved.
+
+/** Sliding Window implementation which enables ranged for loop iteration over 
+ * sequential input buffers of varying length.
+ */
+#pragma once
+
+#include "CoreMinimal.h"
+
+
+namespace Audio
+{
+	// Forward delcaration
+	template <class SampleType>
+	class TSlidingWindow;
+
+
+	/** TSlidingBuffer
+	 *
+	 * TSlidingBuffer defines the window size and hop size of the sliding window, and
+	 * it stores any samples needed to produce additional windows.
+	 *
+	 * TSlidingBuffer should be used in conjunction with the TSlidingWindow, TScopedSlidingWindow 
+	 * or TAutoSlidingWindow classes.
+	 */
+	template <class SampleType>
+	class TSlidingBuffer
+	{
+		// Give TSlidingWindow access to StorageBuffer without exposing StorageBuffer public interface.
+		friend class TSlidingWindow<SampleType>;
+
+		public:
+			/** NumWindowSamples describes the number of samples in a window */
+			const int32 NumWindowSamples;
+
+			/** NumHopSamples describes the number of samples between adjacent windows */
+			const int32 NumHopSamples;
+
+			/**
+			 * Constructs a TSlidingBuffer with a constant window and hop size
+			 */
+			TSlidingBuffer(const int32 InNumWindowSamples, const int32 InNumHopSamples)
+			: NumWindowSamples(InNumWindowSamples)
+			, NumHopSamples(InNumHopSamples)
+			{
+				check(NumWindowSamples > 1);
+				check(NumHopSamples > 1);
+			}
+
+			/**
+			 * StoreForFutureWindows stores the necessary samples from InBuffer which will
+			 * be needed for future windows. It ignores all values in InBuffer which can
+			 * already be composed as a complete window.
+			 */
+			void StoreForFutureWindows(const TArrayView<const SampleType>& InBuffer)
+			{
+				int32 NumSamples = InBuffer.Num() + StorageBuffer.Num();
+
+				if (NumSamples < NumWindowSamples)
+				{
+					StorageBuffer.Append(InBuffer.GetData(), InBuffer.Num());
+				}
+				else
+				{
+					int32 NumWindowsGenerated = (NumSamples - NumWindowSamples) / NumHopSamples + 1;
+					int32 NumRemainingSamples = NumSamples - (NumWindowsGenerated * NumHopSamples);
+
+					if (NumRemainingSamples > InBuffer.Num())
+					{
+						int32 NumToRemove = NumRemainingSamples - InBuffer.Num();
+						StorageBuffer.RemoveAt(0, NumToRemove);
+						StorageBuffer.Append(InBuffer.GetData(), InBuffer.Num());
+					}
+					else
+					{
+						StorageBuffer.Reset(NumRemainingSamples);
+						StorageBuffer.AddUninitialized(NumRemainingSamples);
+						int32 NewBufferCopyIndex = InBuffer.Num() - NumRemainingSamples;
+						FMemory::Memcpy(StorageBuffer.GetData(), &InBuffer.GetData()[NewBufferCopyIndex], NumRemainingSamples * sizeof(SampleType));
+					}
+				}
+			}
+
+			/**
+			 * Resets the internal storage.
+			 */
+			void Reset()
+			{
+				StorageBuffer.Reset();
+			}
+
+		private:
+
+			TArray<SampleType> StorageBuffer;
+	};
+
+	/** TSlidingWindow
+	 *
+	 * TSlidingWindow allows windows of samples to be iterated over with STL like iterators. 
+	 *
+	 */
+	template <class SampleType>
+	class TSlidingWindow
+	{
+		friend class TSlidingWindowIterator;
+
+		protected:
+			// Accessed from friendship with TSlidingBuffer
+			const TArray<SampleType>& StorageBuffer;
+
+			// New buffer passed in.
+			const TArrayView<const SampleType>& NewBuffer;
+
+			// Copied from TSlidingBuffer
+			const int32 NumWindowSamples;
+
+			// Copied from TSlidingBuffer
+			const int32 NumHopSamples;
+
+			int32 NumZeroPad;
+
+		private:
+
+			int32 NumSamples;
+			int32 MaxReadIndex;
+
+		public:
+
+			/** TSlidingWindowIterator
+			 *
+			 * An forward iterator which slides a window over the given buffers.
+			 */
+			class TSlidingWindowIterator
+			{
+					const TSlidingWindow& SlidingWindow;
+
+					// Samples in window will be copied into this array.
+					TArray<SampleType>& WindowBuffer;
+
+					int32 ReadIndex;
+
+				public:
+
+					// Sentinel value marking that the last possible window has been generated.
+					static const int32 ReadIndexEnd = INDEX_NONE;
+
+					/**
+					 * Construct an iterator over a sliding window.
+					 */
+					TSlidingWindowIterator(const TSlidingWindow& InSlidingWindow, TArray<SampleType>& OutWindowBuffer, int32 InReadIndex)
+					:	SlidingWindow(InSlidingWindow)
+					,	WindowBuffer(OutWindowBuffer)
+					,	ReadIndex(InReadIndex)
+					{}
+
+					/**
+					 * Increment sliding window iterator forward.
+					 */
+					TSlidingWindowIterator operator++()
+					{
+						if (ReadIndex != ReadIndexEnd)
+						{
+							ReadIndex += SlidingWindow.NumHopSamples;
+							if (ReadIndex > SlidingWindow.MaxReadIndex)
+							{
+								ReadIndex = ReadIndexEnd;
+							}
+						}
+
+						return *this;
+					}
+
+					/**
+					 * Check whether iterators are equal. TSlidingWindowIterators derived from different
+					 * TSlidingWindows should not be compared.
+					 */
+					bool operator!=(const TSlidingWindowIterator& Other) const
+					{
+						return ReadIndex !=	Other.ReadIndex;
+					}
+
+					/**
+					 * Access array of windowed data currently pointed to by iterator.
+					 */
+					TArray<SampleType>& operator*()
+					{
+						if (ReadIndex != ReadIndexEnd)
+						{
+							// Resize output window
+							WindowBuffer.Reset(SlidingWindow.NumWindowSamples);
+							WindowBuffer.AddUninitialized(SlidingWindow.NumWindowSamples);
+
+							int32 SamplesFilled = 0;
+
+							if (ReadIndex < SlidingWindow.StorageBuffer.Num())
+							{
+								// The output window overlaps the storage buffer
+								int32 SamplesToCopy = SlidingWindow.StorageBuffer.Num() - ReadIndex;
+								FMemory::Memcpy(WindowBuffer.GetData(), SlidingWindow.StorageBuffer.GetData(), SamplesToCopy * sizeof(SampleType));
+								SamplesFilled += SamplesToCopy;
+							}
+
+							if (SamplesFilled < SlidingWindow.NumWindowSamples)
+							{
+								// The output window overlaps the new buffer.
+								int32 NewBufferIndex = ReadIndex - SlidingWindow.StorageBuffer.Num() + SamplesFilled;
+								int32 NewBufferRemaining = SlidingWindow.NewBuffer.Num() - NewBufferIndex;
+								int32 SamplesToCopy =FMath::Min(SlidingWindow.NumWindowSamples - SamplesFilled, NewBufferRemaining);
+
+								FMemory::Memcpy(&WindowBuffer.GetData()[SamplesFilled], &SlidingWindow.NewBuffer.GetData()[NewBufferIndex], SamplesToCopy * sizeof(SampleType));
+
+								SamplesFilled += SamplesToCopy;
+							}
+
+							if (SlidingWindow.NumZeroPad > 0 && SamplesFilled < SlidingWindow.NumWindowSamples)
+							{
+								// The output window should be padded with zeros because we are flusing this sliding window.
+								int32 SamplesToZeropad = FMath::Min(SlidingWindow.NumZeroPad, SlidingWindow.NumWindowSamples - SamplesFilled);
+								FMemory::Memset(&WindowBuffer.GetData()[SamplesFilled], 0, sizeof(SampleType) * SamplesToZeropad);
+							}
+						}
+						else
+						{
+							// Empty window if past end sliding window.
+							WindowBuffer.Reset();
+						}
+
+						return WindowBuffer;
+					}
+			};
+
+			/**
+			 * TSlidingWindow constructor
+			 *
+			 * InSlidingBuffer Holds the previous samples which were not completely used in previous sliding windows.  It also defines the window and hop size.
+			 * InNewBuffer Holds new samples which have not yet been ingested by the InSlidingBuffer.
+			 * bDoFlush Controls whether zeros to the final output windows until all possible windows with data from InNewBuffer have been covered.
+			 */
+			TSlidingWindow(const TSlidingBuffer<SampleType>& InSlidingBuffer, const TArrayView<const SampleType>& InNewBuffer, bool bDoFlush)
+			:	StorageBuffer(InSlidingBuffer.StorageBuffer)
+			,	NewBuffer(InNewBuffer)
+			,	NumWindowSamples(InSlidingBuffer.NumWindowSamples)
+			,	NumHopSamples(InSlidingBuffer.NumHopSamples)
+			,	NumZeroPad(0)
+			,	NumSamples(0)
+			,	MaxReadIndex(0)
+			{
+				// Total samples to be slid over.
+				NumSamples = NewBuffer.Num() + StorageBuffer.Num();
+
+				if (bDoFlush)
+				{
+					// If flushing, calculate the number of samples to zeropad
+					if (NumSamples < NumWindowSamples)
+					{
+						NumZeroPad = NumWindowSamples - NumSamples;
+					}
+					else
+					{
+						int32 NumWindowsGenerated = (NumSamples - NumWindowSamples) / NumHopSamples + 1;
+						NumZeroPad = NumWindowSamples - NumSamples + (NumWindowsGenerated * NumHopSamples);
+					}
+
+					NumSamples += NumZeroPad;
+				}
+
+				MaxReadIndex = NumSamples - NumWindowSamples;
+
+				if (MaxReadIndex < 0)
+				{
+					MaxReadIndex = TSlidingWindowIterator::ReadIndexEnd;
+				}
+			}
+
+			virtual ~TSlidingWindow()
+			{
+			}
+
+			/**
+			 * Creates STL like iterator which slides over samples.
+			 *
+			 * OutWindowBuffer Used to construct the TSlidingWindowIterator. The iterator will populate the window with samples when the * operator is called.
+			 */
+			TSlidingWindowIterator begin(TArray<SampleType>& OutWindowBuffer) const
+			{
+				return TSlidingWindowIterator(*this, OutWindowBuffer, 0);
+			}
+
+			/**
+			 * Creates STL like iterator denotes the end of the sliding window.
+			 *
+			 * OutWindowBuffer Used to construct the TSlidingWindowIterator. The iterator will populate the window with samples when the * operator is called.
+			 */
+			TSlidingWindowIterator end(TArray<SampleType>& OutWindowBuffer) const
+			{
+				return TSlidingWindowIterator(*this, OutWindowBuffer, TSlidingWindowIterator::ReadIndexEnd);
+			}
+	};
+
+	/** TScopedSlidingWindow
+	 * 
+	 * TScopedSlidingWindow provides a sliding window iterator interface over arrays. When TScopedSlidingWindow is destructed,
+	 * it calls StoreForFutureWindow(...) on the TSlidingBuffer passed into the constructor.
+	 *
+	 */
+	template <class SampleType>
+	class TScopedSlidingWindow : public TSlidingWindow<SampleType>
+	{
+			// Do not allow copying or moving since that may cause the destructor to be called inadvertantly.
+			TScopedSlidingWindow(TScopedSlidingWindow const &) = delete;
+			void operator=(TScopedSlidingWindow const &) = delete;
+			TScopedSlidingWindow(TScopedSlidingWindow&&) = delete;
+			TScopedSlidingWindow& operator=(TScopedSlidingWindow&&) = delete;
+
+			TSlidingBuffer<SampleType>& SlidingBuffer;
+		public:
+
+			/**
+			 * TScopedSlidingWindow constructor
+			 *
+			 * InSlidingBuffer Holds the previous samples which were not completely used in previous sliding windows.  It also defines the window and hop size.
+			 * InNewBuffer Holds new samples which have not yet been ingested by the InSlidingBuffer.
+			 * bDoFlush Controls whether zeros to the final output windows until all possible windows with data from InNewBuffer have been covered.
+			 */
+			TScopedSlidingWindow(TSlidingBuffer<SampleType>& InSlidingBuffer, const TArrayView<const SampleType>& InNewBuffer, bool bDoFlush = false)
+			:	TSlidingWindow(InSlidingBuffer, InNewBuffer, bDoFlush)
+			,	SlidingBuffer(InSlidingBuffer)
+			{}
+
+			/**
+			 * Calls InSlidingBuffer.StoreForFutureWindows(InNewBuffer).
+			 */
+			virtual ~TScopedSlidingWindow()
+			{
+				SlidingBuffer.StoreForFutureWindows(NewBuffer);
+			}
+	};
+
+	/** TAutoSlidingWindow
+	 * 
+	 * TAutoSlidingWindow enables use of a sliding window within a range-based for loop.
+	 *
+	 * Example:
+	 *
+	 * void ProcessAudio(TSlidingBuffer<float>& SlidingBuffer, const TArray<float>& NewSamples)
+	 * {
+	 * 		TArray<float> WindowData;
+	 * 		TAutoSlidingWindow<float> SlidingWindow(SlidingBuffer, NewSamples, WindowData);
+	 *
+	 * 		for (TArray<float>& Window : SlidingWindow)
+	 * 		{
+	 * 			... audio processing on single window here
+	 * 		}
+	 * }
+	 *
+	 * int main()
+	 * {
+	 * 		int32 NumWindowSamples = 4;
+	 * 		int32 NumHopSamples = 2;
+	 * 		TSlidingBuffer<float> SlidingBuffer(NumWindowSamples, NumHopSamples);
+	 *
+	 * 		TArray<float> Buffer1({1, 2, 3, 4, 5, 6, 7});
+	 *
+	 * 		ProcessAudio(SlidingBuffer, Buffer1);
+	 *
+	 * 		TArray<float> Buffer2({8, 9, 10, 11});
+	 *
+	 * 		ProcessAudio(SlidingBuffer, Buffer2);
+	 * }
+	 */
+	template <class SampleType>
+	class TAutoSlidingWindow : public TScopedSlidingWindow<SampleType>
+	{
+		TArray<SampleType>& WindowBuffer;
+
+		public:
+			/**
+			 * TAutoSlidingWindow constructor
+			 *
+			 * InSlidingBuffer Holds the previous samples which were not completely used in previous sliding windows.  It also defines the window and hop size.
+			 * InNewBuffer Holds new samples which have not yet been ingested by the InSlidingBuffer.
+			 * OutWindow is shared by all iterators created by calling begin() or end().
+			 * bDoFlush Controls whether zeros to the final output windows until all possible windows with data from InNewBuffer have been covered.
+			 */
+			TAutoSlidingWindow(TSlidingBuffer<SampleType>& InBuffer, const TArrayView<const SampleType>& InNewBuffer, TArray<SampleType>& OutWindow, bool bDoFlush = false)
+			:	TScopedSlidingWindow(InBuffer, InNewBuffer, bDoFlush)
+			,	WindowBuffer(OutWindow)
+			{}
+
+			/**
+			 * Creates STL like iterator which slides over samples.
+			 *
+			 * This iterator maintains a reference to the OutWindow passed into the constructor. That array will be manipulated when the iterator's * operator is called. 
+			 */
+			TSlidingWindow<SampleType>::TSlidingWindowIterator begin() 
+			{
+				return TSlidingWindow<SampleType>::begin(WindowBuffer);
+			}
+
+			/**
+			 * Creates STL like iterator denotes the end of the sliding window.
+			 *
+			 * This iterator maintains a reference to the OutWindow passed into the constructor. That array will be manipulated when the iterator's * operator is called. 
+			 */
+			TSlidingWindow<SampleType>::TSlidingWindowIterator end()
+			{
+				return TSlidingWindow<SampleType>::end(WindowBuffer);
+			}
+	};
+}
