Fixing issue with submix modifications in editor

diff --git a/Engine/Source/Runtime/AudioMixer/Private/AudioMixerSubmix.cpp b/Engine/Source/Runtime/AudioMixer/Private/AudioMixerSubmix.cpp
--- a/Engine/Source/Runtime/AudioMixer/Private/AudioMixerSubmix.cpp
+++ b/Engine/Source/Runtime/AudioMixer/Private/AudioMixerSubmix.cpp
@@ -59,89 +59,98 @@ namespace Audio
 
 	void FMixerSubmix::Init(USoundSubmix* InSoundSubmix)
 	{
+		check(IsInAudioThread());
 		if (InSoundSubmix != nullptr)
 		{
 			OwningSubmixObject = InSoundSubmix;
+	
+			SubmixCommand([this]()
+			{
+				InitInternal();
+			});
+		}
+	}
 
-			// Set the initialized output volume
-			InitializedOutputVolume = FMath::Clamp(InSoundSubmix->OutputVolume, 0.0f, 1.0f);
+	void FMixerSubmix::InitInternal()
+	{
+		// Set the initialized output volume
+		InitializedOutputVolume = FMath::Clamp(OwningSubmixObject->OutputVolume, 0.0f, 1.0f);
 
-			if (!FMath::IsNearlyEqual(InitializedOutputVolume, 1.0f))
-			{
-				bApplyOutputVolumeScale = true;
-			}
+		if (!FMath::IsNearlyEqual(InitializedOutputVolume, 1.0f))
+		{
+			bApplyOutputVolumeScale = true;
+		}
 
-			// Loop through the submix's presets and make new instances of effects in the same order as the presets
-			ClearSoundEffectSubmixes();
+		// Loop through the submix's presets and make new instances of effects in the same order as the presets
+		ClearSoundEffectSubmixes();
 
-			for (USoundEffectSubmixPreset* EffectPreset : InSoundSubmix->SubmixEffectChain)
+		for (USoundEffectSubmixPreset* EffectPreset : OwningSubmixObject->SubmixEffectChain)
+		{
+			if (EffectPreset)
 			{
-				if (EffectPreset)
-				{
-					// Create a new effect instance using the preset
-					FSoundEffectSubmix* SubmixEffect = static_cast<FSoundEffectSubmix*>(EffectPreset->CreateNewEffect());
+				// Create a new effect instance using the preset
+				FSoundEffectSubmix* SubmixEffect = static_cast<FSoundEffectSubmix*>(EffectPreset->CreateNewEffect());
 
-					FSoundEffectSubmixInitData InitData;
-					InitData.SampleRate = MixerDevice->GetSampleRate();
-					InitData.PresetSettings = nullptr;
+				FSoundEffectSubmixInitData InitData;
+				InitData.SampleRate = MixerDevice->GetSampleRate();
+				InitData.PresetSettings = nullptr;
 
-					// Now set the preset
-					SubmixEffect->Init(InitData);
-					SubmixEffect->SetPreset(EffectPreset);
-					SubmixEffect->SetEnabled(true);
+				// Now set the preset
+				SubmixEffect->Init(InitData);
+				SubmixEffect->SetPreset(EffectPreset);
+				SubmixEffect->SetEnabled(true);
 
-					FSubmixEffectInfo EffectInfo;
-					EffectInfo.PresetId = EffectPreset->GetUniqueID();
-					EffectInfo.EffectInstance = SubmixEffect;
+				FSubmixEffectInfo EffectInfo;
+				EffectInfo.PresetId = EffectPreset->GetUniqueID();
+				EffectInfo.EffectInstance = SubmixEffect;
 
-					// Add the effect to this submix's chain
-					EffectSubmixChain.Add(EffectInfo);
-				}
+				// Add the effect to this submix's chain
+				EffectSubmixChain.Add(EffectInfo);
 			}
+		}
 
-			ChannelFormat = InSoundSubmix->ChannelFormat;
-			
-			if (ChannelFormat == ESubmixChannelFormat::Ambisonics)
-			{
-				//Get the ambisonics mixer.
-				AmbisonicsMixer = MixerDevice->GetAmbisonicsMixer();
+		ChannelFormat = OwningSubmixObject->ChannelFormat;
+
+		if (ChannelFormat == ESubmixChannelFormat::Ambisonics)
+		{
+			//Get the ambisonics mixer.
+			AmbisonicsMixer = MixerDevice->GetAmbisonicsMixer();
 
-				//If we do have a valid ambisonics decoder, lets use it. Otherwise, treat this submix like a device submix.
-				if (AmbisonicsMixer.IsValid())
+			//If we do have a valid ambisonics decoder, lets use it. Otherwise, treat this submix like a device submix.
+			if (AmbisonicsMixer.IsValid())
+			{
+				if (!OwningSubmixObject->AmbisonicsPluginSettings)
 				{
-					if (!InSoundSubmix->AmbisonicsPluginSettings)
-					{
-						InSoundSubmix->AmbisonicsPluginSettings = AmbisonicsMixer->GetDefaultSettings();
-					}
+					OwningSubmixObject->AmbisonicsPluginSettings = AmbisonicsMixer->GetDefaultSettings();
+				}
 
-					if (InSoundSubmix->AmbisonicsPluginSettings != nullptr)
-					{
-						OnAmbisonicsSettingsChanged(InSoundSubmix->AmbisonicsPluginSettings);
-					}
-					else
-					{
-						//Default to first order ambisonics.
-						NumChannels = 4;
-						const int32 NumOutputFrames = MixerDevice->GetNumOutputFrames();
-						NumSamples = NumChannels * NumOutputFrames;
-					}
+				if (OwningSubmixObject->AmbisonicsPluginSettings != nullptr)
+				{
+					OnAmbisonicsSettingsChanged(OwningSubmixObject->AmbisonicsPluginSettings);
 				}
 				else
 				{
-					// There is no valid ambisonics decoder, so fall back to standard downmixing.
-					ChannelFormat = ESubmixChannelFormat::Device;
-					NumChannels = MixerDevice->GetNumChannelsForSubmixFormat(ChannelFormat);
+					//Default to first order ambisonics.
+					NumChannels = 4;
 					const int32 NumOutputFrames = MixerDevice->GetNumOutputFrames();
 					NumSamples = NumChannels * NumOutputFrames;
 				}
 			}
 			else
 			{
+				// There is no valid ambisonics decoder, so fall back to standard downmixing.
+				ChannelFormat = ESubmixChannelFormat::Device;
 				NumChannels = MixerDevice->GetNumChannelsForSubmixFormat(ChannelFormat);
 				const int32 NumOutputFrames = MixerDevice->GetNumOutputFrames();
 				NumSamples = NumChannels * NumOutputFrames;
 			}
 		}
+		else
+		{
+			NumChannels = MixerDevice->GetNumChannelsForSubmixFormat(ChannelFormat);
+			const int32 NumOutputFrames = MixerDevice->GetNumOutputFrames();
+			NumSamples = NumChannels * NumOutputFrames;
+		}
 	}
 
 	void FMixerSubmix::SetParentSubmix(TWeakPtr<FMixerSubmix, ESPMode::ThreadSafe> SubmixWeakPtr)
@@ -1001,7 +1010,7 @@ namespace Audio
 	void FMixerSubmix::SetOutputVolume(float InVolume)
 	{
 		InVolume = FMath::Clamp(InVolume, 0.0f, 1.0f);
-		if (FMath::IsNearlyEqual(InitializedOutputVolume, InVolume))
+		if (!FMath::IsNearlyEqual(InitializedOutputVolume, InVolume))
 		{
 			InitializedOutputVolume = InVolume;
 			bApplyOutputVolumeScale = true;
diff --git a/Engine/Source/Runtime/AudioMixer/Public/AudioMixerSubmix.h b/Engine/Source/Runtime/AudioMixer/Public/AudioMixerSubmix.h
--- a/Engine/Source/Runtime/AudioMixer/Public/AudioMixerSubmix.h
+++ b/Engine/Source/Runtime/AudioMixer/Public/AudioMixerSubmix.h
@@ -159,6 +159,9 @@ namespace Audio
 		void BroadcastEnvelope();
 
 	protected:
+		// Initialize the submix internal 
+		void InitInternal();
+
 		// Down mix the given buffer to the desired down mix channel count
 		void FormatChangeBuffer(const ESubmixChannelFormat NewChannelType, AlignedFloatBuffer& InBuffer, AlignedFloatBuffer& OutNewBuffer);
 
diff --git a/Engine/Source/Runtime/Engine/Private/SoundSubmix.cpp b/Engine/Source/Runtime/Engine/Private/SoundSubmix.cpp
--- a/Engine/Source/Runtime/Engine/Private/SoundSubmix.cpp
+++ b/Engine/Source/Runtime/Engine/Private/SoundSubmix.cpp
@@ -252,6 +252,9 @@ void USoundSubmix::PreEditChange(UProperty* PropertyAboutToChange)
 
 void USoundSubmix::PostEditChangeProperty(struct FPropertyChangedEvent& PropertyChangedEvent)
 {
+	// Whether or not we need to reinit the submix. Not all properties require reinitialization.
+	bool bReinitSubmix = true;
+
 	if (PropertyChangedEvent.Property != nullptr)
 	{
 		static const FName NAME_ChildSubmixes(TEXT("ChildSubmixes"));
@@ -330,16 +333,20 @@ void USoundSubmix::PostEditChangeProperty(struct FPropertyChangedEvent& Property
 			{
 				AudioDeviceManager->UpdateSubmix(this);
 			}
+			bReinitSubmix = false;
 		}
 	}
 
-	// Use the main/default audio device for storing and retrieving sound class properties
-	FAudioDeviceManager* AudioDeviceManager = (GEngine ? GEngine->GetAudioDeviceManager() : nullptr);
-
-	// Force the properties to be initialized for this SoundSubmix on all active audio devices
-	if (AudioDeviceManager)
+	if (bReinitSubmix)
 	{
-		AudioDeviceManager->RegisterSoundSubmix(this);
+		// Use the main/default audio device for storing and retrieving sound class properties
+		FAudioDeviceManager* AudioDeviceManager = (GEngine ? GEngine->GetAudioDeviceManager() : nullptr);
+
+		// Force the properties to be initialized for this SoundSubmix on all active audio devices
+		if (AudioDeviceManager)
+		{
+			AudioDeviceManager->RegisterSoundSubmix(this);
+		}
 	}
 
 	Super::PostEditChangeProperty(PropertyChangedEvent);
