fix to handle zero number of samples in reverb callback

diff --git a/Engine/Source/Runtime/AudioMixer/Private/DSP/ReverbFast.cpp b/Engine/Source/Runtime/AudioMixer/Private/DSP/ReverbFast.cpp
--- a/Engine/Source/Runtime/AudioMixer/Private/DSP/ReverbFast.cpp
+++ b/Engine/Source/Runtime/AudioMixer/Private/DSP/ReverbFast.cpp
@@ -106,6 +106,13 @@ namespace Audio {
 		const int32 InNum = InSamples.Num();
 		const int32 InNumFrames = InNum / InNumChannels;
 
+		if (0 == InNum)
+		{
+			// If not input samples, then reset output samples and return
+			OutSamples.Reset(0);
+			return;
+		}
+
 		if (!bEnableEarlyReflections && !bEnableLateReflections)
 		{
 			// Zero output buffers if all reverb is disabled. 
