Fix for AudioMixerDevice using incorrectly reporting audio buffer size it is going to use

diff --git a/Engine/Source/Runtime/AudioMixer/Private/AudioMixerDevice.cpp b/Engine/Source/Runtime/AudioMixer/Private/AudioMixerDevice.cpp
--- a/Engine/Source/Runtime/AudioMixer/Private/AudioMixerDevice.cpp
+++ b/Engine/Source/Runtime/AudioMixer/Private/AudioMixerDevice.cpp
@@ -840,7 +840,7 @@ namespace Audio
 		UE_LOG(LogAudioMixer, Display, TEXT("Audio Mixer Platform Settings:"));
 		UE_LOG(LogAudioMixer, Display, TEXT("	Sample Rate:						  %d"), Settings.SampleRate);
 		UE_LOG(LogAudioMixer, Display, TEXT("	Callback Buffer Frame Size Requested: %d"), Settings.CallbackBufferFrameSize);
-		UE_LOG(LogAudioMixer, Display, TEXT("	Callback Buffer Frame Size To Use:	  %d"), AudioMixerPlatform->GetNumFrames(PlatformSettings.CallbackBufferFrameSize));
+		UE_LOG(LogAudioMixer, Display, TEXT("	Callback Buffer Frame Size To Use:	  %d"), AudioMixerPlatform->GetNumFrames(Settings.CallbackBufferFrameSize));
 		UE_LOG(LogAudioMixer, Display, TEXT("	Number of buffers to queue:			  %d"), Settings.NumBuffers);
 		UE_LOG(LogAudioMixer, Display, TEXT("	Max Channels (voices):				  %d"), Settings.MaxChannels);
 		UE_LOG(LogAudioMixer, Display, TEXT("	Number of Async Source Workers:		  %d"), Settings.NumSourceWorkers);
