Added cvar to force sync audio decodes (can be enabled for low-end/low-core devices)

diff --git a/Engine/Source/Runtime/AudioMixer/Private/AudioMixerSourceDecode.cpp b/Engine/Source/Runtime/AudioMixer/Private/AudioMixerSourceDecode.cpp
--- a/Engine/Source/Runtime/AudioMixer/Private/AudioMixerSourceDecode.cpp
+++ b/Engine/Source/Runtime/AudioMixer/Private/AudioMixerSourceDecode.cpp
@@ -10,6 +10,13 @@
 #include "Async/Async.h"
 #include "AudioDecompress.h"
 
+static int32 ForceSyncAudioDecodesCvar = 1;
+FAutoConsoleVariableRef CVarForceSyncAudioDecodes(
+                                                  TEXT("au.ForceSyncAudioDecodes"),
+                                                  ForceSyncAudioDecodesCvar,
+                                                  TEXT("Disables using async tasks for processing sources.\n")
+                                                  TEXT("0: Not Disabled, 1: Disabled"),
+                                                  ECVF_Default);
 namespace Audio
 {
 
@@ -220,6 +227,12 @@ public:
 	FHeaderDecodeHandle(const FHeaderParseAudioTaskData& InJobData)
 	{
 		Task = new FAsyncTask<FAsyncDecodeWorker>(InJobData);
+        if (ForceSyncAudioDecodesCvar)
+        {
+            Task->StartSynchronousTask();
+            return;
+        }
+        
 		Task->StartBackgroundTask();
 	}
 
@@ -235,6 +248,12 @@ public:
 	FProceduralDecodeHandle(const FProceduralAudioTaskData& InJobData)
 	{
 		Task = new FAsyncTask<FAsyncDecodeWorker>(InJobData);
+        if (ForceSyncAudioDecodesCvar)
+        {
+            Task->StartSynchronousTask();
+            return;
+        }
+        
 		Task->StartBackgroundTask();
 	}
 
@@ -257,6 +276,12 @@ public:
 	FDecodeHandle(const FDecodeAudioTaskData& InJobData)
 	{
 		Task = new FAsyncTask<FAsyncDecodeWorker>(InJobData);
+        if (ForceSyncAudioDecodesCvar)
+        {
+            Task->StartSynchronousTask();
+            return;
+        }
+        
 		const bool bUseBackground = ShouldUseBackgroundPoolFor_FAsyncRealtimeAudioTask();
 		Task->StartBackgroundTask(bUseBackground ? GBackgroundPriorityThreadPool : GThreadPool);
 	}
