Setting default vaule of ForceSyncAudioDecodesCvar to 0 instead of 1

diff --git a/Engine/Source/Runtime/AudioMixer/Private/AudioMixerSourceDecode.cpp b/Engine/Source/Runtime/AudioMixer/Private/AudioMixerSourceDecode.cpp
--- a/Engine/Source/Runtime/AudioMixer/Private/AudioMixerSourceDecode.cpp
+++ b/Engine/Source/Runtime/AudioMixer/Private/AudioMixerSourceDecode.cpp
@@ -10,7 +10,7 @@
 #include "Async/Async.h"
 #include "AudioDecompress.h"
 
-static int32 ForceSyncAudioDecodesCvar = 1;
+static int32 ForceSyncAudioDecodesCvar = 0;
 FAutoConsoleVariableRef CVarForceSyncAudioDecodes(
                                                   TEXT("au.ForceSyncAudioDecodes"),
                                                   ForceSyncAudioDecodesCvar,
