Fix for default 2D reverb and spatialization plugins not spatializing

diff --git a/Engine/Source/Runtime/AudioMixer/Private/AudioMixerSubmix.cpp b/Engine/Source/Runtime/AudioMixer/Private/AudioMixerSubmix.cpp
--- a/Engine/Source/Runtime/AudioMixer/Private/AudioMixerSubmix.cpp
+++ b/Engine/Source/Runtime/AudioMixer/Private/AudioMixerSubmix.cpp
@@ -62,12 +62,21 @@ namespace Audio
 		check(IsInAudioThread());
 		if (InSoundSubmix != nullptr)
 		{
-			OwningSubmixObject = InSoundSubmix;
-	
-			SubmixCommand([this]()
+			// This is a first init and needs to be synchronous
+			if (!OwningSubmixObject)
 			{
+				OwningSubmixObject = InSoundSubmix;
 				InitInternal();
-			});
+			}
+			else
+			{
+				// This is a re-init and needs to be thread safe
+				check(OwningSubmixObject == InSoundSubmix);
+				SubmixCommand([this]()
+				{
+					InitInternal();
+				});
+			}
 		}
 	}
 
