make the audio mixer command queue flush timeout tweakable + adjust it for some platforms.

diff --git a/Engine/Source/Runtime/AudioMixer/Private/AudioMixerSourceManager.cpp b/Engine/Source/Runtime/AudioMixer/Private/AudioMixerSourceManager.cpp
--- a/Engine/Source/Runtime/AudioMixer/Private/AudioMixerSourceManager.cpp
+++ b/Engine/Source/Runtime/AudioMixer/Private/AudioMixerSourceManager.cpp
@@ -75,6 +75,14 @@ FAutoConsoleVariableRef CVarFlushCommandBufferOnTimeout(
 	TEXT("0: Not Disabled, 1: Disabled"),
 	ECVF_Default);
 
+static int32 CommandBufferFlushWaitTimeMsCvar = 1000;
+FAutoConsoleVariableRef CVarCommandBufferFlushWaitTimeMs(
+	TEXT("au.CommandBufferFlushWaitTimeMs"),
+	CommandBufferFlushWaitTimeMsCvar,
+	TEXT("How long to wait for the command buffer flush to complete.\n"),
+	ECVF_Default);
+
+
 #define ONEOVERSHORTMAX (3.0517578125e-5f) // 1/32768
 #define ENVELOPE_TAIL_THRESHOLD (1.58489e-5f) // -96 dB
 
@@ -2534,7 +2542,7 @@ namespace Audio
 
 		// Make sure current current executing 
 		bool bTimedOut = false;
-		if (!CommandsProcessedEvent->Wait(1000))
+		if (!CommandsProcessedEvent->Wait(CommandBufferFlushWaitTimeMsCvar))
 		{
 			CommandsProcessedEvent->Trigger();
 			bTimedOut = true;
