Adding Mel and constant Q Transforms

diff --git a/Engine/Source/Runtime/AudioMixer/Private/DSP/ConstantQ.cpp b/Engine/Source/Runtime/AudioMixer/Private/DSP/ConstantQ.cpp
new file mode 100644
--- /dev/null
+++ b/Engine/Source/Runtime/AudioMixer/Private/DSP/ConstantQ.cpp
@@ -0,0 +1,163 @@
+// Copyright 1998-2019 Epic Games, Inc. All Rights Reserved.
+
+#include "DSP/ConstantQ.h"
+#include "CoreMinimal.h"
+#include "DSP/FloatArrayMath.h"
+
+namespace
+{
+	float GetConstantQCenterFrequency(const int32 InBandIndex, const float InBaseFrequency, const float InBandsPerOctave)
+	{
+		check(InBandsPerOctave > 0.f);
+		return InBaseFrequency * FMath::Pow(2.f, static_cast<float>(InBandIndex) / InBandsPerOctave);
+	}
+
+	float GetConstantQBandWidth(const float InBandCenter, const float InBandsPerOctave, const float InBandWidthStretch)
+	{
+		check(InBandsPerOctave > 0.f);
+		return InBandWidthStretch * InBandCenter * (FMath::Pow(2.f, 1.f / InBandsPerOctave) - 1.f);
+	}
+
+	void FillArrayWithTruncatedGaussian(const float InCenterFreq, const float InBandWidth, const int32 InFFTSize, const float InSampleRate, TArray<float>& OutOffsetArray, int32& OutOffsetIndex)
+	{
+		check(InBandWidth > 0.f);
+		check(InFFTSize > 0);
+		check(InCenterFreq >= 0.f);
+		check(InSampleRate > 0.f);
+
+		// Determine points where gaussian will value will go below a small number
+		const float SignificantHalfBandWidth = InBandWidth * FMath::Sqrt(-2.f * FMath::Loge(SMALL_NUMBER));
+		const float LowestSignificantFreq = FMath::Clamp(InCenterFreq - SignificantHalfBandWidth, 0.f, InSampleRate / 2.f);
+		const float HighestSignificantFreq = FMath::Clamp(InCenterFreq + SignificantHalfBandWidth, 0.f, InSampleRate / 2.f);
+		int32 LowestSignificantIndex = FMath::CeilToInt(InFFTSize * LowestSignificantFreq / InSampleRate);
+		int32 HighestSignificantIndex = FMath::FloorToInt(InFFTSize * HighestSignificantFreq / InSampleRate);
+ 		int32 Num = HighestSignificantIndex - LowestSignificantIndex + 1;
+
+		if (Num < 1)
+		{
+			Num = 1;
+		}
+
+		// Prepare outputs
+		OutOffsetIndex = LowestSignificantIndex;
+		OutOffsetArray.Reset(Num);
+		OutOffsetArray.AddUninitialized(Num);
+
+		// Fill array with truncated Gaussian
+		const float BandWidthSquared = InBandWidth * InBandWidth;
+		for (int32 i = 0; i < Num; i++)
+		{
+			float FFTBinHz = static_cast<float>(i + LowestSignificantIndex) * InSampleRate / InFFTSize;
+			float DeltaHz = FFTBinHz - InCenterFreq;
+			OutOffsetArray[i] = FMath::Exp(-0.5 * (DeltaHz * DeltaHz) / BandWidthSquared);
+		}
+	}
+}
+
+namespace Audio
+{
+	TUniquePtr<FContiguousSparse2DKernelTransform> NewPseudoConstantQKernelTransform(const FPseudoConstantQKernelSettings& InSettings, const int32 InFFTSize, const float InSampleRate)
+	{
+		check(InSampleRate > 0.f);
+		check(InFFTSize > 0);
+
+		const int32 NumUsefulFFTBins = InFFTSize / 2 + 1;
+
+		TUniquePtr<FContiguousSparse2DKernelTransform> Transform = MakeUnique<FContiguousSparse2DKernelTransform>(NumUsefulFFTBins, InSettings.NumBands);
+
+		for (int32 CQTBandIndex = 0; CQTBandIndex < InSettings.NumBands; CQTBandIndex++)
+		{
+			// Determine band center and width for this CQT band
+			const float BandCenter = GetConstantQCenterFrequency(CQTBandIndex, InSettings.KernelLowestCenterFreq, InSettings.NumBandsPerOctave);
+			const float BandWidth = GetConstantQBandWidth(BandCenter, InSettings.NumBandsPerOctave, InSettings.BandWidthStretch);
+
+			if ((BandCenter - BandWidth) > InSampleRate)
+			{
+				// TODO: have this check be fore a very large number
+				continue;
+			}
+
+			// Create gaussian centered around center freq and with appropriate bandwidth
+			TArray<float> OffsetBandWeights;
+			int32 OffsetBandWeightsIndex = 0;
+			FillArrayWithTruncatedGaussian(BandCenter, BandWidth, InFFTSize, InSampleRate, OffsetBandWeights, OffsetBandWeightsIndex);
+
+
+			// set infs and nans to zero
+			for (int32 i = 0; i < OffsetBandWeights.Num(); i++)
+			{
+				if (!FMath::IsFinite(OffsetBandWeights[i]))
+				{
+					OffsetBandWeights[i] = 0.f;
+				}
+			}
+
+			// Need to have a lower bound for 
+			float DigitalBandWidth = FMath::Max(1.f / static_cast<float>(InFFTSize), BandWidth / InSampleRate);
+
+			// Sanity check the CQT bins to make sure we the bandwidth wasn't so small that the array is essentially empty
+			if (OffsetBandWeights.Num() > 0)
+			{
+				if (OffsetBandWeights.FilterByPredicate([](float InVal) { return InVal >= 0.5f; }).Num() < 1)
+				{
+					// Hit a condition where all values in band weights are below 0.5f. 
+					// It's a bit of an arbitrary threshold, but it tells us that the bandwidth
+					// is low enough and the FFT granularity course enough to where our pseudo 
+					// cqt windows will likely miss data. In this case we force the window to 
+					// have a single value of 1 at the nearest FFT bin.
+					FMemory::Memset(OffsetBandWeights.GetData(), 0, OffsetBandWeights.Num() * sizeof(float));
+					int32 NearestIndex = FMath::RoundToInt(InFFTSize * BandCenter / InSampleRate) - OffsetBandWeightsIndex;
+
+					if (NearestIndex < OffsetBandWeights.Num())
+					{
+						NearestIndex = FMath::Clamp(NearestIndex, 0, OffsetBandWeights.Num() - 1);
+						OffsetBandWeights[NearestIndex] = 1.f / static_cast<float>(InFFTSize);
+
+						// Need to update digital bandwidth to be bandwidth of FFTbin
+						DigitalBandWidth = 1.f / static_cast<float>(InFFTSize);
+					}
+				}
+			}
+
+
+
+
+			// Normalize window
+			float NormDenom = 1.f;
+			switch (InSettings.Normalization)
+			{
+				case EPseudoConstantQNormalization::EqualAmplitude:
+					NormDenom = 1.f;
+					break;
+				case EPseudoConstantQNormalization::EqualEuclideanNorm:
+					NormDenom = FMath::Sqrt(DigitalBandWidth * InFFTSize * FMath::Sqrt(PI  * 2.f));
+					break;
+				case EPseudoConstantQNormalization::EqualEnergy:
+				default:
+					NormDenom = DigitalBandWidth * InFFTSize * FMath::Sqrt(PI  * 2.f);
+					break;
+			}
+
+			if ((NormDenom > 0.f) && (NormDenom != 1.f))
+			{
+				ArrayMultiplyByConstantInPlace(OffsetBandWeights, 1.f / NormDenom);
+			}
+
+			// Truncate if necessary
+			if (OffsetBandWeightsIndex >= NumUsefulFFTBins)
+			{
+				OffsetBandWeightsIndex = 0;
+				OffsetBandWeights.Empty();
+			}
+			else if ((OffsetBandWeightsIndex + OffsetBandWeights.Num()) > NumUsefulFFTBins)
+			{
+				OffsetBandWeights.SetNum(NumUsefulFFTBins - OffsetBandWeightsIndex);
+			}
+
+			// Store row in transform
+			Transform->SetRow(CQTBandIndex, OffsetBandWeightsIndex, OffsetBandWeights);
+		}
+
+		return Transform;
+	}
+}
diff --git a/Engine/Source/Runtime/AudioMixer/Private/DSP/FloatArrayMath.cpp b/Engine/Source/Runtime/AudioMixer/Private/DSP/FloatArrayMath.cpp
--- a/Engine/Source/Runtime/AudioMixer/Private/DSP/FloatArrayMath.cpp
+++ b/Engine/Source/Runtime/AudioMixer/Private/DSP/FloatArrayMath.cpp
@@ -58,7 +58,7 @@ namespace Audio
 		const float LastSummedData = SummedData.Last();
 		
 		
-		const int32 LastIndexBeforeEndBoundaryCondition = FMath::Max(WindowOrigin, Num - WindowSize + WindowOrigin + 1);
+		const int32 LastIndexBeforeEndBoundaryCondition = FMath::Max(WindowOrigin + 1, Num - WindowSize + WindowOrigin + 1);
 		const int32 StartOffset = -WindowOrigin - 1;
 		const int32 EndOffset = WindowSize - WindowOrigin - 1;
 		const int32 WindowTail = WindowSize - WindowOrigin;
@@ -202,4 +202,93 @@ namespace Audio
 			InArrayData[i] *= InMultiplier;
 		}
 	}
+
+	FContiguousSparse2DKernelTransform::FContiguousSparse2DKernelTransform(const int32 NumInElements, const int32 NumOutElements)
+	:	NumIn(NumInElements)
+	,	NumOut(NumOutElements)
+	{
+		check(NumIn >= 0);
+		check(NumOut >= 0)
+		FRow EmptyRow;
+		EmptyRow.StartIndex = 0;
+
+		// Fill up the kernel with emptp rows
+		Kernel.Init(EmptyRow, NumOut);
+	}
+
+	int32 FContiguousSparse2DKernelTransform::GetNumInElements() const
+	{
+		return NumIn;
+	}
+
+
+	int32 FContiguousSparse2DKernelTransform::GetNumOutElements() const
+	{
+		return NumOut;
+	}
+
+	void FContiguousSparse2DKernelTransform::SetRow(const int32 RowIndex, const int32 StartIndex, TArrayView<const float> OffsetValues)
+	{
+		check((StartIndex + OffsetValues.Num()) <= NumIn);
+
+		// Copy row data internally
+		Kernel[RowIndex].StartIndex = StartIndex;
+		Kernel[RowIndex].OffsetValues = TArray<float>(OffsetValues.GetData(), OffsetValues.Num());
+	}
+
+	void FContiguousSparse2DKernelTransform::TransformArray(TArrayView<const float> InArray, TArray<float>& OutArray) const
+	{
+		check(InArray.Num() == NumIn);
+
+		// Resize output
+		OutArray.Reset(NumOut);
+		if (NumOut > 0)
+		{
+			OutArray.AddUninitialized(NumOut);
+		}
+
+		TransformArray(InArray.GetData(), OutArray.GetData());
+	}
+
+	void FContiguousSparse2DKernelTransform::TransformArray(TArrayView<const float> InArray, AlignedFloatBuffer& OutArray) const
+	{	
+		check(InArray.Num() == NumIn);
+
+		// Resize output
+		OutArray.Reset(NumOut);
+		if (NumOut > 0)
+		{
+			OutArray.AddUninitialized(NumOut);
+		}
+
+		TransformArray(InArray.GetData(), OutArray.GetData());
+	}
+
+	void FContiguousSparse2DKernelTransform::TransformArray(const float* InArray, float* OutArray) const
+	{
+		check(nullptr != InArray);
+		check(nullptr != OutArray);
+
+		// Initialize output
+		FMemory::Memset(OutArray, 0, sizeof(float) * NumOut);
+
+		// Apply kernel one row at a time
+		const FRow* KernelData = Kernel.GetData();
+		for (int32 RowIndex = 0; RowIndex < Kernel.Num(); RowIndex++)
+		{
+			const FRow& Row = KernelData[RowIndex];
+
+			// Get offset pointer into input array.
+			const float* OffsetInData = &InArray[Row.StartIndex];
+			// Get offset pointer of row.
+			const float* RowValuePtr = Row.OffsetValues.GetData();
+
+			// dot prod 'em. 
+			int32 NumToMult = Row.OffsetValues.Num();
+			for (int32 i = 0; i < NumToMult; i++)
+			{
+				OutArray[RowIndex] += OffsetInData[i] * RowValuePtr[i];
+			}
+		}
+	}
 }
diff --git a/Engine/Source/Runtime/AudioMixer/Private/DSP/MelScale.cpp b/Engine/Source/Runtime/AudioMixer/Private/DSP/MelScale.cpp
new file mode 100644
--- /dev/null
+++ b/Engine/Source/Runtime/AudioMixer/Private/DSP/MelScale.cpp
@@ -0,0 +1,155 @@
+// Copyright 1998-2019 Epic Games, Inc. All Rights Reserved.
+
+#include "DSP/MelScale.h"
+#include "CoreMinimal.h"
+#include "DSP/FloatArrayMath.h"
+
+namespace
+{
+	// Fill an array with a triangle. Slopes do not need to match.
+	void FillArrayWithTriangle(const float InLowerIndex, const float InCenterIndex, const float InUpperIndex, int32& OutOffsetIndex, TArray<float>& OutArray)
+	{
+		const int32 MinIndex = FMath::Max(0, FMath::CeilToInt(InLowerIndex)); 
+		const int32 MaxIndex = FMath::FloorToInt(InUpperIndex);
+		const int32 Num = FMath::Max(0, MaxIndex - MinIndex + 1);
+
+		// Prepare output 
+		OutOffsetIndex = MinIndex;
+		OutArray.Reset(Num);
+		if (Num < 1)
+		{
+			return;
+		}
+
+		OutArray.AddUninitialized(Num);
+		FMemory::Memset(OutArray.GetData(), 0, sizeof(float) * Num);
+
+		// determin upper and lower scales for triangle
+		const float LowerScale = 1.f / FMath::Max(SMALL_NUMBER, InCenterIndex - InLowerIndex);
+		const float UpperScale = 1.f / FMath::Max(SMALL_NUMBER, InUpperIndex - InCenterIndex);
+
+		// Build traingle windows for each row
+		float* OutArrayData = OutArray.GetData();
+		for (int32 OffsetIndex = 0; OffsetIndex < Num; OffsetIndex++)
+		{
+			const float ActualIndex = static_cast<float>(MinIndex + OffsetIndex);
+			float Value = 0.0f;
+			if (ActualIndex < InCenterIndex)
+			{
+				Value = (ActualIndex - InLowerIndex) * LowerScale;
+			}
+			else
+			{
+				Value = (InUpperIndex - ActualIndex) * UpperScale;
+			}
+
+			OutArrayData[OffsetIndex] = Value < 0.0f ? 0.0f : Value;
+		}
+	}
+}
+
+namespace Audio
+{
+	float MelToHz(float InMel) {
+		return 700.f * (FMath::Exp(InMel / 1127.f) - 1.f);
+	}
+
+	float HzToMel(float InHz) {
+		return 1127.f * FMath::Loge(InHz / 700.f + 1.f);
+	}
+
+	TUniquePtr<FContiguousSparse2DKernelTransform> NewMelSpectrumKernelTransform(const FMelSpectrumKernelSettings& InSettings, const int32 InFFTSize, const float InSampleRate)
+	{
+		// Convert frequency bounds to mel
+		float MinMel = HzToMel(InSettings.KernelMinCenterFreq);
+		float MaxMel = HzToMel(InSettings.KernelMaxCenterFreq);
+		float DeltaMel = (MaxMel - MinMel) / (InSettings.NumBands - 1);
+
+		// Only need to look at half of spectrum because magnitudes are mirrored.
+		const int32 NumUsefulFrequencyBins = InFFTSize / 2 + 1;
+		float BinsPerHz = (float)InFFTSize / (float)InSampleRate;
+		float HzPerBin = (float)InSampleRate / (float)InFFTSize;
+
+		// Used to define bounds of mel band windows.
+		float PreviousMelCenter = MinMel - DeltaMel;
+		float CurrentMelCenter = MinMel;
+		float NextMelCenter = MinMel + DeltaMel;
+
+		// Mel band windows are contiguous and sparse, so use this transform to make FFT->Mel fast.
+		TUniquePtr<FContiguousSparse2DKernelTransform> KernelTransform = MakeUnique<FContiguousSparse2DKernelTransform>(NumUsefulFrequencyBins, InSettings.NumBands);
+
+		// Create mel band windows one by one.
+		for (int32 MelBandIndex = 0; MelBandIndex < InSettings.NumBands; MelBandIndex++)
+		{
+			// Convert Mels to Hz
+			float LowerBandFreq = MelToHz(PreviousMelCenter);
+			float UpperBandFreq = MelToHz(NextMelCenter);
+			float CenterBandFreq = MelToHz(CurrentMelCenter);
+
+			//  Stretch frequency bandwidth if necessary
+			if (InSettings.BandWidthStretch != 1.0f)
+			{
+				LowerBandFreq = CenterBandFreq + InSettings.BandWidthStretch * (LowerBandFreq - CenterBandFreq);
+				UpperBandFreq = CenterBandFreq + InSettings.BandWidthStretch * (UpperBandFreq - CenterBandFreq);
+			}
+
+			// Determine bounding bins
+			float LowerFFTBin = LowerBandFreq * BinsPerHz;
+			float UpperFFTBin = UpperBandFreq * BinsPerHz;
+			float CenterFFTBin = CenterBandFreq * BinsPerHz;
+
+			// Create a triangle in frequency domain to account for mel band.
+			TArray<float> RowOffsetValues;
+			int32 RowOffsetIndex = 0;
+			FillArrayWithTriangle(LowerFFTBin, CenterFFTBin, UpperFFTBin, RowOffsetIndex, RowOffsetValues);
+
+			// By default, FillArrayWithTriangle creates constant amplitude triangle. So only noralize 
+			// if constant energy is chosen.
+			float NormDenom = 1.f;
+			switch (InSettings.Normalization)
+			{
+				case EMelNormalization::EqualAmplitude:
+					// By default, FillArrayWithTriangle creates constant amplitude triangle. So only noralize 
+					// if constant energy is chosen.
+					NormDenom = 1.f;
+					break;
+
+				case EMelNormalization::EqualEuclideanNorm:
+					NormDenom = FMath::Sqrt(0.5f * (UpperFFTBin - LowerFFTBin));
+					break;
+
+				case EMelNormalization::EqualEnergy:
+				default:
+					// weight bands to have constant energy of 1
+					NormDenom = 0.5f * (UpperFFTBin - LowerFFTBin);
+					break;
+			}
+
+			if ((NormDenom > 0.f) && (NormDenom != 1.f))
+			{
+				ArrayMultiplyByConstantInPlace(RowOffsetValues, 1.f / NormDenom);
+			}
+
+			// Check bounds of data being input into kernel
+			int32 MaxNumRowOFfsetValues = FMath::Max(0, NumUsefulFrequencyBins - RowOffsetIndex);
+			if (RowOffsetValues.Num() > MaxNumRowOFfsetValues)
+			{
+				RowOffsetValues.SetNum(MaxNumRowOFfsetValues);
+			}
+			if (RowOffsetIndex > NumUsefulFrequencyBins)
+			{
+				RowOffsetIndex = 0;
+			}
+
+			// Add row to kernel
+			KernelTransform->SetRow(MelBandIndex, RowOffsetIndex, TArrayView<const float>(RowOffsetValues));
+
+			// increment mel band counters 
+			PreviousMelCenter = CurrentMelCenter;
+			CurrentMelCenter = NextMelCenter;
+			NextMelCenter += DeltaMel;
+		}
+
+		return KernelTransform;
+	}
+}
diff --git a/Engine/Source/Runtime/AudioMixer/Public/DSP/ConstantQ.h b/Engine/Source/Runtime/AudioMixer/Public/DSP/ConstantQ.h
new file mode 100644
--- /dev/null
+++ b/Engine/Source/Runtime/AudioMixer/Public/DSP/ConstantQ.h
@@ -0,0 +1,55 @@
+// Copyright 1998-2019 Epic Games, Inc. All Rights Reserved.
+
+#pragma once
+
+#include "CoreMinimal.h"
+#include "DSP/FloatArrayMath.h"
+
+namespace Audio
+{
+	/** Band normalization schemes */
+	enum class EPseudoConstantQNormalization : uint8
+	{
+		EqualAmplitude,		//< No energy scaling. All bands weighted so max is 1.
+		EqualEuclideanNorm, //< Scale energy by euclidean norm. Good when using magnitude spectrum.
+		EqualEnergy			//< Scale energy by total energy. Good when using power spectrum.
+	};
+
+	/** Settings for Pseudo Constant Q Kernel generation. */
+	struct AUDIOMIXER_API FPseudoConstantQKernelSettings
+	{
+		int32 NumBands;									//< Total number of resulting constant Q bands.
+		float NumBandsPerOctave;							//< Number of bands to space within an octave.
+		float KernelLowestCenterFreq;					//< The starting frequency of the first band.
+		float BandWidthStretch;							//< Stretching factor to control overlap of adjacent bands.
+		EPseudoConstantQNormalization Normalization;	//< Normalization scheme for bands.
+
+		/** Default settings for Pseudo Constant Q Kernel */
+		FPseudoConstantQKernelSettings()
+		:	NumBands(96)
+		,	NumBandsPerOctave(12)
+		,	KernelLowestCenterFreq(40.f)
+		,	BandWidthStretch(1.0) 
+		,	Normalization(EPseudoConstantQNormalization::EqualEnergy)
+		{}
+	};
+
+	/** NewPseudoConstantQKernelTransform
+	 *
+	 *  This creates a new Pseudo Constant Q Kernel Transform. Pseudo Constant Q differs from standard Constant Q in
+	 *  that it applies a window to an existing DFT output as opposed to using a bank of filters.  The use of
+	 *  an FFT speeds up the calculation but introduces some limitations in bandwidth due to the granularity of the
+	 *  DFT. 
+	 *
+	 *  The resulting kernel will transform an array with (InFFTSize / 2 + 1) elements into an array with 
+	 *  InSettings.NumBands elements.
+	 *
+	 *  InSettings controls the number of bands and band properties.
+	 *  InFFTSize controls the expected input size to the kernel.
+	 *  InSampleRate describes the sampling rate of the analyzed audio.
+	 *
+	 *  Returns a TUniquePtr to a FContiguousSparse2DKernelTransform.
+	 */
+	AUDIOMIXER_API TUniquePtr<FContiguousSparse2DKernelTransform> NewPseudoConstantQKernelTransform(const FPseudoConstantQKernelSettings& InSettings, const int32 InFFTSize, const float InSampleRate);
+
+}
diff --git a/Engine/Source/Runtime/AudioMixer/Public/DSP/FloatArrayMath.h b/Engine/Source/Runtime/AudioMixer/Public/DSP/FloatArrayMath.h
--- a/Engine/Source/Runtime/AudioMixer/Public/DSP/FloatArrayMath.h
+++ b/Engine/Source/Runtime/AudioMixer/Public/DSP/FloatArrayMath.h
@@ -3,6 +3,7 @@
 #pragma once
 
 #include "CoreMinimal.h"
+#include "DSP/BufferVectorOperations.h"
 
 namespace Audio
 {
@@ -11,7 +12,7 @@ namespace Audio
 	 *  InData contains data to be cumulatively summed.
 	 *  OutData contains sum and is same size as InData.
 	 */
-	void ArrayCumulativeSum(const TArray<float>& InData, TArray<float>& OutData);
+	AUDIOMIXER_API void ArrayCumulativeSum(const TArray<float>& InData, TArray<float>& OutData);
 
 	/** Mean filter of array.
 	 *
@@ -23,7 +24,7 @@ namespace Audio
 	 *  WindowOrigin describes the offset from the windows first sample to the index of OutData. For example, if WindowOrigin = WindowSize/4, then OutData[i] = Mean(InData[i - Window/4 : i + 3 * Window / 4]).
 	 *  OutData contains the produceds data.
 	 */
-	void ArrayMeanFilter(const TArray<float>& InData, int32 WindowSize, int32 WindowOrigin, TArray<float>& OutData);
+	AUDIOMIXER_API void ArrayMeanFilter(const TArray<float>& InData, int32 WindowSize, int32 WindowOrigin, TArray<float>& OutData);
 
 	/** Max filter of array.
 	 *
@@ -34,11 +35,82 @@ namespace Audio
 	 *  WindowOrigin describes the offset from the windows first sample to the index of OutData. For example, if WindowOrigin = WindowSize/4, then OutData[i] = Max(InData[i - Window/4 : i + 3 * Window / 4]).
 	 *  OutData contains the produceds data.
 	 */
-	void ArrayMaxFilter(const TArray<float>& InData, int32 WindowSize, int32 WindowOrigin, TArray<float>& OutData);
+	AUDIOMIXER_API void ArrayMaxFilter(const TArray<float>& InData, int32 WindowSize, int32 WindowOrigin, TArray<float>& OutData);
 
 	/** Computes the EuclideanNorm of the InArray. Same as calculating the energy in window. */
 	void ArrayGetEuclideanNorm(const TArray<float>& InArray, float& OutEuclideanNorm);
 
 	/** Multiplies each element in InArray by InMultiplier */
-	void ArrayMultiplyByConstantInPlace(TArray<float>& InArray, const float InMultiplier);
+	AUDIOMIXER_API void ArrayMultiplyByConstantInPlace(TArray<float>& InArray, const float InMultiplier);
+
+
+	/** FContiguousSparse2DKernelTransform
+	 *
+	 *  FContiguousSparse2DKernelTransform applies a matrix transformation to an input array. 
+	 *  [OutArray] = [[Kernal]][InArray]  
+	 *
+	 *  It provides some optimization by exploit the contiguous and sparse qualities of the kernel rows,
+	 *  which allows it to skip multiplications with the number zero. 
+	 *
+	 *  It works with non-sparse and non-contiguous kernels as well, but will be more computationally 
+	 *  expensive than a naive implementation. Also, only takes advantage of sparse contiguous rows, not columns.
+	 */
+	class AUDIOMIXER_API FContiguousSparse2DKernelTransform
+	{
+		struct FRow
+		{
+			int32 StartIndex;
+			TArray<float> OffsetValues;
+		};
+
+	public:
+
+		/**
+		 * NumInElements sets the expected number of input array elements as well as the number of elements in a row.
+		 * NumOutElements sets the number of output array elements as well as the number or rows.
+		 */
+		FContiguousSparse2DKernelTransform(const int32 NumInElements, const int32 NumOutElements);
+
+		/** Returns the required size of the input array */
+		int32 GetNumInElements() const;
+
+		/** Returns the size of the output array */
+		int32 GetNumOutElements() const;
+	
+		/** Set the kernel values for an individual row.
+		 *
+		 *  RowIndex determines which row is being set.
+		 *  StartIndex denotes the offset into the row where the OffsetValues will be inserted.
+		 *  OffsetValues contains the contiguous chunk of values which represent all the nonzero elements in the row.
+		 */
+		void SetRow(const int32 RowIndex, const int32 StartIndex, TArrayView<const float> OffsetValues);
+
+		/** Transforms the input array given the kernel.
+		 *
+		 *  InArray is the array to be transformed. It must have `NumInElements` number of elements.
+		 *  OutArray is the transformed array. It will have `NumOutElements` number of elements.
+		 */
+		void TransformArray(TArrayView<const float> InArray, TArray<float>& OutArray) const;
+
+		/** Transforms the input array given the kernel.
+		 *
+		 *  InArray is the array to be transformed. It must have `NumInElements` number of elements.
+		 *  OutArray is the transformed array. It will have `NumOutElements` number of elements.
+		 */
+		void TransformArray(TArrayView<const float> InArray, AlignedFloatBuffer& OutArray) const;
+
+		/** Transforms the input array given the kernel.
+		 *
+		 *  InArray is the array to be transformed. It must have `NumInElements` number of elements.
+		 *  OutArray is the transformed array. It must be allocated to hold at least NumOutElements. 
+		 */
+		void TransformArray(const float* InArray, float* OutArray) const;
+
+	private:
+
+		int32 NumIn;
+		int32 NumOut;
+		TArray<FRow> Kernel;
+	};
+
 }
diff --git a/Engine/Source/Runtime/AudioMixer/Public/DSP/MelScale.h b/Engine/Source/Runtime/AudioMixer/Public/DSP/MelScale.h
new file mode 100644
--- /dev/null
+++ b/Engine/Source/Runtime/AudioMixer/Public/DSP/MelScale.h
@@ -0,0 +1,44 @@
+// Copyright 1998-2019 Epic Games, Inc. All Rights Reserved.
+
+#pragma once
+#include "CoreMinimal.h"
+#include "FloatArrayMath.h"
+
+
+namespace Audio
+{
+	/** Convert a mel frequency to a frequency in hz. */
+	AUDIOMIXER_API float MelToHz(float InMel);
+
+	/** Convert  a frequency in hz to a mel frequency. */
+	AUDIOMIXER_API float HzToMel(float InHz);
+
+	/** Normalization methods for a mel transform. */
+	enum class EMelNormalization : uint8
+	{
+		EqualAmplitude, 	//< Useful when working on magnitude spectrum
+		EqualEuclideanNorm, //< Scale energy by euclidean norm. Good when using magnitude spectrum.
+		EqualEnergy			//< Useful when working on power spectrum
+	};
+
+	/** Settings for a mel kernel which transforms an linearly space spectrum (e.g. FFT Magnitude)
+	 * to a mel spectrum */
+	struct AUDIOMIXER_API FMelSpectrumKernelSettings
+	{
+		int32 NumBands;						//< Number of bands in Mel spectrum
+		float KernelMinCenterFreq;			//< Minimum frequency of lowest mel band.
+		float KernelMaxCenterFreq;			//< Maximum frequency of highest mel band.
+		EMelNormalization Normalization;	//< Mel band scaling. EqualAmplitude is better for the magnitude spectrum, while EqualEnergy is better for the power spectrum.
+		float BandWidthStretch;				//< Controls the wideness of each band.
+
+		FMelSpectrumKernelSettings()
+		:	NumBands(26)
+		,	KernelMinCenterFreq(100.f)
+		,	KernelMaxCenterFreq(10000.f)
+		,	Normalization(EMelNormalization::EqualEnergy)
+		,	BandWidthStretch(1.f)
+		{}
+	};
+
+	AUDIOMIXER_API TUniquePtr<FContiguousSparse2DKernelTransform> NewMelSpectrumKernelTransform(const FMelSpectrumKernelSettings& InSettings, const int32 InFFTSize, const float InSampleRate);
+}
