[NO NON-PII COMMIT MESSAGE]

diff --git a/Engine/Source/Runtime/AudioMixer/Private/AudioMixerDevice.cpp b/Engine/Source/Runtime/AudioMixer/Private/AudioMixerDevice.cpp
--- a/Engine/Source/Runtime/AudioMixer/Private/AudioMixerDevice.cpp
+++ b/Engine/Source/Runtime/AudioMixer/Private/AudioMixerDevice.cpp
@@ -18,6 +18,7 @@
 #include "Misc/App.h"
 #include "Sound/AudioSettings.h"
 #include "ProfilingDebugging/CsvProfiler.h"
+#include "Async/Async.h"
 
 #if WITH_EDITOR
 #include "AudioEditorModule.h"
@@ -1632,12 +1633,22 @@ namespace Audio
 		if (!IsInAudioThread())
 		{
 			DECLARE_CYCLE_STAT(TEXT("FAudioThreadTask.RegisterSubmixBufferListener"), STAT_RegisterSubmixBufferListener, STATGROUP_AudioThreadCommands);
-
-			FAudioThread::RunCommandOnAudioThread([this, InSubmixBufferListener, InSubmix]()
+			auto AudioThreadCommand = [this, InSubmixBufferListener, InSubmix]()
 			{
 				CSV_SCOPED_TIMING_STAT(Audio, StopSpectrumAnalysis);
 				RegisterSubmixBufferListener(InSubmixBufferListener, InSubmix);
-			}, GET_STATID(STAT_RegisterSubmixBufferListener));
+			};
+
+			if (IsInGameThread())
+			{
+				FAudioThread::RunCommandOnAudioThread(AudioThreadCommand, GET_STATID(STAT_RegisterSubmixBufferListener));
+			}
+			else
+			{
+				AsyncTask(ENamedThreads::GameThread, [AudioThreadCommand]() {
+					FAudioThread::RunCommandOnAudioThread(AudioThreadCommand, GET_STATID(STAT_RegisterSubmixBufferListener));
+				});
+			}
 			return;
 		}
 
@@ -1662,11 +1673,23 @@ namespace Audio
 		{
 			DECLARE_CYCLE_STAT(TEXT("FAudioThreadTask.UnregisterSubmixBufferListener"), STAT_UnregisterSubmixBufferListener, STATGROUP_AudioThreadCommands);
 
-			FAudioThread::RunCommandOnAudioThread([this, InSubmixBufferListener, InSubmix]()
+			auto AudioThreadCommand = [this, InSubmixBufferListener, InSubmix]()
 			{
 				CSV_SCOPED_TIMING_STAT(Audio, UnregisterSubmixBufferListener);
 				UnregisterSubmixBufferListener(InSubmixBufferListener, InSubmix);
-			}, GET_STATID(STAT_UnregisterSubmixBufferListener));
+			};
+
+			if (IsInGameThread())
+			{
+				FAudioThread::RunCommandOnAudioThread(AudioThreadCommand, GET_STATID(STAT_UnregisterSubmixBufferListener));
+			}
+			else
+			{
+				AsyncTask(ENamedThreads::GameThread, [AudioThreadCommand]() {
+					FAudioThread::RunCommandOnAudioThread(AudioThreadCommand, GET_STATID(STAT_UnregisterSubmixBufferListener));
+				});
+			}
+
 			return;
 		}
 
