Switching to TArrayViews in some instances
Adding magnitude and power spectrum methods.
Adding some array math routines.

diff --git a/Engine/Source/Runtime/AudioMixer/Private/DSP/BufferVectorOperations.cpp b/Engine/Source/Runtime/AudioMixer/Private/DSP/BufferVectorOperations.cpp
new file mode 100644
--- /dev/null
+++ b/Engine/Source/Runtime/AudioMixer/Private/DSP/BufferVectorOperations.cpp
@@ -0,0 +1,1842 @@
+// Copyright 1998-2019 Epic Games, Inc. All Rights Reserved.
+
+#include "DSP/BufferVectorOperations.h"
+#include "AudioMixer.h"
+
+#define AUDIO_USE_SIMD 1
+
+namespace Audio
+{
+	/* Sets a values to zero if value is denormal. Denormal numbers significantly slow down floating point operations. */
+	void BufferUnderflowClampFast(AlignedFloatBuffer& InOutBuffer)
+	{
+		BufferUnderflowClampFast(InOutBuffer.GetData(), InOutBuffer.Num());
+	}
+	
+	/* Sets a values to zero if value is denormal. Denormal numbers significantly slow down floating point operations. */
+	void BufferUnderflowClampFast(float* RESTRICT InOutBuffer, const int32 InNum)
+	{
+		checkf(InNum >= 4, TEXT("Buffer must have atleast 4 elements."));
+		checkf(0 == (InNum % 4), TEXT("Buffer length be a multiple of 4."));
+		checkf(IsAligned<float*>(InOutBuffer, AUDIO_SIMD_FLOAT_ALIGNMENT), TEXT("Memory must be aligned to use vector operations."));
+
+		const VectorRegister VFMIN = MakeVectorRegister(FLT_MIN, FLT_MIN, FLT_MIN, FLT_MIN);
+		const VectorRegister VNFMIN = MakeVectorRegister(-FLT_MIN, -FLT_MIN, -FLT_MIN, -FLT_MIN);
+
+		for (int32 i = 0; i < InNum; i += 4)
+		{
+			VectorRegister VInOut = VectorLoadAligned(&InOutBuffer[i]);
+
+			// Create mask of denormal numbers.
+			VectorRegister Mask = VectorBitwiseAnd(VectorCompareGT(VInOut, VNFMIN), VectorCompareLT(VInOut, VFMIN));
+
+			// Choose between zero or original number based upon mask.
+			VInOut = VectorSelect(Mask, GlobalVectorConstants::FloatZero, VInOut);
+			VectorStoreAligned(VInOut, &InOutBuffer[i]);
+		}
+	}
+
+	void BufferMultiplyByConstant(const AlignedFloatBuffer& InFloatBuffer, float InValue, AlignedFloatBuffer& OutFloatBuffer)
+	{
+		check(InFloatBuffer.Num() >= 4);
+
+		// Prepare output buffer
+		OutFloatBuffer.Reset();
+		OutFloatBuffer.AddUninitialized(InFloatBuffer.Num());
+
+		check(InFloatBuffer.Num() == OutFloatBuffer.Num());
+
+		const int32 NumSamples = InFloatBuffer.Num();
+
+		// Get ptrs to audio buffers to avoid bounds check in non-shipping builds
+		const float* InBufferPtr = InFloatBuffer.GetData();
+		float* OutBufferPtr = OutFloatBuffer.GetData();
+
+		BufferMultiplyByConstant(InBufferPtr, InValue, OutBufferPtr, NumSamples);
+	}
+
+	void BufferMultiplyByConstant(const float* RESTRICT InFloatBuffer, float InValue, float* RESTRICT OutFloatBuffer, const int32 InNumSamples)
+	{
+		check(InNumSamples >= 4);
+#if !AUDIO_USE_SIMD
+		for (int32 i = 0; i < InNumSamples; ++i)
+		{
+			OutFloatBuffer[i] = InValue * InFloatBuffer[i];
+		}
+#else
+
+		// Can only SIMD on multiple of 4 buffers, we'll do normal multiples on last bit
+		const int32 NumSamplesRemaining = InNumSamples % 4;
+		const int32 NumSamplesToSimd = InNumSamples - NumSamplesRemaining;
+
+		// Load the single value we want to multiply all values by into a vector register
+		const VectorRegister MultiplyValue = VectorLoadFloat1(&InValue);
+		for (int32 i = 0; i < NumSamplesToSimd; i += 4)
+		{
+			// Load the next 4 samples of the input buffer into a register
+			VectorRegister InputBufferRegister = VectorLoadAligned(&InFloatBuffer[i]);
+
+			// Perform the multiply
+			VectorRegister Temp = VectorMultiply(InputBufferRegister, MultiplyValue);
+
+			// Store results into the output buffer
+			VectorStoreAligned(Temp, &OutFloatBuffer[i]);
+		}
+
+		// Perform remaining non-simd values left over
+		for (int32 i = 0; i < NumSamplesRemaining; ++i)
+		{
+			OutFloatBuffer[NumSamplesToSimd + i] = InValue * InFloatBuffer[NumSamplesToSimd + i];
+		}
+#endif
+	}
+
+	void MultiplyBufferByConstantInPlace(AlignedFloatBuffer& InBuffer, float InGain)
+	{
+		MultiplyBufferByConstantInPlace(InBuffer.GetData(), InBuffer.Num(), InGain);
+	}
+
+	void MultiplyBufferByConstantInPlace(float* RESTRICT InBuffer, int32 NumSamples, float InGain)
+	{
+		const VectorRegister Gain = VectorLoadFloat1(&InGain);
+
+		for (int32 i = 0; i < NumSamples; i += 4)
+		{
+			VectorRegister Output = VectorLoadAligned(&InBuffer[i]);
+			Output = VectorMultiply(Output, Gain);
+			VectorStoreAligned(Output, &InBuffer[i]);
+		}
+	}
+
+	/* Performs an element-wise weighted sum OutputBuffer = (InBuffer1 x InGain1) + (InBuffer2 x InGain2) */
+	void BufferWeightedSumFast(const AlignedFloatBuffer& InBuffer1, float InGain1, const AlignedFloatBuffer& InBuffer2, float InGain2, AlignedFloatBuffer& OutBuffer)
+	{
+		checkf(InBuffer1.Num() == InBuffer2.Num(), TEXT("Buffers must be equal length"));
+		OutBuffer.Reset();
+		OutBuffer.AddUninitialized(InBuffer1.Num());
+
+		BufferWeightedSumFast(InBuffer1.GetData(), InGain1, InBuffer2.GetData(), InGain2, OutBuffer.GetData(), InBuffer1.Num());
+	}
+
+	/* Performs an element-wise weighted sum OutputBuffer = (InBuffer1 x InGain1) + InBuffer2 */
+	void BufferWeightedSumFast(const AlignedFloatBuffer& InBuffer1, float InGain1, const AlignedFloatBuffer& InBuffer2, AlignedFloatBuffer& OutBuffer)
+	{
+		checkf(InBuffer1.Num() == InBuffer2.Num(), TEXT("Buffers must be equal length"));
+		OutBuffer.Reset();
+		OutBuffer.AddUninitialized(InBuffer1.Num());
+
+		BufferWeightedSumFast(InBuffer1.GetData(), InGain1, InBuffer2.GetData(), OutBuffer.GetData(), InBuffer1.Num());
+	}
+
+	/* Performs an element-wise weighted sum OutputBuffer = (InBuffer1 x InGain1) + (InBuffer2 x InGain2) */
+	void BufferWeightedSumFast(const float* RESTRICT InBuffer1, float InGain1, const float* RESTRICT InBuffer2, float InGain2, float* RESTRICT OutBuffer, int32 InNum)
+	{
+		checkf(InNum >= 4, TEXT("Buffer must have atleast 4 elements."));
+		checkf(0 == (InNum % 4), TEXT("Buffer length be a multiple of 4."));
+		checkf(IsAligned<const float*>(InBuffer1, AUDIO_SIMD_FLOAT_ALIGNMENT), TEXT("Memory must be aligned to use vector operations."));
+		checkf(IsAligned<const float*>(InBuffer2, AUDIO_SIMD_FLOAT_ALIGNMENT), TEXT("Memory must be aligned to use vector operations."));
+		checkf(IsAligned<float*>(OutBuffer, AUDIO_SIMD_FLOAT_ALIGNMENT), TEXT("Memory must be aligned to use vector operations."));
+
+		VectorRegister Gain1Vector = VectorLoadFloat1(&InGain1);
+		VectorRegister Gain2Vector = VectorLoadFloat1(&InGain2);
+
+		for (int32 i = 0; i < InNum; i += 4)
+		{
+			// InBuffer1 x InGain1
+			VectorRegister Input1  = VectorLoadAligned(&InBuffer1[i]);
+			VectorRegister Weighted1 = VectorMultiply(Input1, Gain1Vector);
+
+			// InBuffer2 x InGain2
+			VectorRegister Input2  = VectorLoadAligned(&InBuffer2[i]);
+			VectorRegister Weighted2 = VectorMultiply(Input2, Gain2Vector);
+
+			VectorRegister Output = VectorAdd(Weighted1, Weighted2);
+			VectorStoreAligned(Output, &OutBuffer[i]);
+		}
+	}
+
+	/* Performs an element-wise weighted sum OutputBuffer = (InBuffer1 x InGain1) + InBuffer2 */
+	void BufferWeightedSumFast(const float* RESTRICT InBuffer1, float InGain1, const float* RESTRICT InBuffer2, float* RESTRICT OutBuffer, int32 InNum)
+	{
+		checkf(InNum >= 4, TEXT("Buffer must have atleast 4 elements."));
+		checkf(0 == (InNum % 4), TEXT("Buffer length be a multiple of 4."));
+		checkf(IsAligned<const float*>(InBuffer1, AUDIO_SIMD_FLOAT_ALIGNMENT), TEXT("Memory must be aligned to use vector operations."));
+		checkf(IsAligned<const float*>(InBuffer2, AUDIO_SIMD_FLOAT_ALIGNMENT), TEXT("Memory must be aligned to use vector operations."));
+		checkf(IsAligned<float*>(OutBuffer, AUDIO_SIMD_FLOAT_ALIGNMENT), TEXT("Memory must be aligned to use vector operations."));
+
+		VectorRegister Gain1Vector = VectorLoadFloat1(&InGain1);
+
+		for (int32 i = 0; i < InNum; i += 4)
+		{
+			// InBuffer1 x InGain1
+			VectorRegister Input1  = VectorLoadAligned(&InBuffer1[i]);
+			VectorRegister Weighted1 = VectorMultiply(Input1, Gain1Vector);
+
+			VectorRegister Input2  = VectorLoadAligned(&InBuffer2[i]);
+
+			VectorRegister Output = VectorAdd(Weighted1, Input2);
+			VectorStoreAligned(Output, &OutBuffer[i]);
+		}
+	}
+
+	void FadeBufferFast(AlignedFloatBuffer& OutFloatBuffer, const float StartValue, const float EndValue)
+	{
+		FadeBufferFast(OutFloatBuffer.GetData(), OutFloatBuffer.Num(), StartValue, EndValue);
+	}
+
+	void FadeBufferFast(float* RESTRICT OutFloatBuffer, int32 NumSamples, const float StartValue, const float EndValue)
+	{
+		checkf(IsAligned<float*>(OutFloatBuffer, 4), TEXT("Memory must be aligned to use vector operations."));
+		checkf(NumSamples % 4 == 0, TEXT("Please use a buffer size that is a multiple of 4."));
+
+#if !AUDIO_USE_SIMD
+		float Gain = StartValue;
+		if (FMath::IsNearlyEqual(StartValue, EndValue))
+		{
+			// No need to do anything if start and end values are both 0.0
+			if (StartValue == 0.0f)
+			{
+				FMemory::Memset(OutFloatBuffer, 0, sizeof(float)*NumSamples);
+			}
+			else
+			{
+				// Only need to do a buffer multiply if start and end values are the same
+				for (int32 i = 0; i < NumSamples; ++i)
+				{
+					OutFloatBuffer[i] = OutFloatBuffer[i] * Gain;
+				}
+			}
+		}
+		else
+		{
+			// Do a fade from start to end
+			const float DeltaValue = ((EndValue - StartValue) / NumSamples);
+			for (int32 i = 0; i < NumSamples; ++i)
+			{
+				OutFloatBuffer[i] = OutFloatBuffer[i] * Gain;
+				Gain += DeltaValue;
+			}
+		}
+#else
+		const int32 NumIterations = NumSamples / 4;
+
+		if (FMath::IsNearlyEqual(StartValue, EndValue))
+		{
+			// No need to do anything if start and end values are both 0.0
+			if (StartValue == 0.0f)
+			{
+				FMemory::Memset(OutFloatBuffer, 0, sizeof(float)*NumSamples);
+			}
+			else
+			{
+				VectorRegister Gain = VectorLoadFloat1(&StartValue);
+
+				for (int32 i = 0; i < NumSamples; i += 4)
+				{
+					VectorRegister Output = VectorLoadAligned(&OutFloatBuffer[i]);
+					Output = VectorMultiply(Output, Gain);
+					VectorStoreAligned(Output, &OutFloatBuffer[i]);
+				}
+			}
+		}
+		else
+		{
+			const float DeltaValue = ((EndValue - StartValue) / NumIterations);
+
+			VectorRegister Gain = VectorLoadFloat1(&StartValue);
+			VectorRegister Delta = VectorLoadFloat1(&DeltaValue);
+
+			for (int32 i = 0; i < NumSamples; i += 4)
+			{
+				VectorRegister Output = VectorLoadAligned(&OutFloatBuffer[i]);
+				Output = VectorMultiply(Output, Gain);
+				Gain = VectorAdd(Gain, Delta);
+				VectorStoreAligned(Output, &OutFloatBuffer[i]);
+			}
+		}
+#endif
+	}
+
+	void MixInBufferFast(const AlignedFloatBuffer& InFloatBuffer, AlignedFloatBuffer& BufferToSumTo, const float Gain)
+	{
+		MixInBufferFast(InFloatBuffer.GetData(), BufferToSumTo.GetData(), InFloatBuffer.Num(), Gain);
+	}
+
+	void MixInBufferFast(const float* RESTRICT InFloatBuffer, float* RESTRICT BufferToSumTo, int32 NumSamples, const float Gain)
+	{
+		checkf(IsAligned<const float*>(InFloatBuffer, 4), TEXT("Memory must be aligned to use vector operations."));
+		checkf(IsAligned<float*>(BufferToSumTo, 4), TEXT("Memory must be aligned to use vector operations."));
+		checkf(NumSamples % 4 == 0, TEXT("Please use a buffer size that is a multiple of 4."));
+
+#if !AUDIO_USE_SIMD
+		for (int32 i = 0; i < NumSamples; ++i)
+		{
+			BufferToSumTo[i] += InFloatBuffer[i] * Gain;
+		}
+#else
+		VectorRegister GainVector = VectorLoadFloat1(&Gain);
+
+		for (int32 i = 0; i < NumSamples; i += 4)
+		{
+			VectorRegister Output = VectorLoadAligned(&BufferToSumTo[i]);
+			VectorRegister Input  = VectorLoadAligned(&InFloatBuffer[i]);
+			Output = VectorMultiplyAdd(Input, GainVector, Output);
+			VectorStoreAligned(Output, &BufferToSumTo[i]);
+		}
+#endif
+	}
+
+	void MixInBufferFast(const AlignedFloatBuffer& InFloatBuffer, AlignedFloatBuffer& BufferToSumTo)
+	{
+		checkf(InFloatBuffer.Num() == BufferToSumTo.Num(), TEXT("Buffers must be equal size"));
+		MixInBufferFast(InFloatBuffer.GetData(), BufferToSumTo.GetData(), InFloatBuffer.Num());
+	}
+
+	void MixInBufferFast(const float* RESTRICT InFloatBuffer, float* RESTRICT BufferToSumTo, int32 NumSamples)
+	{
+		checkf(IsAligned<const float*>(InFloatBuffer, 4), TEXT("Memory must be aligned to use vector operations."));
+		checkf(IsAligned<float*>(BufferToSumTo, 4), TEXT("Memory must be aligned to use vector operations."));
+		checkf(NumSamples % 4 == 0, TEXT("Please use a buffer size that is a multiple of 4."));
+
+#if !AUDIO_USE_SIMD
+		for (int32 i = 0; i < NumSamples; ++i)
+		{
+			BufferToSumTo[i] += InFloatBuffer[i];
+		}
+#else
+		for (int32 i = 0; i < NumSamples; i += 4)
+		{
+			VectorRegister Output = VectorLoadAligned(&BufferToSumTo[i]);
+			VectorRegister Input = VectorLoadAligned(&InFloatBuffer[i]);
+			Output = VectorAdd(Input, Output);
+			VectorStoreAligned(Output, &BufferToSumTo[i]);
+		}
+#endif
+	}
+
+	/* Subtracts two buffers together element-wise. */
+	void BufferSubtractFast(const AlignedFloatBuffer& InMinuend, const AlignedFloatBuffer& InSubtrahend, AlignedFloatBuffer& OutputBuffer)
+	{
+		const int32 InNum = InMinuend.Num();
+		OutputBuffer.Reset(InNum);
+		OutputBuffer.AddUninitialized(InNum);
+
+		checkf(InMinuend.Num() == InSubtrahend.Num(), TEXT("Input buffers must be equal length"));
+
+		BufferSubtractFast(InMinuend.GetData(), InSubtrahend.GetData(), OutputBuffer.GetData(), OutputBuffer.Num());
+	}
+
+	/* Subtracts two buffers together element-wise. */
+	void BufferSubtractFast(const float* RESTRICT InMinuend, const float* RESTRICT InSubtrahend, float* RESTRICT OutBuffer, int32 InNum)
+	{
+		checkf(InNum >= 4, TEXT("Buffer must have atleast 4 elements."));
+		checkf(0 == (InNum % 4), TEXT("Buffer length be a multiple of 4."));
+		checkf(IsAligned<const float*>(InMinuend, AUDIO_SIMD_FLOAT_ALIGNMENT), TEXT("Memory must be aligned to use vector operations."));
+		checkf(IsAligned<const float*>(InSubtrahend, AUDIO_SIMD_FLOAT_ALIGNMENT), TEXT("Memory must be aligned to use vector operations."));
+		checkf(IsAligned<float*>(OutBuffer, AUDIO_SIMD_FLOAT_ALIGNMENT), TEXT("Memory must be aligned to use vector operations."));
+
+		for (int32 i = 0; i < InNum; i += 4)
+		{
+			VectorRegister Input1 = VectorLoadAligned(&InMinuend[i]);
+			VectorRegister Input2 = VectorLoadAligned(&InSubtrahend[i]);
+			VectorRegister Output = VectorSubtract(Input1, Input2);
+			VectorStoreAligned(Output, &OutBuffer[i]);
+		}
+	}
+
+	/* Performs element-wise in-place subtraction placing the result in the subtrahend. InOutSubtrahend = InMinuend - InOutSubtrahend */
+	void BufferSubtractInPlace1Fast(const AlignedFloatBuffer& InMinuend, AlignedFloatBuffer& InOutSubtrahend)
+	{
+		checkf(InMinuend.Num() == InOutSubtrahend.Num(), TEXT("Input buffers must be equal length"));
+		BufferSubtractInPlace1Fast(InMinuend.GetData(), InOutSubtrahend.GetData(), InMinuend.Num());
+	}
+
+	/* Performs element-wise in-place subtraction placing the result in the subtrahend. InOutSubtrahend = InMinuend - InOutSubtrahend */
+	void BufferSubtractInPlace1Fast(const float* RESTRICT InMinuend, float* RESTRICT InOutSubtrahend, int32 InNum)
+	{
+		checkf(InNum >= 4, TEXT("Buffer must have atleast 4 elements."));
+		checkf(0 == (InNum % 4), TEXT("Buffer length be a multiple of 4."));
+		checkf(IsAligned<const float*>(InMinuend, AUDIO_SIMD_FLOAT_ALIGNMENT), TEXT("Memory must be aligned to use vector operations."));
+		checkf(IsAligned<const float*>(InOutSubtrahend, AUDIO_SIMD_FLOAT_ALIGNMENT), TEXT("Memory must be aligned to use vector operations."));
+		
+		for (int32 i = 0; i < InNum; i += 4)
+		{
+			VectorRegister Input1 = VectorLoadAligned(&InMinuend[i]);
+			VectorRegister Input2 = VectorLoadAligned(&InOutSubtrahend[i]);
+
+			VectorRegister Output = VectorSubtract(Input1, Input2);
+			VectorStoreAligned(Output, &InOutSubtrahend[i]);
+		}
+	}
+	
+	/* Performs element-wise in-place subtraction placing the result in the minuend. InOutMinuend = InOutMinuend - InSubtrahend */
+	void BufferSubtractInPlace2Fast(AlignedFloatBuffer& InOutMinuend, const AlignedFloatBuffer& InSubtrahend)
+	{
+		checkf(InOutMinuend.Num() == InSubtrahend.Num(), TEXT("Input buffers must be equal length"));
+		BufferSubtractInPlace2Fast(InOutMinuend.GetData(), InSubtrahend.GetData(), InOutMinuend.Num());
+	}
+
+	/* Performs element-wise in-place subtraction placing the result in the minuend. InOutMinuend = InOutMinuend - InSubtrahend */
+	void BufferSubtractInPlace2Fast(float* RESTRICT InOutMinuend, const float* RESTRICT InSubtrahend, int32 InNum)
+	{
+		checkf(InNum >= 4, TEXT("Buffer must have atleast 4 elements."));
+		checkf(0 == (InNum % 4), TEXT("Buffer length be a multiple of 4."));
+		checkf(IsAligned<const float*>(InOutMinuend, AUDIO_SIMD_FLOAT_ALIGNMENT), TEXT("Memory must be aligned to use vector operations."));
+		checkf(IsAligned<const float*>(InSubtrahend, AUDIO_SIMD_FLOAT_ALIGNMENT), TEXT("Memory must be aligned to use vector operations."));
+
+		for (int32 i = 0; i < InNum; i += 4)
+		{
+			VectorRegister Input1 = VectorLoadAligned(&InOutMinuend[i]);
+			VectorRegister Input2 = VectorLoadAligned(&InSubtrahend[i]);
+
+			VectorRegister Output = VectorSubtract(Input1, Input2);
+			VectorStoreAligned(Output, &InOutMinuend[i]);
+		}
+	}
+
+	void SumBuffers(const AlignedFloatBuffer& InFloatBuffer1, const AlignedFloatBuffer& InFloatBuffer2, AlignedFloatBuffer& OutputBuffer)
+	{
+		checkf(InFloatBuffer1.Num() == InFloatBuffer2.Num(), TEXT("Input buffers must be equal length"));
+		const int32 InNum = InFloatBuffer1.Num();
+		OutputBuffer.Reset(InNum);
+		OutputBuffer.AddUninitialized(InNum);
+
+		SumBuffers(InFloatBuffer1.GetData(), InFloatBuffer2.GetData(), OutputBuffer.GetData(), InNum);
+	}
+
+	void SumBuffers(const float* RESTRICT InFloatBuffer1, const float* RESTRICT InFloatBuffer2, float* RESTRICT OutputBuffer, int32 NumSamples)
+	{
+		checkf(IsAligned<const float*>(InFloatBuffer1, 4), TEXT("Memory must be aligned to use vector operations."));
+		checkf(IsAligned<const float*>(InFloatBuffer2, 4), TEXT("Memory must be aligned to use vector operations."));
+		checkf(IsAligned<float*>(OutputBuffer, 4), TEXT("Memory must be aligned to use vector operations."));
+		checkf(NumSamples % 4 == 0, TEXT("Please use a buffer size that is a multiple of 4."));
+
+#if !AUDIO_USE_SIMD
+		for (int32 i = 0; i < NumSamples; ++i)
+		{
+			OutputBuffer[i] = InFloatBuffer1[i] + InFloatBuffer2[i];
+		}
+#else
+		for (int32 i = 0; i < NumSamples; i += 4)
+		{
+			VectorRegister Input1 = VectorLoadAligned(&InFloatBuffer1[i]);
+			VectorRegister Input2 = VectorLoadAligned(&InFloatBuffer2[i]);
+
+			VectorRegister Output = VectorAdd(Input1, Input2);
+			VectorStoreAligned(Output, &OutputBuffer[i]);
+		}
+#endif
+	}
+
+	void MultiplyBuffersInPlace(const AlignedFloatBuffer& InFloatBuffer, AlignedFloatBuffer& BufferToMultiply)
+	{
+		MultiplyBuffersInPlace(InFloatBuffer.GetData(), BufferToMultiply.GetData(), BufferToMultiply.Num());
+	}
+
+	void MultiplyBuffersInPlace(const float* RESTRICT InFloatBuffer, float* RESTRICT BufferToMultiply, int32 NumSamples)
+	{
+		checkf(IsAligned<const float*>(InFloatBuffer, 4), TEXT("Memory must be aligned to use vector operations."));
+		checkf(IsAligned<const float*>(BufferToMultiply, 4), TEXT("Memory must be aligned to use vector operations."));
+		checkf(NumSamples % 4 == 0, TEXT("Please use a buffer size that is a multiple of 4."));
+
+		for (int32 i = 0; i < NumSamples; i += 4)
+		{
+			VectorRegister Input1 = VectorLoadAligned(&InFloatBuffer[i]);
+			VectorRegister Output = VectorLoadAligned(&BufferToMultiply[i]);
+
+			Output = VectorMultiply(Input1, Output);
+			VectorStoreAligned(Output, &BufferToMultiply[i]);
+		}
+	}
+
+	float GetMagnitude(const AlignedFloatBuffer& Buffer)
+	{
+		return GetMagnitude(Buffer.GetData(), Buffer.Num());
+	}
+
+	float GetMagnitude(const float* RESTRICT Buffer, int32 NumSamples)
+	{
+		checkf(NumSamples % 4 == 0, TEXT("Please use a buffer size that is a multiple of 4."));
+
+#if !AUDIO_USE_SIMD
+		float Sum = 0.0f;
+		for (int32 i = 0; i < NumSamples; ++i)
+		{
+			Sum += Buffer[i] * Buffer[i];
+		}
+		return FMath::Sqrt(Sum);
+#else
+		VectorRegister Sum = VectorZero();
+
+		const float Exponent = 2.0f;
+		VectorRegister ExponentVector = VectorLoadFloat1(&Exponent);
+
+		for (int32 i = 0; i < NumSamples; i += 4)
+		{
+			VectorRegister Input = VectorPow(VectorLoadAligned(&Buffer[i]), ExponentVector);
+			Sum = VectorAdd(Sum, Input);
+		}
+
+		float PartionedSums[4];
+		VectorStoreAligned(Sum, PartionedSums);
+
+		return FMath::Sqrt(PartionedSums[0] + PartionedSums[1] + PartionedSums[2] + PartionedSums[3]);
+#endif
+	}
+
+	float GetAverageAmplitude(const AlignedFloatBuffer& Buffer)
+	{
+		checkf(Buffer.Num() % 4 == 0, TEXT("Please use a buffer size that is a multiple of 4."));
+
+		return GetAverageAmplitude(Buffer.GetData(), Buffer.Num());
+	}
+
+	float GetAverageAmplitude(const float* RESTRICT Buffer, int32 NumSamples)
+	{
+		checkf(NumSamples % 4 == 0, TEXT("Please use a buffer size that is a multiple of 4."));
+
+#if !AUDIO_USE_SIMD
+		float Sum = 0.0f;
+		for (int32 i = 0; i < NumSamples; ++i)
+		{
+			Sum += Buffer[i];
+		}
+		return Sum / NumSamples;
+#else
+		VectorRegister Sum = VectorZero();
+
+		for (int32 i = 0; i < NumSamples; i += 4)
+		{
+			VectorRegister Input = VectorAbs(VectorLoadAligned(&Buffer[i]));
+			Sum = VectorAdd(Sum, Input);
+		}
+
+		float PartionedSums[4];
+		VectorStore(Sum, PartionedSums);
+		
+		return (PartionedSums[0] + PartionedSums[1] + PartionedSums[2] + PartionedSums[3]) / NumSamples;
+#endif
+	}
+
+	/**
+	 * CHANNEL MIXING OPERATIONS:
+	 * To understand these functions, it's best that you have prior experience reading SIMD code.
+	 * These functions are all variations on component-wise matrix multiplies. 
+	 * There are two types of functions below:
+	 * 
+	 * Apply[N]ChannelGain:
+	 * These are all in-place multiplies of an N-length gain vector and an N-length frame.
+	 * There are two flavors of every variant of this function: The non-interpolating form (which takes a single gain matrix)
+	 * And the interpolating form (which takes a start gain matrix and interpolates to the end gain matrix over the given number of frames).
+	 * All non-interpolating forms of these functions use the following steps:
+	 *    1. Create a const GainVector, or series of GainVectors, that maps to the multiplies required for each iteration.
+	 *    2. In a loop:
+	 *           i.   load a frame or number of frames into a vector register or series of vector registers (these are named Result).
+	 *           ii.  perform a vector multiply on result with the corresponding gain vector.
+	 *           iii. store the result vector in the same position in the buffer we loaded from.
+	 *
+	 * The interpolating forms of these functions use the following steps:
+	 *    1. Initialize a non-const GainVector, or series of GainVectors, from StartGains, that maps to the multiplies required for each iteration.
+	 *    2. Compute the amount we add to GainVector for each iteration to reach Destination Gains and store it in the const GainDeltasVector.
+	 *    3. In a loop:
+	 *           i.   load a frame or number of frames into a vector register or series of vector registers (these are named Result).
+	 *           ii.  perform a vector multiply on result with the corresponding gain vector.
+	 *           iii. store the result vector in the same position in the buffer we loaded from.
+	 *           iv.  increment each GainVector by it's corresponding GainDeltasVector.
+	 *
+	 *
+	 * MixMonoTo[N]ChannelsFast and Mix2ChannelsTo[N]ChannelsFast:
+	 * These, like Apply[N]ChannelGain, all have non-interpolating and interpolating forms.
+	 * All non-interpolating forms of these functions use the following steps:
+	 *    1. Create a const GainVector, or series of GainVectors, that maps to the multiplies required for each input channel for each iteration.
+	 *    2. In a loop:
+	 *           i.   load a frame or number of frames into a const vector register or series of const vector registers (these are named Input).
+	 *           ii.  perform a vector multiply on input with the corresponding gain vector and store the result in a new vector or series of vectors named Result.
+	 *           iii. if there is a second input channel, store the results of the following MultiplyAdd operation to Results: (Gain Vectors for second channel) * (Input vectors for second channel) + (Result vectors from step ii).
+	 *
+	 * Interpolating forms of these functions use the following steps:
+	 *    1. Initialize a non-const GainVector, or series of GainVectors, from StartGains, that maps to the multiplies required for each input channel for each iteration.
+	 *    2. Compute the amount we add to each GainVector for each iteration to reach the vector's corresponding DestinationGains and store it in a corresponding GainDeltaVector.
+	 *    3. In a loop:
+	 *           i.   load a frame or number of frames into a const vector register or series of const vector registers (these are named Input).
+	 *           ii.  perform a vector multiply on input with the corresponding gain vector and store the result in a new vector or series of vectors named Result.
+	 *           iii. if there is a second input channel, store the results of the following MultiplyAdd operation to Results: (Gain Vectors for second channel) * (Input vectors for second channel) + (Result vectors from step ii).
+	 *           iv.  increment each GainVector by it's corresponding GainDeltasVector.
+	 * 
+	 * DETERMINING THE VECTOR LAYOUT FOR EACH FUNCTION:
+	 * For every variant of Mix[N]ChannelsTo[N]ChannelsFast, we use the least common multiple of the number of output channels and the SIMD vector length (4) to calulate the length of our matrix.
+	 * For example, MixMonoTo4ChannelsFast can use a single VectorRegister for each variable. GainVector's values are [g0, g1, g2, g3], input channels are mapped to [i0, i0, i0, i0], and output channels are mapped to [o0, o1, o2, o3].
+	 * MixMonoTo8ChannelsFast has an LCM of 8, so we use two VectorRegister for each variable. This results in the following layout:
+	 * GainVector1:   [g0, g1, g2, g3] GainVector2:   [g4, g5, g6, g7]
+	 * InputVector1:  [i0, i0, i0, i0] InputVector2:  [i0, i0, i0, i0]
+	 * ResultVector1: [o0, o1, o2, o3] ResultVector2: [o4, o5, o6, o7]
+	 *
+	 * The general naming convention for vector variables is [Name]Vector[VectorIndex] for MixMonoTo[N]ChannelsFast functions.
+	 * For Mix2ChannelsTo[N]ChannelsFast functions, the naming convention for vector variables is [Name]Vector[VectorIndex][InputChannelIndex].
+	 *
+	 * For clarity, the layout of vectors for each function variant is given in a block comment above that function.
+	 */
+
+	void Apply2ChannelGain(AlignedFloatBuffer& StereoBuffer, const float* RESTRICT Gains)
+	{
+		Apply2ChannelGain(StereoBuffer.GetData(), StereoBuffer.Num(), Gains);
+	}
+
+	void Apply2ChannelGain(float* RESTRICT StereoBuffer, int32 NumSamples, const float* RESTRICT Gains)
+	{
+		const VectorRegister GainVector = VectorLoadFloat2(Gains);
+		
+		for (int32 i = 0; i < NumSamples; i += 4)
+		{
+			VectorRegister Result = VectorLoadAligned(&StereoBuffer[i]);
+			Result = VectorMultiply(Result, GainVector);
+			VectorStoreAligned(Result, &StereoBuffer[i]);
+		}
+	}
+
+	void Apply2ChannelGain(AlignedFloatBuffer& StereoBuffer, const float* RESTRICT StartGains, const float* RESTRICT EndGains)
+	{
+		Apply2ChannelGain(StereoBuffer.GetData(), StereoBuffer.Num(), StartGains, EndGains);
+	}
+
+	void Apply2ChannelGain(float* RESTRICT StereoBuffer, int32 NumSamples, const float* RESTRICT StartGains, const float* RESTRICT EndGains)
+	{
+		// Initialize GainVector at StartGains and compute GainDeltasVector:
+		VectorRegister GainVector = VectorLoadFloat2(StartGains);
+		const VectorRegister DestinationVector = VectorLoadFloat2(EndGains);
+		const VectorRegister NumFramesVector = VectorSetFloat1(NumSamples / 4.0f);
+		const VectorRegister GainDeltasVector = VectorDivide(VectorSubtract(DestinationVector, GainVector), NumFramesVector);
+
+		for (int32 i = 0; i < NumSamples; i += 4)
+		{
+			VectorRegister Result = VectorLoadAligned(&StereoBuffer[i]);
+			Result = VectorMultiply(Result, GainVector);
+			VectorStoreAligned(Result, &StereoBuffer[i]);
+
+			GainVector = VectorAdd(GainVector, GainDeltasVector);
+		}
+	}
+
+	void MixMonoTo2ChannelsFast(const AlignedFloatBuffer& MonoBuffer, AlignedFloatBuffer& DestinationBuffer, const float* RESTRICT Gains)
+	{
+		MixMonoTo2ChannelsFast(MonoBuffer.GetData(), DestinationBuffer.GetData(), DestinationBuffer.Num() / 2, Gains);
+	}
+
+	/**
+	* See CHANNEL MIXING OPERATIONS above for more info.
+	* 2 frames per iteration:
+	* +------------+---------+---------+---------+---------+
+	* | VectorName | Index 0 | Index 1 | Index 2 | Index 3 |
+	* +------------+---------+---------+---------+---------+
+	* | Gain       | g0      | g1      | g0      | g1      |
+	* |            | *       | *       | *       | *       |
+	* | Input      | i0      | i0      | i0      | i0      |
+	* |            | =       | =       | =       | =       |
+	* | Output     | o0      | o1      | o2      | o3      |
+	* +------------+---------+---------+---------+---------+
+	*/
+	void MixMonoTo2ChannelsFast(const float* RESTRICT MonoBuffer, float* RESTRICT DestinationBuffer, int32 NumFrames, const float* RESTRICT Gains)
+	{
+		const VectorRegister GainVector = VectorLoadFloat2(Gains);
+		for (int32 i = 0; i < NumFrames; i += 2)
+		{
+			VectorRegister Result = VectorSet(MonoBuffer[i], MonoBuffer[i], MonoBuffer[i + 1], MonoBuffer[i + 1]);
+			Result = VectorMultiply(Result, GainVector);
+			VectorStoreAligned(Result, &DestinationBuffer[i*2]);
+		}
+	}
+
+	void MixMonoTo2ChannelsFast(const AlignedFloatBuffer& MonoBuffer, AlignedFloatBuffer& DestinationBuffer, const float* RESTRICT StartGains, const float* RESTRICT EndGains)
+	{
+		MixMonoTo2ChannelsFast(MonoBuffer.GetData(), DestinationBuffer.GetData(), DestinationBuffer.Num() / 2, StartGains, EndGains);
+	}
+
+	/**
+	* See CHANNEL MIXING OPERATIONS above for more info.
+	* 2 frames per iteration:
+	* +------------+---------+---------+---------+---------+
+	* | VectorName | Index 0 | Index 1 | Index 2 | Index 3 |
+	* +------------+---------+---------+---------+---------+
+	* | Gain       | g0      | g1      | g0      | g1      |
+	* |            | *       | *       | *       | *       |
+	* | Input      | i0      | i0      | i1      | i1      |
+	* |            | =       | =       | =       | =       |
+	* | Output     | o0      | o1      | o2      | o3      |
+	* +------------+---------+---------+---------+---------+
+	*/
+	void MixMonoTo2ChannelsFast(const float* RESTRICT MonoBuffer, float* RESTRICT DestinationBuffer, int32 NumFrames, const float* RESTRICT StartGains, const float* RESTRICT EndGains)
+	{
+		// Initialize GainVector at StartGains and compute GainDeltasVector:
+		VectorRegister GainVector = VectorLoadFloat2(StartGains);
+		const VectorRegister DestinationVector = VectorLoadFloat2(EndGains);
+		const VectorRegister NumFramesVector = VectorSetFloat1(NumFrames / 2.0f);
+		const VectorRegister GainDeltasVector = VectorDivide(VectorSubtract(DestinationVector, GainVector), NumFramesVector);
+
+		// To help with stair stepping, we initialize the second frame in GainVector to be half a GainDeltas vector higher than the first frame.
+		const VectorRegister VectorOfHalf = VectorSet(0.5f, 0.5f, 1.0f, 1.0f);
+		const VectorRegister HalfOfDeltaVector = VectorMultiply(GainDeltasVector, VectorOfHalf);
+		GainVector = VectorAdd(GainVector, HalfOfDeltaVector);
+
+		for (int32 i = 0; i < NumFrames; i += 2)
+		{
+			VectorRegister Result = VectorSet(MonoBuffer[i], MonoBuffer[i], MonoBuffer[i + 1], MonoBuffer[i + 1]);
+			Result = VectorMultiply(Result, GainVector);
+			VectorStoreAligned(Result, &DestinationBuffer[i*2]);
+
+			GainVector = VectorAdd(GainVector, GainDeltasVector);
+		}
+	}
+
+	void MixMonoTo2ChannelsFast(const AlignedFloatBuffer& MonoBuffer, AlignedFloatBuffer& DestinationBuffer)
+	{
+		MixMonoTo2ChannelsFast(MonoBuffer.GetData(), DestinationBuffer.GetData(), MonoBuffer.Num());
+	}
+
+	/**
+	* See CHANNEL MIXING OPERATIONS above for more info.
+	* 2 frames per iteration:
+	* +------------+---------+---------+---------+---------+
+	* | VectorName | Index 0 | Index 1 | Index 2 | Index 3 |
+	* +------------+---------+---------+---------+---------+
+	* | Input      | i0      | i0      | i1      | i1      |
+	* |            | =       | =       | =       | =       |
+	* | Output     | o0      | o1      | o2      | o3      |
+	* +------------+---------+---------+---------+---------+
+	*/
+	void MixMonoTo2ChannelsFast(const float* RESTRICT MonoBuffer, float* RESTRICT DestinationBuffer, int32 InNumFrames)
+	{
+		checkf(InNumFrames >= 4, TEXT("Buffer must have atleast 4 elements."));
+		checkf(0 == (InNumFrames % 4), TEXT("Buffer length be a multiple of 4."));
+		checkf(IsAligned<const float*>(MonoBuffer, AUDIO_SIMD_FLOAT_ALIGNMENT), TEXT("Memory must be aligned to use vector operations."));
+		checkf(IsAligned<float*>(DestinationBuffer, AUDIO_SIMD_FLOAT_ALIGNMENT), TEXT("Memory must be aligned to use vector operations."));
+
+		int32 OutPos = 0;
+		for (int32 i = 0; i < InNumFrames; i += 4)
+		{
+			VectorRegister Input = VectorLoadAligned(&MonoBuffer[i]);
+			VectorRegister Output = VectorSwizzle(Input, 0, 0, 1, 1);
+			VectorStoreAligned(Output, &DestinationBuffer[OutPos]);
+			OutPos += 4;
+			Output = VectorSwizzle(Input, 2, 2, 3, 3);
+			VectorStoreAligned(Output, &DestinationBuffer[OutPos]);
+			OutPos += 4;
+		}
+	}
+
+	void Mix2ChannelsTo2ChannelsFast(const AlignedFloatBuffer& SourceBuffer, AlignedFloatBuffer& DestinationBuffer, const float* RESTRICT Gains)
+	{
+		Mix2ChannelsTo2ChannelsFast(SourceBuffer.GetData(), DestinationBuffer.GetData(), DestinationBuffer.Num() / 2, Gains);
+	}
+
+	/**
+	* See CHANNEL MIXING OPERATIONS above for more info.
+	* 2 frames per iteration:
+	* +------------+---------+---------+---------+---------+
+	* | VectorName | Index 0 | Index 1 | Index 2 | Index 3 |
+	* +------------+---------+---------+---------+---------+
+	* | Gain1      | g0      | g1      | g0      | g1      |
+	* |            | *       | *       | *       | *       |
+	* | Input1     | i0      | i0      | i2      | i2      |
+	* |            | +       | +       | +       | +       |
+	* | Gain2      | g2      | g3      | g2      | g3      |
+	* |            | *       | *       | *       | *       |
+	* | Input2     | i1      | i1      | i3      | i3      |
+	* |            | =       | =       | =       | =       |
+	* | Output     | o0      | o1      | o2      | o3      |
+	* +------------+---------+---------+---------+---------+
+	*/
+
+	void Mix2ChannelsTo2ChannelsFast(const float* RESTRICT SourceBuffer, float* RESTRICT DestinationBuffer, int32 NumFrames, const float* RESTRICT Gains)
+	{
+		const VectorRegister GainVector1 = VectorLoadFloat2(Gains);
+		const VectorRegister GainVector2 = VectorLoadFloat2(Gains + 2);
+
+		for (int32 i = 0; i < NumFrames; i += 2)
+		{
+			const VectorRegister Input1 = VectorSet(SourceBuffer[i * 2], SourceBuffer[i * 2], SourceBuffer[i * 2 + 2], SourceBuffer[i * 2 + 2]);
+			const VectorRegister Input2 = VectorSet(SourceBuffer[i * 2 + 1], SourceBuffer[i * 2 + 1], SourceBuffer[i * 2 + 3], SourceBuffer[i * 2 + 3]);
+
+			VectorRegister Result = VectorMultiply(Input1, GainVector1);
+			Result = VectorMultiplyAdd(Input2, GainVector2, Result);
+
+			VectorStoreAligned(Result, &DestinationBuffer[i * 2]);
+		}
+	}
+
+	void Mix2ChannelsTo2ChannelsFast(const AlignedFloatBuffer& SourceBuffer, AlignedFloatBuffer& DestinationBuffer, const float* RESTRICT StartGains, const float* RESTRICT EndGains)
+	{
+		Mix2ChannelsTo2ChannelsFast(SourceBuffer.GetData(), DestinationBuffer.GetData(), DestinationBuffer.Num() / 2, StartGains, EndGains);
+	}
+
+	/**
+	* See CHANNEL MIXING OPERATIONS above for more info.
+	* 2 frames per iteration:
+	* +------------+---------+---------+---------+---------+
+	* | VectorName | Index 0 | Index 1 | Index 2 | Index 3 |
+	* +------------+---------+---------+---------+---------+
+	* | Gain1      | g0      | g1      | g0      | g1      |
+	* |            | *       | *       | *       | *       |
+	* | Input1     | i0      | i0      | i2      | i2      |
+	* |            | +       | +       | +       | +       |
+	* | Gain2      | g2      | g3      | g2      | g3      |
+	* |            | *       | *       | *       | *       |
+	* | Input2     | i1      | i1      | i3      | i3      |
+	* |            | =       | =       | =       | =       |
+	* | Output     | o0      | o1      | o2      | o3      |
+	* +------------+---------+---------+---------+---------+
+	*/
+
+	void Mix2ChannelsTo2ChannelsFast(const float* RESTRICT SourceBuffer, float* RESTRICT DestinationBuffer, int32 NumFrames, const float* RESTRICT StartGains, const float* RESTRICT EndGains)
+	{
+		const VectorRegister NumFramesVector = VectorSetFloat1(NumFrames / 2.0f);
+
+		VectorRegister GainVector1 = VectorLoadFloat2(StartGains);
+		const VectorRegister DestinationVector1 = VectorLoadFloat2(EndGains);
+		const VectorRegister GainDeltasVector1 = VectorDivide(VectorSubtract(DestinationVector1, GainVector1), NumFramesVector);
+
+		// To help with stair stepping, we initialize the second frame in GainVector to be half a GainDeltas vector higher than the first frame.
+		const VectorRegister VectorOfHalf = VectorSet(0.5f, 0.5f, 1.0f, 1.0f);
+
+		const VectorRegister HalfOfDeltaVector1 = VectorMultiply(GainDeltasVector1, VectorOfHalf);
+		GainVector1 = VectorAdd(GainVector1, HalfOfDeltaVector1);
+
+		VectorRegister GainVector2 = VectorLoadFloat2(StartGains + 2);
+		const VectorRegister DestinationVector2 = VectorLoadFloat2(EndGains + 2);
+		const VectorRegister GainDeltasVector2 = VectorDivide(VectorSubtract(DestinationVector2, GainVector2), NumFramesVector);
+
+		const VectorRegister HalfOfDeltaVector2 = VectorMultiply(GainDeltasVector2, VectorOfHalf);
+		GainVector2 = VectorAdd(GainVector2, HalfOfDeltaVector2);
+
+		for (int32 i = 0; i < NumFrames; i += 2)
+		{
+			const VectorRegister Input1 = VectorSet(SourceBuffer[i * 2], SourceBuffer[i * 2], SourceBuffer[i * 2 + 2], SourceBuffer[i * 2 + 2]);
+			const VectorRegister Input2 = VectorSet(SourceBuffer[i * 2 + 1], SourceBuffer[i * 2 + 1], SourceBuffer[i * 2 + 3], SourceBuffer[i * 2 + 3]);
+
+			VectorRegister Result = VectorMultiply(Input1, GainVector1);
+			Result = VectorMultiplyAdd(Input2, GainVector2, Result);
+
+			VectorStoreAligned(Result, &DestinationBuffer[i * 2]);
+
+			GainVector1 = VectorAdd(GainVector1, GainDeltasVector1);
+			GainVector2 = VectorAdd(GainVector2, GainDeltasVector2);
+		}
+	}
+
+	void Apply4ChannelGain(AlignedFloatBuffer& InterleavedBuffer, const float* RESTRICT Gains)
+	{
+		Apply4ChannelGain(InterleavedBuffer.GetData(), InterleavedBuffer.Num(), Gains);
+	}
+
+	void Apply4ChannelGain(float* RESTRICT InterleavedBuffer, int32 NumSamples, const float* RESTRICT Gains)
+	{
+		const VectorRegister GainVector = VectorLoadAligned(Gains);
+
+		for (int32 i = 0; i < NumSamples; i += 4)
+		{
+			VectorRegister Result = VectorLoadAligned(&InterleavedBuffer[i]);
+			Result = VectorMultiply(Result, GainVector);
+			VectorStoreAligned(Result, &InterleavedBuffer[i]);
+		}
+	}
+
+	void Apply4ChannelGain(AlignedFloatBuffer& InterleavedBuffer, const float* RESTRICT StartGains, const float* RESTRICT EndGains)
+	{
+		Apply4ChannelGain(InterleavedBuffer.GetData(), InterleavedBuffer.Num(), StartGains, EndGains);
+	}
+
+	void Apply4ChannelGain(float* RESTRICT InterleavedBuffer, int32 NumSamples, const float* RESTRICT StartGains, const float* RESTRICT EndGains)
+	{
+		// Initialize GainVector at StartGains and compute GainDeltasVector:
+		VectorRegister GainVector = VectorLoadAligned(StartGains);
+		const VectorRegister DestinationVector = VectorLoadAligned(EndGains);
+		const VectorRegister NumFramesVector = VectorSetFloat1(NumSamples / 4.0f);
+		const VectorRegister GainDeltasVector = VectorDivide(VectorSubtract(DestinationVector, GainVector), NumFramesVector);
+
+		for (int32 i = 0; i < NumSamples; i += 4)
+		{
+			VectorRegister Result = VectorLoadAligned(&InterleavedBuffer[i]);
+			Result = VectorMultiply(Result, GainVector);
+			VectorStoreAligned(Result, &InterleavedBuffer[i]);
+
+			GainVector = VectorAdd(GainVector, GainDeltasVector);
+		}
+	}
+
+	void MixMonoTo4ChannelsFast(const AlignedFloatBuffer& MonoBuffer, AlignedFloatBuffer& DestinationBuffer, const float* RESTRICT Gains)
+	{
+		MixMonoTo4ChannelsFast(MonoBuffer.GetData(), DestinationBuffer.GetData(), DestinationBuffer.Num() / 4, Gains);
+	}
+
+	/**
+	* See CHANNEL MIXING OPERATIONS above for more info.
+	* 1 frame per iteration:
+	* +------------+---------+---------+---------+---------+
+	* | VectorName | Index 0 | Index 1 | Index 2 | Index 3 |
+	* +------------+---------+---------+---------+---------+
+	* | Gain       | g0      | g1      | g2      | g3      |
+	* |            | *       | *       | *       | *       |
+	* | Input      | i0      | i0      | i0      | i0      |
+	* |            | =       | =       | =       | =       |
+	* | Output     | o0      | o1      | o2      | o3      |
+	* +------------+---------+---------+---------+---------+
+	*/
+
+	void MixMonoTo4ChannelsFast(const float* RESTRICT MonoBuffer, float* RESTRICT DestinationBuffer, int32 NumFrames, const float* RESTRICT Gains)
+	{
+		const VectorRegister GainVector = VectorLoadAligned(Gains);
+
+		for (int32 i = 0; i < NumFrames; i++)
+		{
+			VectorRegister Result = VectorLoadFloat1(&MonoBuffer[i]);
+			Result = VectorMultiply(Result, GainVector);
+			VectorStoreAligned(Result, &DestinationBuffer[i*4]);
+		}
+	}
+
+	void MixMonoTo4ChannelsFast(const AlignedFloatBuffer& MonoBuffer, AlignedFloatBuffer& DestinationBuffer, const float* RESTRICT StartGains, const float* RESTRICT EndGains)
+	{
+		MixMonoTo4ChannelsFast(MonoBuffer.GetData(), DestinationBuffer.GetData(), DestinationBuffer.Num() / 4, StartGains, EndGains);
+	}
+
+
+	/**
+	* See CHANNEL MIXING OPERATIONS above for more info.
+	* 1 frame per iteration:
+	* +------------+---------+---------+---------+---------+
+	* | VectorName | Index 0 | Index 1 | Index 2 | Index 3 |
+	* +------------+---------+---------+---------+---------+
+	* | Gain       | g0      | g1      | g2      | g3      |
+	* |            | *       | *       | *       | *       |
+	* | Input      | i0      | i0      | i0      | i0      |
+	* |            | =       | =       | =       | =       |
+	* | Output     | o0      | o1      | o2      | o3      |
+	* +------------+---------+---------+---------+---------+
+	*/
+	void MixMonoTo4ChannelsFast(const float* RESTRICT MonoBuffer, float* RESTRICT DestinationBuffer, int32 NumFrames, const float* RESTRICT StartGains, const float* RESTRICT EndGains)
+	{
+		VectorRegister GainVector = VectorLoadAligned(StartGains);
+		const VectorRegister DestinationVector = VectorLoadAligned(EndGains);
+		const VectorRegister NumFramesVector = VectorSetFloat1((float) NumFrames);
+		const VectorRegister GainDeltasVector = VectorDivide(VectorSubtract(DestinationVector, GainVector), NumFramesVector);
+
+		for (int32 i = 0; i < NumFrames; i++)
+		{
+			VectorRegister Result = VectorLoadFloat1(&MonoBuffer[i]);
+			Result = VectorMultiply(Result, GainVector);
+			VectorStoreAligned(Result, &DestinationBuffer[i * 4]);
+
+			GainVector = VectorAdd(GainVector, GainDeltasVector);
+		}
+	}
+
+	void Mix2ChannelsTo4ChannelsFast(const AlignedFloatBuffer& SourceBuffer, AlignedFloatBuffer& DestinationBuffer, const float* RESTRICT Gains)
+	{
+		Mix2ChannelsTo4ChannelsFast(SourceBuffer.GetData(), DestinationBuffer.GetData(), DestinationBuffer.Num() / 4, Gains);
+	}
+
+	/**
+	* See CHANNEL MIXING OPERATIONS above for more info.
+	* 1 frame per iteration:
+	* +------------+---------+---------+---------+---------+
+	* | VectorName | Index 0 | Index 1 | Index 2 | Index 3 |
+	* +------------+---------+---------+---------+---------+
+	* | Gain1      | g0      | g1      | g2      | g3      |
+	* |            | *       | *       | *       | *       |
+	* | Input1     | i0      | i0      | i0      | i0      |
+	* |            | +       | +       | +       | +       |
+	* | Gain2      | g4      | g5      | g6      | g7      |
+	* |            | *       | *       | *       | *       |
+	* | Input2     | i1      | i1      | i1      | i1      |
+	* |            | =       | =       | =       | =       |
+	* | Output     | o0      | o1      | o2      | o3      |
+	* +------------+---------+---------+---------+---------+
+	*/
+
+	void Mix2ChannelsTo4ChannelsFast(const float* RESTRICT SourceBuffer, float* RESTRICT DestinationBuffer, int32 NumFrames, const float* RESTRICT Gains)
+	{
+		const VectorRegister GainVector1 = VectorLoadAligned(Gains);
+		const VectorRegister GainVector2 = VectorLoadAligned(Gains + 4);
+
+		for (int32 i = 0; i < NumFrames; i++)
+		{
+			const VectorRegister Input1 = VectorLoadFloat1(&SourceBuffer[i * 2]);
+			const VectorRegister Input2 = VectorLoadFloat1(&SourceBuffer[i * 2 + 1]);
+
+			VectorRegister Result = VectorMultiply(Input1, GainVector1);
+			Result = VectorMultiplyAdd(Input2, GainVector2, Result);
+			VectorStoreAligned(Result, &DestinationBuffer[i * 4]);
+		}
+	}
+
+	void Mix2ChannelsTo4ChannelsFast(const AlignedFloatBuffer& SourceBuffer, AlignedFloatBuffer& DestinationBuffer, const float* RESTRICT StartGains, const float* RESTRICT EndGains)
+	{
+		Mix2ChannelsTo4ChannelsFast(SourceBuffer.GetData(), DestinationBuffer.GetData(), DestinationBuffer.Num() / 4, StartGains, EndGains);
+	}
+
+	/**
+	* See CHANNEL MIXING OPERATIONS above for more info.
+	* 1 frame per iteration:
+	* +------------+---------+---------+---------+---------+
+	* | VectorName | Index 0 | Index 1 | Index 2 | Index 3 |
+	* +------------+---------+---------+---------+---------+
+	* | Gain1      | g0      | g1      | g2      | g3      |
+	* |            | *       | *       | *       | *       |
+	* | Input1     | i0      | i0      | i0      | i0      |
+	* |            | +       | +       | +       | +       |
+	* | Gain2      | g4      | g5      | g6      | g7      |
+	* |            | *       | *       | *       | *       |
+	* | Input2     | i1      | i1      | i1      | i1      |
+	* |            | =       | =       | =       | =       |
+	* | Output     | o0      | o1      | o2      | o3      |
+	* +------------+---------+---------+---------+---------+
+	*/
+	void Mix2ChannelsTo4ChannelsFast(const float* RESTRICT SourceBuffer, float* RESTRICT DestinationBuffer, int32 NumFrames, const float* RESTRICT StartGains, const float* RESTRICT EndGains)
+	{
+		const VectorRegister NumFramesVector = VectorSetFloat1((float) NumFrames);
+
+		VectorRegister GainVector1 = VectorLoadAligned(StartGains);
+		const VectorRegister DestinationVector1 = VectorLoadAligned(EndGains);
+		const VectorRegister GainDeltasVector1 = VectorDivide(VectorSubtract(DestinationVector1, GainVector1), NumFramesVector);
+
+		VectorRegister GainVector2 = VectorLoadAligned(StartGains + 4);
+		const VectorRegister DestinationVector2 = VectorLoadAligned(EndGains + 4);
+		const VectorRegister GainDeltasVector2 = VectorDivide(VectorSubtract(DestinationVector2, GainVector2), NumFramesVector);
+
+		for (int32 i = 0; i < NumFrames; i++)
+		{
+			const VectorRegister Input1 = VectorLoadFloat1(&SourceBuffer[i * 2]);
+			const VectorRegister Input2 = VectorLoadFloat1(&SourceBuffer[i * 2 + 1]);
+
+			VectorRegister Result = VectorMultiply(Input1, GainVector1);
+			Result = VectorMultiplyAdd(Input2, GainVector2, Result);
+			VectorStoreAligned(Result, &DestinationBuffer[i * 4]);
+
+			GainVector1 = VectorAdd(GainVector1, GainDeltasVector1);
+			GainVector2 = VectorAdd(GainVector2, GainDeltasVector2);
+		}
+	}
+
+	void Apply6ChannelGain(AlignedFloatBuffer& InterleavedBuffer, const float* RESTRICT Gains)
+	{
+		Apply6ChannelGain(InterleavedBuffer.GetData(), InterleavedBuffer.Num(), Gains);
+	}
+
+	void Apply6ChannelGain(float* RESTRICT InterleavedBuffer, int32 NumSamples, const float* RESTRICT Gains)
+	{
+		const VectorRegister GainVector1 = VectorLoadAligned(Gains);
+		const VectorRegister GainVector2 = VectorSet(Gains[4], Gains[5], Gains[0], Gains[1]);
+		const VectorRegister GainVector3 = VectorLoad(&Gains[2]);
+
+		for (int32 i = 0; i < NumSamples; i += 12)
+		{
+			VectorRegister Result = VectorLoadAligned(&InterleavedBuffer[i]);
+			Result = VectorMultiply(Result, GainVector1);
+			VectorStoreAligned(Result, &InterleavedBuffer[i]);
+
+			Result = VectorLoadAligned(&InterleavedBuffer[i + 4]);
+			Result = VectorMultiply(Result, GainVector2);
+			VectorStoreAligned(Result, &InterleavedBuffer[i + 4]);
+
+			Result = VectorLoadAligned(&InterleavedBuffer[i + 8]);
+			Result = VectorMultiply(Result, GainVector3);
+			VectorStoreAligned(Result, &InterleavedBuffer[i + 8]);
+		}
+	}
+
+	void Apply6ChannelGain(AlignedFloatBuffer& InterleavedBuffer, const float* RESTRICT StartGains, const float* RESTRICT EndGains)
+	{
+		Apply6ChannelGain(InterleavedBuffer.GetData(), InterleavedBuffer.Num(), StartGains, EndGains);
+	}
+
+	void Apply6ChannelGain(float* RESTRICT InterleavedBuffer, int32 NumSamples, const float* RESTRICT StartGains, const float* RESTRICT EndGains)
+	{
+		const VectorRegister NumFramesVector = VectorSetFloat1(NumSamples / 12.0f);
+
+		VectorRegister GainVector1 = VectorLoadAligned(StartGains);
+		const VectorRegister DestinationVector1 = VectorLoadAligned(EndGains);
+		const VectorRegister GainDeltasVector1 = VectorDivide(VectorSubtract(DestinationVector1, GainVector1), NumFramesVector);
+
+		VectorRegister GainVector2 = VectorSet(StartGains[4], StartGains[5], StartGains[0], StartGains[1]);
+		const VectorRegister DestinationVector2 = VectorSet(EndGains[4], EndGains[5], EndGains[0], EndGains[1]);
+		const VectorRegister GainDeltasVector2 = VectorDivide(VectorSubtract(DestinationVector2, GainVector2), NumFramesVector);
+
+		VectorRegister GainVector3 = VectorLoad(&StartGains[2]);
+		const VectorRegister DestinationVector3 = VectorLoad(&EndGains[2]);
+		const VectorRegister GainDeltasVector3 = VectorDivide(VectorSubtract(DestinationVector3, GainVector3), NumFramesVector);
+
+		for (int32 i = 0; i < NumSamples; i += 12)
+		{
+			VectorRegister Result = VectorLoadAligned(&InterleavedBuffer[i]);
+			Result = VectorMultiply(Result, GainVector1);
+			VectorStoreAligned(Result, &InterleavedBuffer[i]);
+
+			GainVector1 = VectorAdd(GainVector1, GainDeltasVector1);
+
+			Result = VectorLoadAligned(&InterleavedBuffer[i + 4]);
+			Result = VectorMultiply(Result, GainVector2);
+			VectorStoreAligned(Result, &InterleavedBuffer[i + 4]);
+
+			GainVector2 = VectorAdd(GainVector2, GainDeltasVector2);
+
+			Result = VectorLoadAligned(&InterleavedBuffer[i + 8]);
+			Result = VectorMultiply(Result, GainVector3);
+			VectorStoreAligned(Result, &InterleavedBuffer[i + 8]);
+
+			GainVector3 = VectorAdd(GainVector3, GainDeltasVector3);
+		}
+	}
+
+	void MixMonoTo6ChannelsFast(const AlignedFloatBuffer& MonoBuffer, AlignedFloatBuffer& DestinationBuffer, const float* RESTRICT Gains)
+	{
+		MixMonoTo6ChannelsFast(MonoBuffer.GetData(), DestinationBuffer.GetData(), DestinationBuffer.Num() / 6, Gains);
+	}
+
+	/**
+	* See CHANNEL MIXING OPERATIONS above for more info.
+	* 2 frames per iteration:
+	* +------------+----------+---------+---------+---------+----------+---------+---------+---------+----------+---------+----------+----------+
+	* |            | Vector 1 |         |         |         | Vector 2 |         |         |         | Vector 3 |         |          |          |
+	* +------------+----------+---------+---------+---------+----------+---------+---------+---------+----------+---------+----------+----------+
+	* |            | Index 0  | Index 1 | Index 2 | Index 3 | Index 4  | Index 5 | Index 6 | Index 7 | Index 8  | Index 9 | Index 10 | Index 11 |
+	* | Gain       | g0       | g1      | g2      | g3      | g4       | g5      | g0      | g1      | g2       | g3      | g4       | g5       |
+	* |            | *        | *       | *       | *       | *        | *       | *       | *       | *        | *       | *        | *        |
+	* | Input      | i0       | i0      | i0      | i0      | i0       | i0      | i1      | i1      | i1       | i1      | i1       | i1       |
+	* |            | =        | =       | =       | =       | =        | =       | =       | =       | =        | =       | =        | =        |
+	* | Output     | o0       | o1      | o2      | o3      | o4       | o5      | o6      | o7      | o8       | o9      | o10      | o11      |
+	* +------------+----------+---------+---------+---------+----------+---------+---------+---------+----------+---------+----------+----------+
+	*/
+
+	void MixMonoTo6ChannelsFast(const float* RESTRICT MonoBuffer, float* RESTRICT DestinationBuffer, int32 NumFrames, const float* RESTRICT Gains)
+	{
+		const VectorRegister GainVector1 = VectorLoadAligned(Gains);
+		const VectorRegister GainVector2 = VectorSet(Gains[4], Gains[5], Gains[0], Gains[1]);
+		const VectorRegister GainVector3 = VectorLoad(&Gains[2]);
+
+		for (int32 i = 0; i < NumFrames; i += 2)
+		{
+			const VectorRegister Input1 = VectorLoadFloat1(&MonoBuffer[i]);
+			const VectorRegister Input2 = VectorSet(MonoBuffer[i], MonoBuffer[i], MonoBuffer[i + 1], MonoBuffer[i + 1]);
+			const VectorRegister Input3 = VectorLoadFloat1(&MonoBuffer[i + 1]);
+
+			VectorRegister Result = VectorMultiply(Input1, GainVector1);
+			VectorStoreAligned(Result, &DestinationBuffer[i * 6]);
+
+			Result = VectorMultiply(Input2, GainVector2);
+			VectorStoreAligned(Result, &DestinationBuffer[i * 6 + 4]);
+
+			Result = VectorMultiply(Input3, GainVector3);
+			VectorStoreAligned(Result, &DestinationBuffer[i * 6 + 8]);
+		}
+	}
+
+	void MixMonoTo6ChannelsFast(const AlignedFloatBuffer& MonoBuffer, AlignedFloatBuffer& DestinationBuffer, const float* RESTRICT StartGains, const float* RESTRICT EndGains)
+	{
+		MixMonoTo6ChannelsFast(MonoBuffer.GetData(), DestinationBuffer.GetData(), DestinationBuffer.Num() / 6, StartGains, EndGains);
+	}
+
+	/**
+	* See CHANNEL MIXING OPERATIONS above for more info.
+	* 2 frames per iteration:
+	* +------------+----------+---------+---------+---------+----------+---------+---------+---------+----------+---------+----------+----------+
+	* |            | Vector 1 |         |         |         | Vector 2 |         |         |         | Vector 3 |         |          |          |
+	* +------------+----------+---------+---------+---------+----------+---------+---------+---------+----------+---------+----------+----------+
+	* |            | Index 0  | Index 1 | Index 2 | Index 3 | Index 4  | Index 5 | Index 6 | Index 7 | Index 8  | Index 9 | Index 10 | Index 11 |
+	* | Gain       | g0       | g1      | g2      | g3      | g4       | g5      | g0      | g1      | g2       | g3      | g4       | g5       |
+	* |            | *        | *       | *       | *       | *        | *       | *       | *       | *        | *       | *        | *        |
+	* | Input      | i0       | i0      | i0      | i0      | i0       | i0      | i1      | i1      | i1       | i1      | i1       | i1       |
+	* |            | =        | =       | =       | =       | =        | =       | =       | =       | =        | =       | =        | =        |
+	* | Output     | o0       | o1      | o2      | o3      | o4       | o5      | o6      | o7      | o8       | o9      | o10      | o11      |
+	* +------------+----------+---------+---------+---------+----------+---------+---------+---------+----------+---------+----------+----------+
+	*/
+
+	void MixMonoTo6ChannelsFast(const float* RESTRICT MonoBuffer, float* RESTRICT DestinationBuffer, int32 NumFrames, const float* RESTRICT StartGains, const float* RESTRICT EndGains)
+	{
+		const VectorRegister NumFramesVector = VectorSetFloat1(NumFrames / 2.0f);
+
+		VectorRegister GainVector1 = VectorLoadAligned(StartGains);
+		const VectorRegister DestinationVector1 = VectorLoadAligned(EndGains);
+		const VectorRegister GainDeltasVector1 = VectorDivide(VectorSubtract(DestinationVector1, GainVector1), NumFramesVector);
+
+		VectorRegister GainVector2 = VectorSet(StartGains[4], StartGains[5], StartGains[0], StartGains[1]);
+		const VectorRegister DestinationVector2 = VectorSet(EndGains[4], EndGains[5], EndGains[0], EndGains[1]);
+		const VectorRegister GainDeltasVector2 = VectorDivide(VectorSubtract(DestinationVector2, GainVector2), NumFramesVector);
+
+		VectorRegister GainVector3 = VectorLoad(&StartGains[2]);
+		const VectorRegister DestinationVector3 = VectorLoad(&EndGains[2]);
+		const VectorRegister GainDeltasVector3 = VectorDivide(VectorSubtract(DestinationVector3, GainVector3), NumFramesVector);
+
+		for (int32 i = 0; i < NumFrames; i += 2)
+		{
+			const VectorRegister Input1 = VectorLoadFloat1(&MonoBuffer[i]);
+			const VectorRegister Input2 = VectorSet(MonoBuffer[i], MonoBuffer[i], MonoBuffer[i + 1], MonoBuffer[i + 1]);
+			const VectorRegister Input3 = VectorLoadFloat1(&MonoBuffer[i + 1]);
+
+			VectorRegister Result = VectorMultiply(Input1, GainVector1);
+			VectorStoreAligned(Result, &DestinationBuffer[i * 6]);
+
+			GainVector1 = VectorAdd(GainVector1, GainDeltasVector1);
+
+			Result = VectorMultiply(Input2, GainVector2);
+			VectorStoreAligned(Result, &DestinationBuffer[i * 6 + 4]);
+
+			GainVector2 = VectorAdd(GainVector2, GainDeltasVector2);
+
+			Result = VectorMultiply(Input3, GainVector3);
+			VectorStoreAligned(Result, &DestinationBuffer[i * 6 + 8]);
+
+			GainVector3 = VectorAdd(GainVector3, GainDeltasVector3);
+		}
+	}
+
+	void Mix2ChannelsTo6ChannelsFast(const AlignedFloatBuffer& SourceBuffer, AlignedFloatBuffer& DestinationBuffer, const float* RESTRICT Gains)
+	{
+		Mix2ChannelsTo6ChannelsFast(SourceBuffer.GetData(), DestinationBuffer.GetData(), DestinationBuffer.Num() / 6, Gains);
+	}
+
+	/**
+	* See CHANNEL MIXING OPERATIONS above for more info.
+	* 2 frames per iteration:
+	* +--------+----------+---------+---------+---------+----------+---------+---------+---------+----------+---------+----------+----------+
+	* |        | Vector 1 |         |         |         | Vector 2 |         |         |         | Vector 3 |         |          |          |
+	* +--------+----------+---------+---------+---------+----------+---------+---------+---------+----------+---------+----------+----------+
+	* |        | Index 0  | Index 1 | Index 2 | Index 3 | Index 4  | Index 5 | Index 6 | Index 7 | Index 8  | Index 9 | Index 10 | Index 11 |
+	* | Gain1  | g0       | g1      | g2      | g3      | g4       | g5      | g0      | g1      | g2       | g3      | g4       | g5       |
+	* |        | *        | *       | *       | *       | *        | *       | *       | *       | *        | *       | *        | *        |
+	* | Input1 | i0       | i0      | i0      | i0      | i0       | i0      | i2      | i2      | i2       | i2      | i2       | i2       |
+	* |        | +        | +       | +       | +       | +        | +       | +       | +       | +        | +       | +        | +        |
+	* | Gain2  | g6       | g7      | g8      | g9      | g10      | g11     | g6      | g7      | g8       | g9      | g10      | g11      |
+	* |        | *        | *       | *       | *       | *        | *       | *       | *       | *        | *       | *        | *        |
+	* | Input2 | i1       | i1      | i1      | i1      | i1       | i1      | i1      | i3      | i3       | i3      | i3       | i3       |
+	* |        | =        | =       | =       | =       | =        | =       | =       | =       | =        | =       | =        | =        |
+	* | Output | o0       | o1      | o2      | o3      | o4       | o5      | o6      | o7      | o8       | o9      | o10      | o11      |
+	* +--------+----------+---------+---------+---------+----------+---------+---------+---------+----------+---------+----------+----------+
+	*/
+
+	void Mix2ChannelsTo6ChannelsFast(const float* RESTRICT SourceBuffer, float* RESTRICT DestinationBuffer, int32 NumFrames, const float* RESTRICT Gains)
+	{
+		const VectorRegister GainVector11 = VectorLoadAligned(Gains);
+		const VectorRegister GainVector21 = VectorSet(Gains[4], Gains[5], Gains[0], Gains[1]);
+		const VectorRegister GainVector31 = VectorLoad(&Gains[2]);
+
+		const VectorRegister GainVector12 = VectorLoad(Gains + 6);
+		const VectorRegister GainVector22 = VectorSet(Gains[10], Gains[11], Gains[6], Gains[7]);
+		const VectorRegister GainVector32 = VectorLoadAligned(&Gains[8]);
+
+		for (int32 FrameIndex = 0; FrameIndex < NumFrames; FrameIndex += 2)
+		{
+			const int32 InputIndex = FrameIndex * 2;
+			const int32 OutputIndex = FrameIndex * 6;
+
+			const VectorRegister Input11 = VectorLoadFloat1(&SourceBuffer[InputIndex]);
+			const VectorRegister Input21 = VectorSet(SourceBuffer[InputIndex], SourceBuffer[InputIndex], SourceBuffer[InputIndex + 2], SourceBuffer[InputIndex + 2]);
+			const VectorRegister Input31 = VectorLoadFloat1(&SourceBuffer[InputIndex + 2]);
+
+			const VectorRegister Input12 = VectorLoadFloat1(&SourceBuffer[InputIndex + 1]);
+			const VectorRegister Input22 = VectorSet(SourceBuffer[InputIndex + 1], SourceBuffer[InputIndex + 1], SourceBuffer[InputIndex + 3], SourceBuffer[InputIndex + 3]);
+			const VectorRegister Input32 = VectorLoadFloat1(&SourceBuffer[InputIndex + 3]);
+
+			VectorRegister Result = VectorMultiply(Input11, GainVector11);
+			Result = VectorMultiplyAdd(Input12, GainVector12, Result);
+			VectorStoreAligned(Result, &DestinationBuffer[OutputIndex]);
+
+			Result = VectorMultiply(Input21, GainVector21);
+			Result = VectorMultiplyAdd(Input22, GainVector22, Result);
+			VectorStoreAligned(Result, &DestinationBuffer[OutputIndex + 4]);
+
+			Result = VectorMultiply(Input31, GainVector31);
+			Result = VectorMultiplyAdd(Input32, GainVector31, Result);
+			VectorStoreAligned(Result, &DestinationBuffer[OutputIndex + 8]);
+		}
+	}
+
+	void Mix2ChannelsTo6ChannelsFast(const AlignedFloatBuffer& SourceBuffer, AlignedFloatBuffer& DestinationBuffer, const float* RESTRICT StartGains, const float* RESTRICT EndGains)
+	{
+		Mix2ChannelsTo6ChannelsFast(SourceBuffer.GetData(), DestinationBuffer.GetData(), DestinationBuffer.Num() / 6, StartGains, EndGains);
+	}
+
+	/**
+	* See CHANNEL MIXING OPERATIONS above for more info.
+	* 2 frames per iteration:
+	* +--------+----------+---------+---------+---------+----------+---------+---------+---------+----------+---------+----------+----------+
+	* |        | Vector 1 |         |         |         | Vector 2 |         |         |         | Vector 3 |         |          |          |
+	* +--------+----------+---------+---------+---------+----------+---------+---------+---------+----------+---------+----------+----------+
+	* |        | Index 0  | Index 1 | Index 2 | Index 3 | Index 4  | Index 5 | Index 6 | Index 7 | Index 8  | Index 9 | Index 10 | Index 11 |
+	* | Gain1  | g0       | g1      | g2      | g3      | g4       | g5      | g0      | g1      | g2       | g3      | g4       | g5       |
+	* |        | *        | *       | *       | *       | *        | *       | *       | *       | *        | *       | *        | *        |
+	* | Input1 | i0       | i0      | i0      | i0      | i0       | i0      | i2      | i2      | i2       | i2      | i2       | i2       |
+	* |        | +        | +       | +       | +       | +        | +       | +       | +       | +        | +       | +        | +        |
+	* | Gain2  | g6       | g7      | g8      | g9      | g10      | g11     | g6      | g7      | g8       | g9      | g10      | g11      |
+	* |        | *        | *       | *       | *       | *        | *       | *       | *       | *        | *       | *        | *        |
+	* | Input2 | i1       | i1      | i1      | i1      | i1       | i1      | i3      | i3      | i3       | i3      | i3       | i3       |
+	* |        | =        | =       | =       | =       | =        | =       | =       | =       | =        | =       | =        | =        |
+	* | Output | o0       | o1      | o2      | o3      | o4       | o5      | o6      | o7      | o8       | o9      | o10      | o11      |
+	* +--------+----------+---------+---------+---------+----------+---------+---------+---------+----------+---------+----------+----------+
+	*/
+
+	void Mix2ChannelsTo6ChannelsFast(const float* RESTRICT SourceBuffer, float* RESTRICT DestinationBuffer, int32 NumFrames, const float* RESTRICT StartGains, const float* RESTRICT EndGains)
+	{
+		const VectorRegister NumFramesVector = VectorSetFloat1(NumFrames / 2.0f);
+
+		VectorRegister GainVector11 = VectorLoadAligned(StartGains);
+		const VectorRegister DestinationVector11 = VectorLoadAligned(EndGains);
+		const VectorRegister GainDeltasVector11 = VectorDivide(VectorSubtract(DestinationVector11, GainVector11), NumFramesVector);
+
+		VectorRegister GainVector21 = VectorSet(StartGains[4], StartGains[5], StartGains[0], StartGains[1]);
+		const VectorRegister DestinationVector21 = VectorSet(EndGains[4], EndGains[5], EndGains[0], EndGains[1]);
+		const VectorRegister GainDeltasVector21 = VectorDivide(VectorSubtract(DestinationVector21, GainVector21), NumFramesVector);
+
+		// In order to ease stair stepping, we ensure that the second frame is initialized to half the GainDelta more than the first frame.
+		// This gives us a consistent increment across every frame.
+		const VectorRegister DeltaHalf21 = VectorSet(0.0f, 0.0f, 0.5f, 0.5f);
+		const VectorRegister InitializedDelta21 = VectorMultiply(GainDeltasVector21, DeltaHalf21);
+		GainVector21 = VectorAdd(GainVector21, InitializedDelta21);
+
+		VectorRegister GainVector31 = VectorLoad(&StartGains[2]);
+		const VectorRegister DestinationVector31 = VectorLoad(&EndGains[2]);
+		const VectorRegister GainDeltasVector31 = VectorDivide(VectorSubtract(DestinationVector31, GainVector31), NumFramesVector);
+
+		const VectorRegister DeltaHalf31 = VectorSetFloat1(0.5f);
+		const VectorRegister InitializedDelta31 = VectorMultiply(GainDeltasVector31, DeltaHalf31);
+		GainVector31 = VectorAdd(GainVector31, InitializedDelta31);
+
+		VectorRegister GainVector12 = VectorLoad(StartGains + 6);
+		const VectorRegister DestinationVector12 = VectorLoad(EndGains + 6);
+		const VectorRegister GainDeltasVector12 = VectorDivide(VectorSubtract(DestinationVector12, GainVector12), NumFramesVector);
+
+		VectorRegister GainVector22 = VectorSet(StartGains[10], StartGains[11], StartGains[6], StartGains[7]);
+		const VectorRegister DestinationVector22 = VectorSet(EndGains[10], EndGains[11], EndGains[6], EndGains[7]);
+		const VectorRegister GainDeltasVector22 = VectorDivide(VectorSubtract(DestinationVector22, GainVector22), NumFramesVector);
+
+		const VectorRegister DeltaHalf22 = VectorSet(0.0f, 0.0f, 0.5f, 0.5f);
+		const VectorRegister InitializedDelta22 = VectorMultiply(GainDeltasVector22, DeltaHalf22);
+		GainVector22 = VectorAdd(GainVector22, InitializedDelta22);
+
+		VectorRegister GainVector32 = VectorLoadAligned(StartGains + 8);
+		const VectorRegister DestinationVector32 = VectorLoadAligned(EndGains + 8);
+		const VectorRegister GainDeltasVector32 = VectorDivide(VectorSubtract(DestinationVector32, GainVector32), NumFramesVector);
+
+		const VectorRegister DeltaHalf32 = VectorSetFloat1(0.5f);
+		const VectorRegister InitializedDelta32 = VectorMultiply(GainDeltasVector32, DeltaHalf32);
+		GainVector32 = VectorAdd(GainVector32, InitializedDelta32);
+
+		for (int32 FrameIndex = 0; FrameIndex < NumFrames; FrameIndex += 2)
+		{
+			const int32 InputIndex = FrameIndex * 2;
+			const int32 OutputIndex = FrameIndex * 6;
+
+			const VectorRegister Input11 = VectorLoadFloat1(&SourceBuffer[InputIndex]);
+			const VectorRegister Input21 = VectorSet(SourceBuffer[InputIndex], SourceBuffer[InputIndex], SourceBuffer[InputIndex + 2], SourceBuffer[InputIndex + 2]);
+			const VectorRegister Input31 = VectorLoadFloat1(&SourceBuffer[InputIndex + 2]);
+
+			const VectorRegister Input12 = VectorLoadFloat1(&SourceBuffer[InputIndex + 1]);
+			const VectorRegister Input22 = VectorSet(SourceBuffer[InputIndex + 1], SourceBuffer[InputIndex + 1], SourceBuffer[InputIndex + 3], SourceBuffer[InputIndex + 3]);
+			const VectorRegister Input32 = VectorLoadFloat1(&SourceBuffer[InputIndex + 3]);
+
+			VectorRegister Result = VectorMultiply(Input11, GainVector11);
+			Result = VectorMultiplyAdd(Input12, GainVector12, Result);
+			VectorStoreAligned(Result, &DestinationBuffer[OutputIndex]);
+
+			GainVector11 = VectorAdd(GainVector11, GainDeltasVector11);
+			GainVector12 = VectorAdd(GainVector12, GainDeltasVector12);
+
+			Result = VectorMultiply(Input21, GainVector21);
+			Result = VectorMultiplyAdd(Input22, GainVector22, Result);
+			VectorStoreAligned(Result, &DestinationBuffer[OutputIndex + 4]);
+
+			GainVector21 = VectorAdd(GainVector21, GainDeltasVector21);
+			GainVector22 = VectorAdd(GainVector22, GainDeltasVector22);
+
+			Result = VectorMultiply(Input31, GainVector31);
+			Result = VectorMultiplyAdd(Input32, GainVector31, Result);
+			VectorStoreAligned(Result, &DestinationBuffer[OutputIndex + 8]);
+
+			GainVector31 = VectorAdd(GainVector31, GainDeltasVector31);
+			GainVector32 = VectorAdd(GainVector32, GainDeltasVector32);
+		}
+	}
+
+	void Apply8ChannelGain(AlignedFloatBuffer& InterleavedBuffer, const float* RESTRICT Gains)
+	{
+		Apply8ChannelGain(InterleavedBuffer.GetData(), InterleavedBuffer.Num(), Gains);
+	}
+
+	void Apply8ChannelGain(float* RESTRICT InterleavedBuffer, int32 NumSamples, const float* RESTRICT Gains)
+	{
+		const VectorRegister GainVector1 = VectorLoadAligned(Gains);
+		const VectorRegister GainVector2 = VectorLoadAligned(Gains + 4);
+
+		for (int32 i = 0; i < NumSamples; i += 8)
+		{
+			VectorRegister Result = VectorLoadAligned(&InterleavedBuffer[i]);
+			Result = VectorMultiply(Result, GainVector1);
+			VectorStoreAligned(Result, &InterleavedBuffer[i]);
+
+			Result = VectorLoadAligned(&InterleavedBuffer[i + 4]);
+			Result = VectorMultiply(Result, GainVector2);
+			VectorStoreAligned(Result, &InterleavedBuffer[i + 4]);
+		}
+	}
+
+	void Apply8ChannelGain(AlignedFloatBuffer& InterleavedBuffer, const float* RESTRICT StartGains, const float* RESTRICT EndGains)
+	{
+		Apply8ChannelGain(InterleavedBuffer.GetData(), InterleavedBuffer.Num(), StartGains, EndGains);
+	}
+
+	void Apply8ChannelGain(float* RESTRICT InterleavedBuffer, int32 NumSamples, const float* RESTRICT StartGains, const float* RESTRICT EndGains)
+	{
+		const VectorRegister NumFramesVector = VectorSetFloat1(NumSamples / 8.0f);
+
+		VectorRegister GainVector1 = VectorLoadAligned(StartGains);
+		const VectorRegister DestinationVector1 = VectorLoadAligned(EndGains);
+		const VectorRegister GainDeltasVector1 = VectorDivide(VectorSubtract(DestinationVector1, GainVector1), NumFramesVector);
+
+		VectorRegister GainVector2 = VectorLoadAligned(StartGains + 4);
+		const VectorRegister DestinationVector2 = VectorLoadAligned(EndGains + 4);
+		const VectorRegister GainDeltasVector2 = VectorDivide(VectorSubtract(DestinationVector2, GainVector2), NumFramesVector);
+
+		for (int32 i = 0; i < NumSamples; i += 8)
+		{
+			VectorRegister Result = VectorLoadAligned(&InterleavedBuffer[i]);
+			Result = VectorMultiply(Result, GainVector1);
+			VectorStoreAligned(Result, &InterleavedBuffer[i]);
+
+			GainVector1 = VectorAdd(GainVector1, GainDeltasVector1);
+
+			Result = VectorLoadAligned(&InterleavedBuffer[i + 4]);
+			Result = VectorMultiply(Result, GainVector2);
+			VectorStoreAligned(Result, &InterleavedBuffer[i + 4]);
+
+			GainVector2 = VectorAdd(GainVector2, GainDeltasVector2);
+		}
+	}
+
+	void MixMonoTo8ChannelsFast(const AlignedFloatBuffer& MonoBuffer, AlignedFloatBuffer& DestinationBuffer, const float* RESTRICT Gains)
+	{
+		MixMonoTo8ChannelsFast(MonoBuffer.GetData(), DestinationBuffer.GetData(), DestinationBuffer.Num() / 8, Gains);
+	}
+
+	/**
+	* See CHANNEL MIXING OPERATIONS above for more info.
+	* 1 frames per iteration:
+	* +------------+----------+---------+---------+---------+----------+---------+---------+---------+
+	* |            | Vector 1 |         |         |         | Vector 2 |         |         |         |
+	* | VectorName | Index 0  | Index 1 | Index 2 | Index 3 | Index 4  | Index 5 | Index 6 | Index 7 |
+	* | Gain       | g0       | g1      | g2      | g3      | g4       | g5      | g6      | g7      |
+	* |            | *        | *       | *       | *       | *        | *       | *       | *       |
+	* | Input      | i0       | i0      | i0      | i0      | i0       | i0      | i0      | i0      |
+	* |            | =        | =       | =       | =       | =        | =       | =       | =       |
+	* | Output     | o0       | o1      | o2      | o3      | o4       | o5      | o6      | o7      |
+	* +------------+----------+---------+---------+---------+----------+---------+---------+---------+
+	*/
+
+	void MixMonoTo8ChannelsFast(const float* RESTRICT MonoBuffer, float* RESTRICT DestinationBuffer, int32 NumFrames, const float* RESTRICT Gains)
+	{
+		const VectorRegister GainVector1 = VectorLoadAligned(Gains);
+		const VectorRegister GainVector2 = VectorLoadAligned(Gains + 4);
+
+		for (int32 i = 0; i < NumFrames; i++)
+		{
+			VectorRegister Result = VectorLoadFloat1(&MonoBuffer[i]);
+			Result = VectorMultiply(Result, GainVector1);
+			VectorStoreAligned(Result, &DestinationBuffer[i * 8]);
+
+			Result = VectorLoadFloat1(&MonoBuffer[i]);
+			Result = VectorMultiply(Result, GainVector2);
+			VectorStoreAligned(Result, &DestinationBuffer[i * 8 + 4]);
+		}
+	}
+
+	void MixMonoTo8ChannelsFast(const AlignedFloatBuffer& MonoBuffer, AlignedFloatBuffer& DestinationBuffer, const float* RESTRICT StartGains, const float* RESTRICT EndGains)
+	{
+		MixMonoTo8ChannelsFast(MonoBuffer.GetData(), DestinationBuffer.GetData(), DestinationBuffer.Num() / 8, StartGains, EndGains);
+	}
+
+	/**
+	* See CHANNEL MIXING OPERATIONS above for more info.
+	* 1 frames per iteration:
+	* +------------+----------+---------+---------+---------+----------+---------+---------+---------+
+	* |            | Vector 1 |         |         |         | Vector 2 |         |         |         |
+	* | VectorName | Index 0  | Index 1 | Index 2 | Index 3 | Index 4  | Index 5 | Index 6 | Index 7 |
+	* | Gain       | g0       | g1      | g2      | g3      | g4       | g5      | g6      | g7      |
+	* |            | *        | *       | *       | *       | *        | *       | *       | *       |
+	* | Input      | i0       | i0      | i0      | i0      | i0       | i0      | i0      | i0      |
+	* |            | =        | =       | =       | =       | =        | =       | =       | =       |
+	* | Output     | o0       | o1      | o2      | o3      | o4       | o5      | o6      | o7      |
+	* +------------+----------+---------+---------+---------+----------+---------+---------+---------+
+	*/
+	void MixMonoTo8ChannelsFast(const float* RESTRICT MonoBuffer, float* RESTRICT DestinationBuffer, int32 NumFrames, const float* RESTRICT StartGains, const float* RESTRICT EndGains)
+	{
+		const VectorRegister NumFramesVector = VectorSetFloat1((float) NumFrames);
+
+		VectorRegister GainVector1 = VectorLoadAligned(StartGains);
+		const VectorRegister DestinationVector1 = VectorLoadAligned(EndGains);
+		const VectorRegister GainDeltasVector1 = VectorDivide(VectorSubtract(DestinationVector1, GainVector1), NumFramesVector);
+
+		VectorRegister GainVector2 = VectorLoadAligned(StartGains + 4);
+		const VectorRegister DestinationVector2 = VectorLoadAligned(EndGains + 4);
+		const VectorRegister GainDeltasVector2 = VectorDivide(VectorSubtract(DestinationVector2, GainVector2), NumFramesVector);
+
+		for (int32 i = 0; i < NumFrames; i++)
+		{
+			const VectorRegister Input = VectorLoadFloat1(&MonoBuffer[i]);
+			VectorRegister Result = VectorMultiply(Input, GainVector1);
+			VectorStoreAligned(Result, &DestinationBuffer[i * 8]);
+
+			GainVector1 = VectorAdd(GainVector1, GainDeltasVector1);
+
+			Result = VectorMultiply(Input, GainVector2);
+			VectorStoreAligned(Result, &DestinationBuffer[i * 8 + 4]);
+
+			GainVector2 = VectorAdd(GainVector2, GainDeltasVector2);
+		}
+	}
+
+	void Mix2ChannelsTo8ChannelsFast(const AlignedFloatBuffer& SourceBuffer, AlignedFloatBuffer& DestinationBuffer, const float* RESTRICT Gains)
+	{
+		Mix2ChannelsTo8ChannelsFast(SourceBuffer.GetData(), DestinationBuffer.GetData(), DestinationBuffer.Num() / 8, Gains);
+	}
+
+	/**
+	* See CHANNEL MIXING OPERATIONS above for more info.
+	* 1 frames per iteration:
+	* +--------+----------+---------+---------+---------+----------+---------+---------+---------+
+	* |        | Vector 1 |         |         |         | Vector 2 |         |         |         |
+	* +--------+----------+---------+---------+---------+----------+---------+---------+---------+
+	* |        | Index 0  | Index 1 | Index 2 | Index 3 | Index 4  | Index 5 | Index 6 | Index 7 |
+	* | Gain1  | g0       | g1      | g2      | g3      | g4       | g5      | g6      | g7      |
+	* |        | *        | *       | *       | *       | *        | *       | *       | *       |
+	* | Input1 | i0       | i0      | i0      | i0      | i0       | i0      | i0      | i0      |
+	* |        | +        | +       | +       | +       | +        | +       | +       | +       |
+	* | Gain2  | g8       | g9      | g10     | g11     | g12      | g13     | g14     | g5      |
+	* |        | *        | *       | *       | *       | *        | *       | *       | *       |
+	* | Input2 | i1       | i1      | i1      | i1      | i1       | i1      | i1      | i1      |
+	* |        | =        | =       | =       | =       | =        | =       | =       | =       |
+	* | Output | o0       | o1      | o2      | o3      | o4       | o5      | o6      | o7      |
+	* +--------+----------+---------+---------+---------+----------+---------+---------+---------+
+	*/
+
+	void Mix2ChannelsTo8ChannelsFast(const float* RESTRICT SourceBuffer, float* RESTRICT DestinationBuffer, int32 NumFrames, const float* RESTRICT Gains)
+	{
+		const VectorRegister GainVector11 = VectorLoadAligned(Gains);
+		const VectorRegister GainVector21 = VectorLoadAligned(Gains + 4);
+		const VectorRegister GainVector12 = VectorLoadAligned(Gains + 8);
+		const VectorRegister GainVector22 = VectorLoadAligned(Gains + 12);
+
+		for (int32 i = 0; i < NumFrames; i++)
+		{
+			const VectorRegister Input1 = VectorLoadFloat1(&SourceBuffer[i*2]);
+			const VectorRegister Input2 = VectorLoadFloat1(&SourceBuffer[i * 2 + 1]);
+
+			VectorRegister Result = VectorMultiply(Input1, GainVector11);
+			Result = VectorMultiplyAdd(Input2, GainVector12, Result);
+
+			VectorStoreAligned(Result, &DestinationBuffer[i * 8]);
+
+			Result = VectorMultiply(Input1, GainVector21);
+			Result = VectorMultiplyAdd(Input2, GainVector22, Result);
+			VectorStoreAligned(Result, &DestinationBuffer[i * 8 + 4]);
+		}
+	}
+
+	void Mix2ChannelsTo8ChannelsFast(const AlignedFloatBuffer& SourceBuffer, AlignedFloatBuffer& DestinationBuffer, const float* RESTRICT StartGains, const float* RESTRICT EndGains)
+	{
+		Mix2ChannelsTo8ChannelsFast(SourceBuffer.GetData(), DestinationBuffer.GetData(), DestinationBuffer.Num() / 8, StartGains, EndGains);
+	}
+
+	/**
+	* See CHANNEL MIXING OPERATIONS above for more info.
+	* 1 frames per iteration:
+	* +--------+----------+---------+---------+---------+----------+---------+---------+---------+
+	* |        | Vector 1 |         |         |         | Vector 2 |         |         |         |
+	* +--------+----------+---------+---------+---------+----------+---------+---------+---------+
+	* |        | Index 0  | Index 1 | Index 2 | Index 3 | Index 4  | Index 5 | Index 6 | Index 7 |
+	* | Gain1  | g0       | g1      | g2      | g3      | g4       | g5      | g6      | g7      |
+	* |        | *        | *       | *       | *       | *        | *       | *       | *       |
+	* | Input1 | i0       | i0      | i0      | i0      | i0       | i0      | i0      | i0      |
+	* |        | +        | +       | +       | +       | +        | +       | +       | +       |
+	* | Gain2  | g8       | g9      | g10     | g11     | g12      | g13     | g14     | g15     |
+	* |        | *        | *       | *       | *       | *        | *       | *       | *       |
+	* | Input2 | i1       | i1      | i1      | i1      | i1       | i1      | i1      | i1      |
+	* |        | =        | =       | =       | =       | =        | =       | =       | =       |
+	* | Output | o0       | o1      | o2      | o3      | o4       | o5      | o6      | o7      |
+	* +--------+----------+---------+---------+---------+----------+---------+---------+---------+
+	*/
+
+	void Mix2ChannelsTo8ChannelsFast(const float* RESTRICT SourceBuffer, float* RESTRICT DestinationBuffer, int32 NumFrames, const float* RESTRICT StartGains, const float* RESTRICT EndGains)
+	{
+		const VectorRegister NumFramesVector = VectorSetFloat1((float) NumFrames);
+
+		VectorRegister GainVector11 = VectorLoadAligned(StartGains);
+		const VectorRegister DestinationVector11 = VectorLoadAligned(EndGains);
+		const VectorRegister GainDeltasVector11 = VectorDivide(VectorSubtract(DestinationVector11, GainVector11), NumFramesVector);
+
+		VectorRegister GainVector21 = VectorLoadAligned(StartGains + 4);
+		const VectorRegister DestinationVector21 = VectorLoadAligned(EndGains + 4);
+		const VectorRegister GainDeltasVector21 = VectorDivide(VectorSubtract(DestinationVector21, GainVector21), NumFramesVector);
+
+		VectorRegister GainVector12 = VectorLoadAligned(StartGains + 8);
+		const VectorRegister DestinationVector12 = VectorLoadAligned(EndGains + 8);
+		const VectorRegister GainDeltasVector12 = VectorDivide(VectorSubtract(DestinationVector12, GainVector12), NumFramesVector);
+
+		VectorRegister GainVector22 = VectorLoadAligned(StartGains + 12);
+		const VectorRegister DestinationVector22 = VectorLoadAligned(EndGains + 12);
+		const VectorRegister GainDeltasVector22 = VectorDivide(VectorSubtract(DestinationVector22, GainVector22), NumFramesVector);
+
+		for (int32 i = 0; i < NumFrames; i++)
+		{
+			const VectorRegister Input1 = VectorLoadFloat1(&SourceBuffer[i*2]);
+			const VectorRegister Input2 = VectorLoadFloat1(&SourceBuffer[i * 2 + 1]);
+
+			VectorRegister Result = VectorMultiply(Input1, GainVector11);
+			Result = VectorMultiplyAdd(Input2, GainVector12, Result);
+			VectorStoreAligned(Result, &DestinationBuffer[i * 8]);
+
+			GainVector11 = VectorAdd(GainVector11, GainDeltasVector11);
+			GainVector12 = VectorAdd(GainVector12, GainDeltasVector12);
+
+			Result = VectorMultiply(Input1, GainVector21);
+			Result = VectorMultiplyAdd(Input2, GainVector22, Result);
+			VectorStoreAligned(Result, &DestinationBuffer[i * 8 + 4]);
+
+			GainVector21 = VectorAdd(GainVector21, GainDeltasVector21);
+			GainVector22 = VectorAdd(GainVector22, GainDeltasVector22);
+		}
+	}
+
+	/**
+	 * These functions are non-vectorized versions of the Mix[N]ChannelsTo[N]Channels functions above:
+	 */
+	void DownmixBuffer(int32 NumSourceChannels, int32 NumDestinationChannels, const AlignedFloatBuffer& SourceBuffer, AlignedFloatBuffer& DestinationBuffer, const float* RESTRICT Gains)
+	{
+		DownmixBuffer(NumSourceChannels, NumDestinationChannels, SourceBuffer.GetData(), DestinationBuffer.GetData(), SourceBuffer.Num() / NumSourceChannels, Gains);
+	}
+
+	void DownmixBuffer(int32 NumSourceChannels, int32 NumDestinationChannels, const float* RESTRICT SourceBuffer, float* RESTRICT DestinationBuffer, int32 NumFrames, const float* RESTRICT Gains)
+	{
+		for (int32 FrameIndex = 0; FrameIndex < NumFrames; FrameIndex++)
+		{
+			float* RESTRICT OutputFrame = &DestinationBuffer[FrameIndex * NumDestinationChannels];
+			const float* RESTRICT InputFrame = &SourceBuffer[FrameIndex * NumSourceChannels];
+
+			FMemory::Memzero(OutputFrame, NumDestinationChannels * sizeof(float));
+			for (int32 OutputChannelIndex = 0; OutputChannelIndex < NumDestinationChannels; OutputChannelIndex++)
+			{
+				for (int32 InputChannelIndex = 0; InputChannelIndex < NumSourceChannels; InputChannelIndex++)
+				{
+					OutputFrame[OutputChannelIndex] += InputFrame[InputChannelIndex] * Gains[InputChannelIndex * NumDestinationChannels + OutputChannelIndex];
+				}
+			}
+		}
+	}
+
+	void DownmixBuffer(int32 NumSourceChannels, int32 NumDestinationChannels, const AlignedFloatBuffer& SourceBuffer, AlignedFloatBuffer& DestinationBuffer, float* RESTRICT StartGains, const float* RESTRICT EndGains)
+	{
+		DownmixBuffer(NumSourceChannels, NumDestinationChannels, SourceBuffer.GetData(), DestinationBuffer.GetData(), SourceBuffer.Num() / NumSourceChannels, StartGains, EndGains);
+	}
+
+	void DownmixBuffer(int32 NumSourceChannels, int32 NumDestinationChannels, const float* RESTRICT SourceBuffer, float* RESTRICT DestinationBuffer, int32 NumFrames, float* RESTRICT StartGains, const float* RESTRICT EndGains)
+	{
+		// First, build a map of the per-frame delta that we will use to increment StartGains every frame:
+		alignas(16) float GainDeltas[AUDIO_MIXER_MAX_OUTPUT_CHANNELS * AUDIO_MIXER_MAX_OUTPUT_CHANNELS];
+		
+		for (int32 OutputChannelIndex = 0; OutputChannelIndex < NumDestinationChannels; OutputChannelIndex++)
+		{
+			for (int32 InputChannelIndex = 0; InputChannelIndex < NumSourceChannels; InputChannelIndex++)
+			{
+				const int32 GainMatrixIndex = InputChannelIndex * NumDestinationChannels + OutputChannelIndex;
+				GainDeltas[GainMatrixIndex] = (EndGains[GainMatrixIndex] - StartGains[GainMatrixIndex]) / NumFrames;
+			}
+		}
+
+		for (int32 FrameIndex = 0; FrameIndex < NumFrames; FrameIndex++)
+		{
+			float* RESTRICT OutputFrame = &DestinationBuffer[FrameIndex * NumDestinationChannels];
+			const float* RESTRICT InputFrame = &SourceBuffer[FrameIndex * NumSourceChannels];
+
+			FMemory::Memzero(OutputFrame, NumDestinationChannels * sizeof(float));
+			for (int32 OutputChannelIndex = 0; OutputChannelIndex < NumDestinationChannels; OutputChannelIndex++)
+			{
+				for (int32 InputChannelIndex = 0; InputChannelIndex < NumSourceChannels; InputChannelIndex++)
+				{
+					const int32 GainMatrixIndex = InputChannelIndex * NumDestinationChannels + OutputChannelIndex;
+					OutputFrame[OutputChannelIndex] += InputFrame[InputChannelIndex] * StartGains[GainMatrixIndex];
+					StartGains[GainMatrixIndex] += GainDeltas[GainMatrixIndex];
+				}
+			}
+		}
+	}
+
+
+	/** Interleaves samples from two input buffers */
+	void BufferInterleave2ChannelFast(const AlignedFloatBuffer& InBuffer1, const AlignedFloatBuffer& InBuffer2, AlignedFloatBuffer& OutBuffer)
+	{
+		checkf(InBuffer1.Num() == InBuffer2.Num(), TEXT("InBuffer1 Num not equal to InBuffer2 Num"));
+
+		const int32 InNum = InBuffer1.Num();
+
+		OutBuffer.Reset(2 * InNum);
+		OutBuffer.AddUninitialized(2 * InNum);
+		
+		BufferInterleave2ChannelFast(InBuffer1.GetData(), InBuffer2.GetData(), OutBuffer.GetData(), InNum);
+	}
+
+	/** Interleaves samples from two input buffers */
+	void BufferInterleave2ChannelFast(const float* RESTRICT InBuffer1, const float* RESTRICT InBuffer2, float* RESTRICT OutBuffer, const int32 InNum)
+	{
+		checkf(InNum >= 4, TEXT("Buffer must have atleast 4 elements."));
+		checkf(0 == (InNum % 4), TEXT("Buffer length be a multiple of 4."));
+		checkf(IsAligned<const float*>(InBuffer1, AUDIO_SIMD_FLOAT_ALIGNMENT), TEXT("Memory must be aligned to use vector operations."));
+		checkf(IsAligned<const float*>(InBuffer2, AUDIO_SIMD_FLOAT_ALIGNMENT), TEXT("Memory must be aligned to use vector operations."));
+		checkf(IsAligned<float*>(OutBuffer, AUDIO_SIMD_FLOAT_ALIGNMENT), TEXT("Memory must be aligned to use vector operations."));
+
+		const int32 OutNum = 2 * InNum;
+
+		int32 OutPos = 0;
+		for (int32 i = 0; i < InNum; i += 4)
+		{
+			// Vector1[L0, L1, L2, L3]
+			VectorRegister Vector1 = VectorLoadAligned(&InBuffer1[i]);
+			// Vector2[R0, R1, R2, R3]
+			VectorRegister Vector2 = VectorLoadAligned(&InBuffer2[i]);
+
+			// HalfInterleaved[L0, L1, R0, R1]
+			VectorRegister HalfInterleaved = VectorShuffle(Vector1, Vector2, 0, 1, 0, 1);
+			// Interleaved[L0, R0, L1, R1]
+			VectorRegister Interleaved = VectorSwizzle(HalfInterleaved, 0, 2, 1, 3);
+			VectorStoreAligned(Interleaved, &OutBuffer[OutPos]);
+			OutPos += 4;
+
+			// HalfInterleaved[L2, L3, R2, R3]
+			HalfInterleaved = VectorShuffle(Vector1, Vector2, 2, 3, 2, 3);
+			// Interleaved[L2, R2, L3, R3]
+			Interleaved = VectorSwizzle(HalfInterleaved, 0, 2, 1, 3);
+			VectorStoreAligned(Interleaved, &OutBuffer[OutPos]);
+			OutPos += 4;
+		}
+	}
+
+	/** Deinterleaves samples from a 2 channel input buffer */
+	void BufferDeinterleave2ChannelFast(const AlignedFloatBuffer& InBuffer, AlignedFloatBuffer& OutBuffer1, AlignedFloatBuffer& OutBuffer2)
+	{
+		const int32 InNum = InBuffer.Num();
+		const int32 InNumFrames = InNum / 2;
+		const int32 OutNum = InNumFrames;
+
+		OutBuffer1.Reset(OutNum);
+		OutBuffer2.Reset(OutNum);
+		OutBuffer1.AddUninitialized(OutNum);
+		OutBuffer2.AddUninitialized(OutNum);
+		
+		BufferDeinterleave2ChannelFast(InBuffer.GetData(), OutBuffer1.GetData(), OutBuffer2.GetData(), InNumFrames);
+	}
+
+	/** Deinterleaves samples from a 2 channel input buffer */
+	void BufferDeinterleave2ChannelFast(const float* RESTRICT InBuffer, float* RESTRICT OutBuffer1, float* RESTRICT OutBuffer2, const int32 InNumFrames)
+	{
+		checkf(InNumFrames >= 4, TEXT("Buffer must have atleast 4 elements."));
+		checkf(0 == (InNumFrames % 4), TEXT("Buffer length be a multiple of 4."));
+		checkf(IsAligned<const float*>(InBuffer, AUDIO_SIMD_FLOAT_ALIGNMENT), TEXT("Memory must be aligned to use vector operations."));
+		checkf(IsAligned<const float*>(OutBuffer1, 4), TEXT("Memory must be aligned to use vector operations."));
+		checkf(IsAligned<float*>(OutBuffer2, 4), TEXT("Memory must be aligned to use vector operations."));
+
+		int32 InNum = InNumFrames * 2;
+		int32 OutPos = 0;
+		for (int32 InPos = 0; InPos < InNum; InPos += 8)
+		{
+			// load 4 frames (2 frames per vector)
+			VectorRegister InVector1 = VectorLoadAligned(&InBuffer[InPos]);
+			VectorRegister InVector2 = VectorLoadAligned(&InBuffer[InPos + 4]);
+
+			// Write channel 0
+			VectorRegister OutVector = VectorShuffle(InVector1, InVector2, 0, 2, 0, 2);
+			VectorStoreAligned(OutVector, &OutBuffer1[OutPos]);
+
+			// Write channel 1
+			OutVector = VectorShuffle(InVector1, InVector2, 1, 3, 1, 3);
+			VectorStoreAligned(OutVector, &OutBuffer2[OutPos]);
+
+			OutPos += 4;
+		}
+	}
+
+	/** Sums 2 channel interleaved input samples. OutSamples[n] = InSamples[2n] + InSamples[2n + 1] */
+	void BufferSum2ChannelToMonoFast(const AlignedFloatBuffer& InSamples, AlignedFloatBuffer& OutSamples)
+	{
+		const int32 InNum = InSamples.Num();
+		const int32 Frames = InNum / 2;
+
+		OutSamples.Reset(Frames);
+		OutSamples.AddUninitialized(Frames);
+		
+		BufferSum2ChannelToMonoFast(InSamples.GetData(), OutSamples.GetData(), Frames);
+	}
+
+	/** Sums 2 channel interleaved input samples. OutSamples[n] = InSamples[2n] + InSamples[2n + 1] */
+	void BufferSum2ChannelToMonoFast(const float* RESTRICT InSamples, float* RESTRICT OutSamples, const int32 InNumFrames)
+	{
+		checkf(InNumFrames >= 4, TEXT("Buffer must have atleast 4 elements."));
+		checkf(0 == (InNumFrames % 4), TEXT("Buffer length be a multiple of 4."));
+		checkf(IsAligned<const float*>(InSamples, AUDIO_SIMD_FLOAT_ALIGNMENT), TEXT("Memory must be aligned to use vector operations."));
+		checkf(IsAligned<float*>(OutSamples, 4), TEXT("Memory must be aligned to use vector operations."));
+
+		const int32 InNum = InNumFrames * 2;
+		int32 OutPos = 0;
+		for (int32 i = 0; i < InNum; i += 8)
+		{
+			// Load 4 frames (2 frames per vector)
+			// Buffer1[L0, R0, L1, R1]
+			VectorRegister Buffer1 = VectorLoadAligned(&InSamples[i]);
+			// Buffer2[L2, R2, L3, R3]
+			VectorRegister Buffer2 = VectorLoadAligned(&InSamples[i + 4]);
+
+			// Shuffle samples into order
+			// Channel0[L0, L1, L2, L3]
+			VectorRegister Channel0 = VectorShuffle(Buffer1, Buffer2, 0, 2, 0, 2);
+			// Channel1[R0, R1, R2, R3]
+			VectorRegister Channel1 = VectorShuffle(Buffer1, Buffer2, 1, 3, 1, 3);
+
+			// Sum left and right.
+			// Out[L0 + R0, L1 + R1, L2 + R2, L3 + R3]
+			VectorRegister Out = VectorAdd(Channel0, Channel1);
+
+			VectorStoreAligned(Out, &OutSamples[OutPos]);
+			OutPos += 4;
+		}
+	}
+
+	void BufferComplexToPowerFast(const AlignedFloatBuffer& InRealSamples, const AlignedFloatBuffer& InImaginarySamples, AlignedFloatBuffer& OutPowerSamples)
+	{
+		checkf(InRealSamples.Num() == InImaginarySamples.Num(), TEXT("Input buffers must have equal number of elements"));
+
+		const int32 Num = InRealSamples.Num();
+
+		OutPowerSamples.Reset(Num);
+		OutPowerSamples.AddUninitialized(Num);
+
+		BufferComplexToPowerFast(InRealSamples.GetData(), InImaginarySamples.GetData(), OutPowerSamples.GetData(), Num);
+	}
+
+	void BufferComplexToPowerFast(const float* RESTRICT InRealSamples, const float* RESTRICT InImaginarySamples, float* RESTRICT OutPowerSamples, const int32 InNum)
+	{
+		checkf(InNum >= 4, TEXT("Buffer must have atleast 4 elements."));
+		checkf(0 == (InNum % 4), TEXT("Buffer length be a multiple of 4."));
+
+		checkf(IsAligned<const float*>(InRealSamples, AUDIO_SIMD_FLOAT_ALIGNMENT), TEXT("Memory must be aligned to use vector operations."));
+		checkf(IsAligned<const float*>(InImaginarySamples, AUDIO_SIMD_FLOAT_ALIGNMENT), TEXT("Memory must be aligned to use vector operations."));
+		checkf(IsAligned<float*>(OutPowerSamples, AUDIO_SIMD_FLOAT_ALIGNMENT), TEXT("Memory must be aligned to use vector operations."));
+
+		
+		for (int32 i = 0; i < InNum; i += 4)
+		{
+			VectorRegister VInReal = VectorLoadAligned(&InRealSamples[i]);
+			VectorRegister VInRealSquared = VectorMultiply(VInReal, VInReal);
+
+			VectorRegister VInImag = VectorLoadAligned(&InImaginarySamples[i]);
+			VectorRegister VInImagSquared = VectorMultiply(VInImag, VInImag);
+
+			VectorRegister VOut = VectorAdd(VInRealSquared, VInImagSquared);
+
+			VectorStoreAligned(VOut, &OutPowerSamples[i]);
+		}
+	}
+}
diff --git a/Engine/Source/Runtime/AudioMixer/Public/DSP/BufferVectorOperations.h b/Engine/Source/Runtime/AudioMixer/Public/DSP/BufferVectorOperations.h
new file mode 100644
--- /dev/null
+++ b/Engine/Source/Runtime/AudioMixer/Public/DSP/BufferVectorOperations.h
@@ -0,0 +1,253 @@
+// Copyright 1998-2019 Epic Games, Inc. All Rights Reserved.
+
+#pragma once
+
+#include "CoreMinimal.h"
+
+#if PLATFORM_SWITCH
+// Switch uses page alignment for submitted buffers
+#define AUDIO_BUFFER_ALIGNMENT 4096
+#else
+#define AUDIO_BUFFER_ALIGNMENT 16
+#endif
+
+#define AUDIO_SIMD_BYTE_ALIGNMENT (16)
+#define AUDIO_SIMD_FLOAT_ALIGNMENT (4)
+
+namespace Audio
+{
+	/** Aligned allocator used for fast operations. */	
+	typedef TAlignedHeapAllocator<AUDIO_BUFFER_ALIGNMENT> FAudioBufferAlignedAllocator;
+
+	typedef TArray<float, FAudioBufferAlignedAllocator> AlignedFloatBuffer;
+	typedef TArray<uint8, FAudioBufferAlignedAllocator> AlignedByteBuffer;
+	typedef TArray<int32, FAudioBufferAlignedAllocator> AlignedInt32Buffer;
+
+
+	/** CHANNEL-AGNOSTIC OPERATIONS */
+
+	/* Sets a values to zero if value is denormal. Denormal numbers significantly slow down floating point operations. */
+	void AUDIOMIXER_API BufferUnderflowClampFast(AlignedFloatBuffer& InOutBuffer);
+
+	/* Sets a values to zero if value is denormal. Denormal numbers significantly slow down floating point operations. */
+	void AUDIOMIXER_API BufferUnderflowClampFast(float* RESTRICT InOutBuffer, const int32 InNum);
+		
+	/** Multiplies the input aligned float buffer with the given value. */
+	void AUDIOMIXER_API BufferMultiplyByConstant(const AlignedFloatBuffer& InFloatBuffer, float InValue, AlignedFloatBuffer& OutFloatBuffer);
+	void AUDIOMIXER_API BufferMultiplyByConstant(const float* RESTRICT InFloatBuffer, float InValue, float* RESTRICT OutFloatBuffer, const int32 InNumSamples);
+
+	/** Similar to BufferMultiplyByConstant, but (a) assumes a buffer length divisible by 4 and (b) performs the multiply in place. */
+	void AUDIOMIXER_API MultiplyBufferByConstantInPlace(AlignedFloatBuffer& InBuffer, float InGain);
+	void AUDIOMIXER_API MultiplyBufferByConstantInPlace(float* RESTRICT InBuffer, int32 NumSamples, float InGain);
+
+	/* Performs an element-wise weighted sum OutputBuffer = (InBuffer1 x InGain1) + (InBuffer2 x InGain2) */
+	void AUDIOMIXER_API BufferWeightedSumFast(const AlignedFloatBuffer& InBuffer1, float InGain1, const AlignedFloatBuffer& InBuffer2, float InGain2, AlignedFloatBuffer& OutBuffer);
+
+	/* Performs an element-wise weighted sum OutputBuffer = (InBuffer1 x InGain1) + InBuffer2 */
+	void AUDIOMIXER_API BufferWeightedSumFast(const AlignedFloatBuffer& InBuffer1, float InGain1, const AlignedFloatBuffer& InBuffer2, AlignedFloatBuffer& OutBuffer);
+
+	/* Performs an element-wise weighted sum OutputBuffer = (InBuffer1 x InGain1) + (InBuffer2 x InGain2) */
+	void AUDIOMIXER_API BufferWeightedSumFast(const float* RESTRICT InBuffer1, float InGain1, const float* RESTRICT InBuffer2, float InGain2, float* RESTRICT OutBuffer, int32 InNum);
+
+	/* Performs an element-wise weighted sum OutputBuffer = (InBuffer1 x InGain1) + InBuffer2 */
+	void AUDIOMIXER_API BufferWeightedSumFast(const float* RESTRICT InBuffer1, float InGain1, const float* RESTRICT InBuffer2, float* RESTRICT OutBuffer, int32 InNum);
+
+	/* Takes a float buffer and quickly interpolates it's gain from StartValue to EndValue. */
+	/* This operation completely ignores channel counts, so avoid using this function on buffers that are not mono, stereo or quad */
+	/* if the buffer needs to fade all channels uniformly. */
+	void AUDIOMIXER_API FadeBufferFast(AlignedFloatBuffer& OutFloatBuffer, const float StartValue, const float EndValue);
+	void AUDIOMIXER_API FadeBufferFast(float* RESTRICT OutFloatBuffer, int32 NumSamples, const float StartValue, const float EndValue);
+
+	/** Takes buffer InFloatBuffer, optionally multiplies it by Gain, and adds it to BufferToSumTo. */
+	void AUDIOMIXER_API MixInBufferFast(const AlignedFloatBuffer& InFloatBuffer, AlignedFloatBuffer& BufferToSumTo, const float Gain);
+	void AUDIOMIXER_API MixInBufferFast(const float* RESTRICT InFloatBuffer, float* RESTRICT BufferToSumTo, int32 NumSamples, const float Gain);
+	void AUDIOMIXER_API MixInBufferFast(const AlignedFloatBuffer& InFloatBuffer, AlignedFloatBuffer& BufferToSumTo);
+	void AUDIOMIXER_API MixInBufferFast(const float* RESTRICT InFloatBuffer, float* RESTRICT BufferToSumTo, int32 NumSamples);
+
+	/* Subtracts two buffers together element-wise. */
+	void AUDIOMIXER_API BufferSubtractFast(const AlignedFloatBuffer& InMinuend, const AlignedFloatBuffer& InSubtrahend, AlignedFloatBuffer& OutputBuffer);
+
+	/* Subtracts two buffers together element-wise. */
+	void AUDIOMIXER_API BufferSubtractFast(const float* RESTRICT InMinuend, const float* RESTRICT InSubtrahend, float* RESTRICT OutputBuffer, int32 NumSamples);
+
+	/* Performs element-wise in-place subtraction placing the result in the subtrahend. InOutSubtrahend = InMinuend - InOutSubtrahend */
+	void AUDIOMIXER_API BufferSubtractInPlace1Fast(const AlignedFloatBuffer& InMinuend, AlignedFloatBuffer& InOutSubtrahend);
+
+	/* Performs element-wise in-place subtraction placing the result in the subtrahend. InOutSubtrahend = InMinuend - InOutSubtrahend */
+	void AUDIOMIXER_API BufferSubtractInPlace1Fast(const float* RESTRICT InMinuend, float* RESTRICT InOutSubtrahend, int32 NumSamples);
+
+
+	/* Performs element-wise in-place subtraction placing the result in the minuend. InOutMinuend = InOutMinuend - InSubtrahend */
+	void AUDIOMIXER_API BufferSubtractInPlace2Fast(AlignedFloatBuffer& InOutMinuend, const AlignedFloatBuffer& InSubtrahend);
+
+	/* Performs element-wise in-place subtraction placing the result in the minuend. InOutMinuend = InOutMinuend - InSubtrahend */
+	void AUDIOMIXER_API BufferSubtractInPlace2Fast(float* RESTRICT InOutMinuend, const float* RESTRICT InSubtrahend, int32 NumSamples);
+
+	/** Sums two buffers together and places the result in the resulting buffer. */
+	void AUDIOMIXER_API SumBuffers(const AlignedFloatBuffer& InFloatBuffer1, const AlignedFloatBuffer& InFloatBuffer2, AlignedFloatBuffer& OutputBuffer);
+	void AUDIOMIXER_API SumBuffers(const float* RESTRICT InFloatBuffer1, const float* RESTRICT InFloatBuffer2, float* RESTRICT OutputBuffer, int32 NumSamples);
+
+	/** Multiply the second buffer in place by the first buffer. */
+	void AUDIOMIXER_API MultiplyBuffersInPlace(const AlignedFloatBuffer& InFloatBuffer, AlignedFloatBuffer& BufferToMultiply);
+	void AUDIOMIXER_API MultiplyBuffersInPlace(const float* RESTRICT InFloatBuffer, float* RESTRICT BufferToMultiply, int32 NumSamples);
+
+	/** CHANNEL-AGNOSTIC ANALYSIS OPERATIONS */
+
+	/** Takes an audio buffer and returns the magnitude across that buffer. */
+	float AUDIOMIXER_API GetMagnitude(const AlignedFloatBuffer& Buffer);
+	float AUDIOMIXER_API GetMagnitude(const float* RESTRICT Buffer, int32 NumSamples);
+
+	/** Takes an audio buffer and gets the average absolute amplitude across that buffer. */
+	float AUDIOMIXER_API GetAverageAmplitude(const AlignedFloatBuffer& Buffer);
+	float AUDIOMIXER_API GetAverageAmplitude(const float* RESTRICT Buffer, int32 NumSamples);
+
+	/** CHANNEL-SPECIFIC OPERATIONS */
+
+	/** Takes a 2 channel interleaved buffer and applies Gains to it. Gains is expected to point to a 2 float long buffer.
+	 *  StereoBuffer must have an even number of frames.
+	 *  If StartGains and EndGains are provided, this function will interpolate between the two across the buffer.
+	*/
+	void AUDIOMIXER_API Apply2ChannelGain(AlignedFloatBuffer& StereoBuffer, const float* RESTRICT Gains);
+	void AUDIOMIXER_API Apply2ChannelGain(float* RESTRICT StereoBuffer, int32 NumSamples, const float* RESTRICT Gains);
+	void AUDIOMIXER_API Apply2ChannelGain(AlignedFloatBuffer& StereoBuffer, const float* RESTRICT StartGains, const float* RESTRICT EndGains);
+	void AUDIOMIXER_API Apply2ChannelGain(float* RESTRICT StereoBuffer, int32 NumSamples, const float* RESTRICT StartGains, const float* RESTRICT EndGains);
+
+	/** Takes a 1 channel buffer and mixes it to a stereo buffer using Gains. Gains is expected to point to a 2 float long buffer.
+	 *  these buffers must have an even number of frames.
+	 *  If StartGains and EndGains are provided, this function will interpolate between the two across the buffer.
+	*/
+	void AUDIOMIXER_API MixMonoTo2ChannelsFast(const AlignedFloatBuffer& MonoBuffer, AlignedFloatBuffer& DestinationBuffer, const float* RESTRICT Gains);
+	void AUDIOMIXER_API MixMonoTo2ChannelsFast(const float* RESTRICT MonoBuffer, float* RESTRICT DestinationBuffer, int32 NumFrames, const float* RESTRICT Gains);
+	void AUDIOMIXER_API MixMonoTo2ChannelsFast(const AlignedFloatBuffer& MonoBuffer, AlignedFloatBuffer& DestinationBuffer, const float* RESTRICT StartGains, const float* RESTRICT EndGains);
+	void AUDIOMIXER_API MixMonoTo2ChannelsFast(const float* RESTRICT MonoBuffer, float* RESTRICT DestinationBuffer, int32 NumFrames, const float* RESTRICT StartGains, const float* RESTRICT EndGains);
+	void AUDIOMIXER_API MixMonoTo2ChannelsFast(const AlignedFloatBuffer& MonoBuffer, AlignedFloatBuffer& DestinationBuffer);
+	void AUDIOMIXER_API MixMonoTo2ChannelsFast(const float* RESTRICT MonoBuffer, float* RESTRICT DestinationBuffer, int32 NumFrames);
+
+	/** Takes a 2 channel buffer and mixes it to an 2 channel interleaved buffer using Gains. Gains is expected to point to a 16 float long buffer.
+	*  Output gains for the left input channel should be the first 8 values in Gains, and Output gains for the right input channel should be rest.
+	*  NumFrames must be a multiple of 4.
+	*  If StartGains and EndGains are provided, this function will interpolate between the two across the buffer.
+	*/
+	void AUDIOMIXER_API Mix2ChannelsTo2ChannelsFast(const AlignedFloatBuffer& SourceBuffer, AlignedFloatBuffer& DestinationBuffer, const float* RESTRICT Gains);
+	void AUDIOMIXER_API Mix2ChannelsTo2ChannelsFast(const float* RESTRICT SourceBuffer, float* RESTRICT DestinationBuffer, int32 NumFrames, const float* RESTRICT Gains);
+	void AUDIOMIXER_API Mix2ChannelsTo2ChannelsFast(const AlignedFloatBuffer& SourceBuffer, AlignedFloatBuffer& DestinationBuffer, const float* RESTRICT StartGains, const float* RESTRICT EndGains);
+	void AUDIOMIXER_API Mix2ChannelsTo2ChannelsFast(const float* RESTRICT SourceBuffer, float* RESTRICT DestinationBuffer, int32 NumFrames, const float* RESTRICT StartGains, const float* RESTRICT EndGains);
+
+	/** Takes a 4 channel interleaved buffer and applies Gains to it. Gains is expected to point to a 2 float long buffer.
+	 *  If StartGains and EndGains are provided, this function will interpolate between the two across the buffer.
+	*/
+	void AUDIOMIXER_API Apply4ChannelGain(AlignedFloatBuffer& InterleavedBuffer, const float* RESTRICT Gains);
+	void AUDIOMIXER_API Apply4ChannelGain(float* RESTRICT InterleavedBuffer, int32 NumSamples, const float* RESTRICT Gains);
+	void AUDIOMIXER_API Apply4ChannelGain(AlignedFloatBuffer& InterleavedBuffer, const float* RESTRICT StartGains, const float* RESTRICT EndGains);
+	void AUDIOMIXER_API Apply4ChannelGain(float* RESTRICT InterleavedBuffer, int32 NumSamples, const float* RESTRICT StartGains, const float* RESTRICT EndGains);
+
+	/** Takes a 1 channel buffer and mixes it to an 8 channel interleaved buffer using Gains. Gains is expected to point to a 8 float long buffer.
+	*  these buffers must have an even number of frames.
+	*  If StartGains and EndGains are provided, this function will interpolate between the two across the buffer.
+	*/
+	void AUDIOMIXER_API MixMonoTo4ChannelsFast(const AlignedFloatBuffer& MonoBuffer, AlignedFloatBuffer& DestinationBuffer, const float* RESTRICT Gains);
+	void AUDIOMIXER_API MixMonoTo4ChannelsFast(const float* RESTRICT MonoBuffer, float* RESTRICT DestinationBuffer, int32 NumFrames, const float* RESTRICT Gains);
+	void AUDIOMIXER_API MixMonoTo4ChannelsFast(const AlignedFloatBuffer& MonoBuffer, AlignedFloatBuffer& DestinationBuffer, const float* RESTRICT StartGains, const float* RESTRICT EndGains);
+	void AUDIOMIXER_API MixMonoTo4ChannelsFast(const float* RESTRICT MonoBuffer, float* RESTRICT DestinationBuffer, int32 NumFrames, const float* RESTRICT StartGains, const float* RESTRICT EndGains);
+
+	/** Takes a 2 channel buffer and mixes it to an 8 channel interleaved buffer using Gains. Gains is expected to point to a 16 float long buffer.
+	*  Output gains for the left input channel should be the first 8 values in Gains, and Output gains for the right input channel should be rest.
+	*  NumFrames must be a multiple of 4.
+	*  If StartGains and EndGains are provided, this function will interpolate between the two across the buffer.
+	*/
+	void AUDIOMIXER_API Mix2ChannelsTo4ChannelsFast(const AlignedFloatBuffer& SourceBuffer, AlignedFloatBuffer& DestinationBuffer, const float* RESTRICT Gains);
+	void AUDIOMIXER_API Mix2ChannelsTo4ChannelsFast(const float* RESTRICT SourceBuffer, float* RESTRICT DestinationBuffer, int32 NumFrames, const float* RESTRICT Gains);
+	void AUDIOMIXER_API Mix2ChannelsTo4ChannelsFast(const AlignedFloatBuffer& SourceBuffer, AlignedFloatBuffer& DestinationBuffer, const float* RESTRICT StartGains, const float* RESTRICT EndGains);
+	void AUDIOMIXER_API Mix2ChannelsTo4ChannelsFast(const float* RESTRICT SourceBuffer, float* RESTRICT DestinationBuffer, int32 NumFrames, const float* RESTRICT StartGains, const float* RESTRICT EndGains);
+
+	/** Takes a 6 channel interleaved buffer and applies Gains to it. Gains is expected to point to a 2 float long buffer.
+	 *  InterleavedBuffer must have an even number of frames.
+	*/
+	void AUDIOMIXER_API Apply6ChannelGain(AlignedFloatBuffer& InterleavedBuffer, const float* RESTRICT Gains);
+	void AUDIOMIXER_API Apply6ChannelGain(float* RESTRICT InterleavedBuffer, int32 NumSamples, const float* RESTRICT Gains);
+	void AUDIOMIXER_API Apply6ChannelGain(AlignedFloatBuffer& InterleavedBuffer, const float* RESTRICT StartGains, const float* RESTRICT EndGains);
+	void AUDIOMIXER_API Apply6ChannelGain(float* RESTRICT InterleavedBuffer, int32 NumSamples, const float* RESTRICT StartGains, const float* RESTRICT EndGains);
+
+	/** Takes a 1 channel buffer and mixes it to an 8 channel interleaved buffer using Gains. Gains is expected to point to a 8 float long buffer.
+	 *  these buffers must have an even number of frames.
+	 *  If StartGains and EndGains are provided, this function will interpolate between the two across the buffer.
+	*/
+	void AUDIOMIXER_API MixMonoTo6ChannelsFast(const AlignedFloatBuffer& MonoBuffer, AlignedFloatBuffer& DestinationBuffer, const float* RESTRICT Gains);
+	void AUDIOMIXER_API MixMonoTo6ChannelsFast(const float* RESTRICT MonoBuffer, float* RESTRICT DestinationBuffer, int32 NumFrames, const float* RESTRICT Gains);
+	void AUDIOMIXER_API MixMonoTo6ChannelsFast(const AlignedFloatBuffer& MonoBuffer, AlignedFloatBuffer& DestinationBuffer, const float* RESTRICT StartGains, const float* RESTRICT EndGains);
+	void AUDIOMIXER_API MixMonoTo6ChannelsFast(const float* RESTRICT MonoBuffer, float* RESTRICT DestinationBuffer, int32 NumFrames, const float* RESTRICT StartGains, const float* RESTRICT EndGains);
+
+	/** Takes a 2 channel buffer and mixes it to an 8 channel interleaved buffer using Gains. Gains is expected to point to a 16 float long buffer.
+	 *  Output gains for the left input channel should be the first 8 values in Gains, and Output gains for the right input channel should be rest.
+	 *  NumFrames must be a multiple of 4.
+	 *  If StartGains and EndGains are provided, this function will interpolate between the two across the buffer.
+	*/
+	void AUDIOMIXER_API Mix2ChannelsTo6ChannelsFast(const AlignedFloatBuffer& SourceBuffer, AlignedFloatBuffer& DestinationBuffer, const float* RESTRICT Gains);
+	void AUDIOMIXER_API Mix2ChannelsTo6ChannelsFast(const float* RESTRICT SourceBuffer, float* RESTRICT DestinationBuffer, int32 NumFrames, const float* RESTRICT Gains);
+	void AUDIOMIXER_API Mix2ChannelsTo6ChannelsFast(const AlignedFloatBuffer& SourceBuffer, AlignedFloatBuffer& DestinationBuffer, const float* RESTRICT StartGains, const float* RESTRICT EndGains);
+	void AUDIOMIXER_API Mix2ChannelsTo6ChannelsFast(const float* RESTRICT SourceBuffer, float* RESTRICT DestinationBuffer, int32 NumFrames, const float* RESTRICT StartGains, const float* RESTRICT EndGains);
+
+	/** Takes an 8 channel interleaved buffer and applies Gains to it. Gains is expected to point to an 8 float long buffer. */
+	void AUDIOMIXER_API Apply8ChannelGain(AlignedFloatBuffer& InterleavedBuffer, const float* RESTRICT Gains);
+	void AUDIOMIXER_API Apply8ChannelGain(float* RESTRICT InterleavedBuffer, int32 NumSamples, const float* RESTRICT Gains);
+	void AUDIOMIXER_API Apply8ChannelGain(AlignedFloatBuffer& InterleavedBuffer, const float* RESTRICT StartGains, const float* RESTRICT EndGains);
+	void AUDIOMIXER_API Apply8ChannelGain(float* RESTRICT InterleavedBuffer, int32 NumSamples, const float* RESTRICT StartGains, const float* RESTRICT EndGains);
+
+	/** Takes a 1 channel buffer and mixes it to an 8 channel interleaved buffer using Gains. Gains is expected to point to a 8 float long buffer.
+	 *  these buffers must have an even number of frames.
+	 *  If StartGains and EndGains are provided, this function will interpolate between the two across the buffer.
+	*/
+	void AUDIOMIXER_API MixMonoTo8ChannelsFast(const AlignedFloatBuffer& MonoBuffer, AlignedFloatBuffer& DestinationBuffer, const float* RESTRICT Gains);
+	void AUDIOMIXER_API MixMonoTo8ChannelsFast(const float* RESTRICT MonoBuffer, float* RESTRICT DestinationBuffer, int32 NumFrames, const float* RESTRICT Gains);
+	void AUDIOMIXER_API MixMonoTo8ChannelsFast(const AlignedFloatBuffer& MonoBuffer, AlignedFloatBuffer& DestinationBuffer, const float* RESTRICT StartGains, const float* RESTRICT EndGains);
+	void AUDIOMIXER_API MixMonoTo8ChannelsFast(const float* RESTRICT MonoBuffer, float* RESTRICT DestinationBuffer, int32 NumFrames, const float* RESTRICT StartGains, const float* RESTRICT EndGains);
+
+	/** Takes a 2 channel buffer and mixes it to an 8 channel interleaved buffer using Gains. Gains is expected to point to a 16 float long buffer.
+	 *  Output gains for the left input channel should be the first 8 values in Gains, and Output gains for the right input channel should be rest.
+	 *  these buffers must have an even number of frames.
+	 *  If StartGains and EndGains are provided, this function will interpolate between the two across the buffer.
+	*/
+	void AUDIOMIXER_API Mix2ChannelsTo8ChannelsFast(const AlignedFloatBuffer& SourceBuffer, AlignedFloatBuffer& DestinationBuffer, const float* RESTRICT Gains);
+	void AUDIOMIXER_API Mix2ChannelsTo8ChannelsFast(const float* RESTRICT SourceBuffer, float* RESTRICT DestinationBuffer, int32 NumFrames, const float* RESTRICT Gains);
+	void AUDIOMIXER_API Mix2ChannelsTo8ChannelsFast(const AlignedFloatBuffer& SourceBuffer, AlignedFloatBuffer& DestinationBuffer, const float* RESTRICT StartGains, const float* RESTRICT EndGains);
+	void AUDIOMIXER_API Mix2ChannelsTo8ChannelsFast(const float* RESTRICT SourceBuffer, float* RESTRICT DestinationBuffer, int32 NumFrames, const float* RESTRICT StartGains, const float* RESTRICT EndGains);
+
+	/** This is a generalized operation that uses the channel gain matrix provided in Gains to mix an interleaved source buffer to the interleaved downmix buffer.
+	 *  This operation is not explicitly vectorized and will almost always be slower than using one of the functions above.
+	*/
+	void AUDIOMIXER_API DownmixBuffer(int32 NumSourceChannels, int32 NumDestinationChannels, const AlignedFloatBuffer& SourceBuffer, AlignedFloatBuffer& DestinationBuffer, const float* RESTRICT Gains);
+	void AUDIOMIXER_API DownmixBuffer(int32 NumSourceChannels, int32 NumDestinationChannels, const float* RESTRICT SourceBuffer, float* RESTRICT DestinationBuffer, int32 NumFrames, const float* RESTRICT Gains);
+	void AUDIOMIXER_API DownmixBuffer(int32 NumSourceChannels, int32 NumDestinationChannels, const AlignedFloatBuffer& SourceBuffer, AlignedFloatBuffer& DestinationBuffer, float* RESTRICT StartGains, const float* RESTRICT EndGains);
+	void AUDIOMIXER_API DownmixBuffer(int32 NumSourceChannels, int32 NumDestinationChannels, const float* RESTRICT SourceBuffer, float* RESTRICT DestinationBuffer, int32 NumFrames, float* RESTRICT StartGains, const float* RESTRICT EndGains);
+
+
+	/** Interleaves samples from two input buffers */
+	void AUDIOMIXER_API BufferInterleave2ChannelFast(const AlignedFloatBuffer& InBuffer1, const AlignedFloatBuffer& InBuffer2, AlignedFloatBuffer& OutBuffer);
+
+	/** Interleaves samples from two input buffers */
+	void AUDIOMIXER_API BufferInterleave2ChannelFast(const float* RESTRICT InBuffer1, const float* RESTRICT InBuffer2, float* RESTRICT OutBuffer, const int32 InNum);
+
+	/** Deinterleaves samples from a 2 channel input buffer */
+	void AUDIOMIXER_API BufferDeinterleave2ChannelFast(const AlignedFloatBuffer& InBuffer, AlignedFloatBuffer& OutBuffer1, AlignedFloatBuffer& OutBuffer2);
+
+	/** Deinterleaves samples from a 2 channel input buffer */
+	void AUDIOMIXER_API BufferDeinterleave2ChannelFast(const float* RESTRICT InBuffer, float* RESTRICT OutBuffer1, float* RESTRICT OutBuffer2, const int32 InNum);
+
+	/** Sums 2 channel interleaved input samples. OutSamples[n] = InSamples[2n] + InSamples[2n + 1] */
+	void AUDIOMIXER_API BufferSum2ChannelToMonoFast(const AlignedFloatBuffer& InSamples, AlignedFloatBuffer& OutSamples);
+
+	/** Sums 2 channel interleaved input samples. OutSamples[n] = InSamples[2n] + InSamples[2n + 1] */
+	void AUDIOMIXER_API BufferSum2ChannelToMonoFast(const float* RESTRICT InSamples, float* RESTRICT OutSamples, const int32 InNumFrames);
+
+	/** Compute power of complex data. Out[i] = Real[i] * Real[i] + Imaginary[i] * Imaginary[i] */
+	void AUDIOMIXER_API BufferComplexToPowerFast(const AlignedFloatBuffer& InRealSamples, const AlignedFloatBuffer& InImaginarySamples, AlignedFloatBuffer& OutPowerSamples);
+
+	/** Compute power of complex data. Out[i] = Real[i] * Real[i] + Imaginary[i] * Imaginary[i] */
+	void AUDIOMIXER_API BufferComplexToPowerFast(const float* RESTRICT InRealSamples, const float* RESTRICT InImaginarySamples, float* RESTRICT OutPowerSamples, const int32 InNum);
+
+	/** Compute magnitude of complex data. Out[i] = Sqrt(Real[i] * Real[i] + Imaginary[i] * Imaginary[i]) */
+	void AUDIOMIXER_API BufferComplexToMagnitudeFast(const AlignedFloatBuffer& InRealSamples, const AlignedFloatBuffer& InImaginarySamples, AlignedFloatBuffer& OutPowerSamples);
+
+	/** Compute magnitude of complex data. Out[i] = Sqrt(Real[i] * Real[i] + Imaginary[i] * Imaginary[i]) */
+	void AUDIOMIXER_API BufferComplexToMagnitudeFast(const float* RESTRICT InRealSamples, const float* RESTRICT InImaginarySamples, float* RESTRICT OutPowerSamples, const int32 InNum);
+}
