[NO NON-PII COMMIT MESSAGE]

diff --git a/Engine/Source/Runtime/AudioMixer/Private/AudioMixerSubmix.cpp b/Engine/Source/Runtime/AudioMixer/Private/AudioMixerSubmix.cpp
--- a/Engine/Source/Runtime/AudioMixer/Private/AudioMixerSubmix.cpp
+++ b/Engine/Source/Runtime/AudioMixer/Private/AudioMixerSubmix.cpp
@@ -10,6 +10,14 @@
 // Link to "Audio" profiling category
 CSV_DECLARE_CATEGORY_MODULE_EXTERN(AUDIOMIXERCORE_API, Audio);
 
+static int32 RecoverRecordingOnShutdownCVar = 0;
+FAutoConsoleVariableRef CVarRecoverRecordingOnShutdown(
+	TEXT("au.RecoverRecordingOnShutdown"),
+	RecoverRecordingOnShutdownCVar,
+	TEXT("When set to 1, we will attempt to bounce the recording to a wav file if the game is shutdown while a recording is in flight.\n")
+	TEXT("0: Disabled, 1: Enabled"),
+	ECVF_Default);
+
 namespace Audio
 {
 	// Unique IDs for mixer submixe's
@@ -49,7 +57,7 @@ namespace Audio
 			TearDownAmbisonicsDecoder();
 		}
 
-		if (OwningSubmixObject && bIsRecording)
+		if (RecoverRecordingOnShutdownCVar && OwningSubmixObject && bIsRecording)
 		{
 			FString InterruptedFileName = TEXT("InterruptedRecording.wav");
 			UE_LOG(LogAudioMixer, Warning, TEXT("Recording of Submix %s was interrupted. Saving interrupted recording as %s."), *(OwningSubmixObject->GetName()), *InterruptedFileName);
