Trivial: Change comment. VolumeFader implementation not necessarily tied to AudioComponent architecture

diff --git a/Engine/Source/Runtime/AudioMixer/Public/DSP/VolumeFader.h b/Engine/Source/Runtime/AudioMixer/Public/DSP/VolumeFader.h
--- a/Engine/Source/Runtime/AudioMixer/Public/DSP/VolumeFader.h
+++ b/Engine/Source/Runtime/AudioMixer/Public/DSP/VolumeFader.h
@@ -6,7 +6,7 @@
 
 namespace Audio
 {
-	/** Type of fade to use when adjusting the audio component's volume. */
+	/** Type of fade to use when adjusting a fader's volume over time. */
 	enum class EFaderCurve : uint8
 	{
 		// Linear Fade
