Stopgap safeguard against MarkPendingKill destroying SynthComponent out from under SynthSound while generating audio in an async task.

diff --git a/Engine/Source/Runtime/AudioMixer/Private/Components/SynthComponent.cpp b/Engine/Source/Runtime/AudioMixer/Private/Components/SynthComponent.cpp
--- a/Engine/Source/Runtime/AudioMixer/Private/Components/SynthComponent.cpp
+++ b/Engine/Source/Runtime/AudioMixer/Private/Components/SynthComponent.cpp
@@ -12,6 +12,8 @@ USynthSound::USynthSound(const FObjectInitializer& ObjectInitializer)
 
 void USynthSound::Init(USynthComponent* InSynthComponent, const int32 InNumChannels, const int32 InSampleRate, const int32 InCallbackSize)
 {
+	check(InSynthComponent);
+
 	OwningSynthComponent = InSynthComponent;
 	VirtualizationMode = EVirtualizationMode::PlayWhenSilent;
 	NumChannels = InNumChannels;
@@ -54,6 +56,13 @@ int32 USynthSound::OnGeneratePCMAudio(TArray<uint8>& OutAudio, int32 NumSamples)
 	{
 		// If running with audio mixer, the output audio buffer will be in floats already
 		OutAudio.AddZeroed(NumSamples * sizeof(float));
+
+		// Mark pending kill can null this out on the game thread in rare cases.
+		if (!OwningSynthComponent)
+		{
+			return 0;
+		}
+
 		return OwningSynthComponent->OnGeneratePCMAudio((float*)OutAudio.GetData(), NumSamples);
 	}
 	else
@@ -62,6 +71,12 @@ int32 USynthSound::OnGeneratePCMAudio(TArray<uint8>& OutAudio, int32 NumSamples)
 		FloatBuffer.Reset();
 		FloatBuffer.AddZeroed(NumSamples * sizeof(float));
 
+		// Mark pending kill can null this out on the game thread in rare cases.
+		if (!OwningSynthComponent)
+		{
+			return 0;
+		}
+
 		float* FloatBufferDataPtr = FloatBuffer.GetData();
 		int32 NumSamplesGenerated = OwningSynthComponent->OnGeneratePCMAudio(FloatBufferDataPtr, NumSamples);
 
@@ -80,8 +95,11 @@ int32 USynthSound::OnGeneratePCMAudio(TArray<uint8>& OutAudio, int32 NumSamples)
 
 void USynthSound::OnEndGenerate()
 {
-	check(OwningSynthComponent);
-	OwningSynthComponent->OnEndGenerate();
+	// Mark pending kill can null this out on the game thread in rare cases.
+	if(OwningSynthComponent)
+	{
+		OwningSynthComponent->OnEndGenerate();
+	}
 }
 
 Audio::EAudioMixerStreamDataFormat::Type USynthSound::GetGeneratedPCMDataFormat() const
diff --git a/Engine/Source/Runtime/AudioMixer/Public/Components/SynthComponent.h b/Engine/Source/Runtime/AudioMixer/Public/Components/SynthComponent.h
--- a/Engine/Source/Runtime/AudioMixer/Public/Components/SynthComponent.h
+++ b/Engine/Source/Runtime/AudioMixer/Public/Components/SynthComponent.h
@@ -58,6 +58,12 @@ protected:
 
 	TArray<float> FloatBuffer;
 	bool bAudioMixer;
+
+public:
+	USynthComponent* GetOwningSynthComponent()
+	{
+		return OwningSynthComponent;
+	}
 };
 
 UCLASS(ClassGroup = Synth, hidecategories = (Object, ActorComponent, Physics, Rendering, Mobility, LOD))
diff --git a/Engine/Source/Runtime/Engine/Private/Audio.cpp b/Engine/Source/Runtime/Engine/Private/Audio.cpp
--- a/Engine/Source/Runtime/Engine/Private/Audio.cpp
+++ b/Engine/Source/Runtime/Engine/Private/Audio.cpp
@@ -12,6 +12,7 @@
 #include "AudioPluginUtilities.h"
 #include "AudioThread.h"
 #include "Components/AudioComponent.h"
+#include "Components/SynthComponent.h"
 #include "ContentStreaming.h"
 #include "DrawDebugHelpers.h"
 #include "EngineAnalytics.h"
@@ -877,6 +878,15 @@ FArchive& operator<<( FArchive& Ar, FWaveInstance* WaveInstance )
 void FWaveInstance::AddReferencedObjects( FReferenceCollector& Collector )
 {
 	Collector.AddReferencedObject( WaveData );
+
+	if (USynthSound* SynthSound = Cast<USynthSound>(WaveData))
+	{
+		if (USynthComponent* SynthComponent = SynthSound->GetOwningSynthComponent())
+		{
+			Collector.AddReferencedObject(SynthComponent);
+		}
+	}
+
 	Collector.AddReferencedObject( SoundClass );
 	NotifyBufferFinishedHooks.AddReferencedObjects( Collector );
 }
