Audio - Fix for media sound component playback in PIE for sound components that aren't owned by actors.

diff --git a/Engine/Source/Runtime/AudioMixer/Private/Components/SynthComponent.cpp b/Engine/Source/Runtime/AudioMixer/Private/Components/SynthComponent.cpp
--- a/Engine/Source/Runtime/AudioMixer/Private/Components/SynthComponent.cpp
+++ b/Engine/Source/Runtime/AudioMixer/Private/Components/SynthComponent.cpp
@@ -224,7 +224,8 @@ void USynthComponent::Initialize(int32 SampleRateOverride)
 
 		Synth->Init(this, NumChannels, SampleRate, PreferredBufferLength);
 
-		if (FAudioDevice* AudioDevice = AudioComponent->GetAudioDevice())
+		// Retrieve the synth component's audio device vs the audio component's
+		if (FAudioDevice* AudioDevice = GetAudioDevice())
 		{
 			Synth->StartOnAudioDevice(AudioDevice);
 		}
@@ -243,6 +244,33 @@ void USynthComponent::CreateAudioComponent()
 		// Create the audio component which will be used to play the procedural sound wave
 		AudioComponent = NewObject<UAudioComponent>(this);
 
+		AudioComponent->OnAudioSingleEnvelopeValueNative.AddUObject(this, &USynthComponent::OnAudioComponentEnvelopeValue);
+
+		if (!AudioComponent->GetAttachParent() && !AudioComponent->IsAttachedTo(this))
+		{
+			AActor* Owner = GetOwner();
+
+			// If the media component has no owner or the owner doesn't have a world
+			if (!Owner || !Owner->GetWorld())
+			{
+				// Attempt to retrieve the synth component's world and register the audio component with it
+				// This ensures that the synth component plays on the correct world in cases where there isn't an owner
+				if (UWorld* World = GetWorld())
+				{
+					AudioComponent->RegisterComponentWithWorld(World);
+					AudioComponent->AttachToComponent(this, FAttachmentTransformRules::KeepRelativeTransform);
+				}
+				else
+				{
+					AudioComponent->SetupAttachment(this);
+				}
+			}
+			else
+			{
+				AudioComponent->AttachToComponent(this, FAttachmentTransformRules::KeepRelativeTransform);
+				AudioComponent->RegisterComponent();
+			}
+		}
 	}
 
 	if (AudioComponent)
@@ -256,20 +284,6 @@ void USynthComponent::CreateAudioComponent()
 #if WITH_EDITORONLY_DATA
 		AudioComponent->bVisualizeComponent = false;
 #endif
-		if (!AudioComponent->GetAttachParent() && !AudioComponent->IsAttachedTo(this))
-		{
-			if (!GetOwner() || !GetOwner()->GetWorld())
-			{
-				AudioComponent->SetupAttachment(this);
-			}
-			else
-			{
-				AudioComponent->AttachToComponent(this, FAttachmentTransformRules::KeepRelativeTransform);
-				AudioComponent->RegisterComponent();
-			}
-		}
-
-		AudioComponent->OnAudioSingleEnvelopeValueNative.AddUObject(this, &USynthComponent::OnAudioComponentEnvelopeValue);
 
 		// Set defaults to be the same as audio component defaults
 		AudioComponent->EnvelopeFollowerAttackTime = EnvelopeFollowerAttackTime;
@@ -380,6 +394,30 @@ void USynthComponent::PumpPendingMessages()
 	}
 }
 
+FAudioDevice* USynthComponent::GetAudioDevice()
+{
+	// If the synth component has a world, that means it was already registed with that world
+	if (UWorld* World = GetWorld())
+	{
+		// Make sure it has a proper audio device handle and retrieve it
+		if (World->AudioDeviceHandle != INDEX_NONE)
+		{
+			FAudioDeviceManager* AudioDeviceManager = GEngine->GetAudioDeviceManager();
+			check(AudioDeviceManager);
+			return AudioDeviceManager->GetAudioDevice(World->AudioDeviceHandle);
+		}
+	}
+
+	// Otherwise, retrieve the audio component's audio device (probably from it's owner)
+	if (AudioComponent)
+	{
+		return AudioComponent->GetAudioDevice();
+	}
+	
+	// No audio device
+	return nullptr;
+}
+
 int32 USynthComponent::OnGeneratePCMAudio(float* GeneratedPCMData, int32 NumSamples)
 {
 	PumpPendingMessages();
@@ -490,4 +528,4 @@ void USynthComponent::SetSubmixSend(USoundSubmix* Submix, float SendLevel)
 void USynthComponent::SynthCommand(TFunction<void()> Command)
 {
 	CommandQueue.Enqueue(MoveTemp(Command));
-}
\ No newline at end of file
+}
diff --git a/Engine/Source/Runtime/AudioMixer/Public/Components/SynthComponent.h b/Engine/Source/Runtime/AudioMixer/Public/Components/SynthComponent.h
--- a/Engine/Source/Runtime/AudioMixer/Public/Components/SynthComponent.h
+++ b/Engine/Source/Runtime/AudioMixer/Public/Components/SynthComponent.h
@@ -241,7 +241,7 @@ protected:
 	int32 OnGeneratePCMAudio(float* GeneratedPCMData, int32 NumSamples);
 
 	// Gets the audio device associated with this synth component
-	FAudioDevice* GetAudioDevice() { return AudioComponent ? AudioComponent->GetAudioDevice() : nullptr; }
+	FAudioDevice* GetAudioDevice();
 
 	// Can be set by the derived class, defaults to 2
 	int32 NumChannels;
