- Make AudioMixerSource construction/pre-init phase more concise (i.e. follow better RAII pattern) by using static factory constructor and removing pre-init call,
protecting against edge cases around internal ref counting being missed with certain error conditions, teardown being done prematurally, etc.
- Remove "FMixerSourceBuffer::ClearSoundWave" which doesn't actually do what it implies and just call OnEndGenerate directly for single caller

diff --git a/Engine/Source/Runtime/AudioMixer/Private/AudioMixerSource.cpp b/Engine/Source/Runtime/AudioMixer/Private/AudioMixerSource.cpp
--- a/Engine/Source/Runtime/AudioMixer/Private/AudioMixerSource.cpp
+++ b/Engine/Source/Runtime/AudioMixer/Private/AudioMixerSource.cpp
@@ -28,7 +28,6 @@ namespace Audio
 		: FSoundSource(InAudioDevice)
 		, MixerDevice(static_cast<FMixerDevice*>(InAudioDevice))
 		, MixerBuffer(nullptr)
-		, MixerSourceBuffer(nullptr)
 		, MixerSourceVoice(nullptr)
 		, PreviousAzimuth(-1.0f)
 		, PreviousPlaybackPercent(0.0f)
@@ -430,56 +429,54 @@ namespace Audio
 		//  Reset so next instance will warn if algorithm changes in-flight
 		bEditorWarnedChangedSpatialization = false;
 
+		const bool bIsSeeking = InWaveInstance->StartTime > 0.0f;
+
 		check(InWaveInstance);
-		check(MixerBuffer == nullptr);
 		check(AudioDevice);
 
-		bool bIsSeeking = InWaveInstance->StartTime > 0.0f;
-		MixerBuffer = FMixerBuffer::Init(AudioDevice, InWaveInstance->WaveData, bIsSeeking);
-		if (MixerBuffer)
+		check(!MixerBuffer);
+		MixerBuffer = FMixerBuffer::Init(AudioDevice, InWaveInstance->WaveData, bIsSeeking /* bForceRealtime */);
+
+		if (!MixerBuffer)
 		{
-			Buffer = MixerBuffer;
-			WaveInstance = InWaveInstance;
+			MixerSourceBuffer.Reset();
+			return false;
+		}
 
-			LPFFrequency = MAX_FILTER_FREQUENCY;
-			LastLPFFrequency = FLT_MAX;
+		// WaveData must be valid beyond this point, otherwise MixerBuffer
+		// would have failed to init.
+		check(InWaveInstance->WaveData);
+		USoundWave& SoundWave = *InWaveInstance->WaveData;
 
-			HPFFrequency = 0.0;
-			LastHPFFrequency = FLT_MAX;
+		Buffer = MixerBuffer;
+		WaveInstance = InWaveInstance;
 
-			bIsDone = false;
+		LPFFrequency = MAX_FILTER_FREQUENCY;
+		LastLPFFrequency = FLT_MAX;
 
-			// Not all wave data types have a non-zero duration
-			if (InWaveInstance->WaveData->Duration > 0)
-			{
-				if (!InWaveInstance->WaveData->bIsBus)
-				{
-					NumTotalFrames = InWaveInstance->WaveData->Duration * InWaveInstance->WaveData->GetSampleRateForCurrentPlatform();
-					check(NumTotalFrames > 0);
-				}
-				else if (!InWaveInstance->WaveData->IsLooping())
-				{
-					NumTotalFrames = InWaveInstance->WaveData->Duration * AudioDevice->GetSampleRate();
-					check(NumTotalFrames > 0);
-				}
-			}
+		HPFFrequency = 0.0f;
+		LastHPFFrequency = FLT_MAX;
 
-			check(!MixerSourceBuffer.IsValid());
-			MixerSourceBuffer = TSharedPtr<FMixerSourceBuffer>(new FMixerSourceBuffer());
+		bIsDone = false;
 
-			if (MixerSourceBuffer->PreInit(MixerBuffer, InWaveInstance->WaveData, InWaveInstance->LoopingMode, bIsSeeking))
+		// Not all wave data types have a non-zero duration
+		if (SoundWave.Duration > 0.0f)
+		{
+			if (!SoundWave.bIsBus)
 			{
-				// We succeeded in preparing the buffer for initialization, but we are not technically initialized yet.
-				// If the buffer is asynchronously preparing a file-handle, we may not yet initialize the source.
-				return true;
+				NumTotalFrames = SoundWave.Duration * SoundWave.GetSampleRateForCurrentPlatform();
+				check(NumTotalFrames > 0);
+			}
+			else if (!SoundWave.IsLooping())
+			{
+				NumTotalFrames = SoundWave.Duration * AudioDevice->GetSampleRate();
+				check(NumTotalFrames > 0);
 			}
 		}
 
-		// Clear out our mixer source buffer if things failed
-		MixerSourceBuffer = nullptr;
-
-		// Something went wrong with initializing the generator
-		return false;
+		check(!MixerSourceBuffer.IsValid());
+		MixerSourceBuffer = FMixerSourceBuffer::Create(*MixerBuffer, SoundWave, InWaveInstance->LoopingMode, bIsSeeking);
+		return MixerSourceBuffer.IsValid();
 	}
 
 	bool FMixerSource::IsPreparedToInit()
@@ -790,7 +787,7 @@ namespace Audio
 			MixerSourceVoice = nullptr;
 		}
 
-		MixerSourceBuffer = nullptr;
+		MixerSourceBuffer.Reset();
 		Buffer = nullptr;
 		bLoopCallback = false;
 		NumTotalFrames = 0;
diff --git a/Engine/Source/Runtime/AudioMixer/Private/AudioMixerSourceBuffer.cpp b/Engine/Source/Runtime/AudioMixer/Private/AudioMixerSourceBuffer.cpp
--- a/Engine/Source/Runtime/AudioMixer/Private/AudioMixerSourceBuffer.cpp
+++ b/Engine/Source/Runtime/AudioMixer/Private/AudioMixerSourceBuffer.cpp
@@ -55,24 +55,53 @@ namespace Audio
 		return CurrentSample >= NumSamples;
 	}
 
-	FMixerSourceBuffer::FMixerSourceBuffer()
+	TSharedPtr<FMixerSourceBuffer> FMixerSourceBuffer::Create(FMixerBuffer& InBuffer, USoundWave& InWave, ELoopingMode InLoopingMode, bool bInIsSeeking)
+	{
+		LLM_SCOPE(ELLMTag::AudioMixer);
+
+		// Prevent double-triggering procedural soundwaves
+		if (InWave.bProcedural && InWave.IsGeneratingAudio())
+		{
+			UE_LOG(LogAudioMixer, Warning,
+				TEXT("Procedural USoundWave is reinitializing even though it is actively "
+				"generating audio. Sound must be stopped before playing again."));
+			return nullptr;
+		}
+
+		TSharedPtr<FMixerSourceBuffer> NewSourceBuffer = MakeShareable(new FMixerSourceBuffer(InBuffer, InWave, InLoopingMode, bInIsSeeking));
+		return NewSourceBuffer;
+	}
+
+	FMixerSourceBuffer::FMixerSourceBuffer(FMixerBuffer& InBuffer, USoundWave& InWave, ELoopingMode InLoopingMode, bool bInIsSeeking)
 		: NumBuffersQeueued(0)
 		, CurrentBuffer(0)
-		, SoundWave(nullptr)
+		, SoundWave(&InWave)
 		, AsyncRealtimeAudioTask(nullptr)
 		, DecompressionState(nullptr)
-		, LoopingMode(ELoopingMode::LOOP_Never)
-		, NumChannels(0)
-		, BufferType(Audio::EBufferType::Invalid)
-		, NumPrecacheFrames(0)
+		, LoopingMode(InLoopingMode)
+		, NumChannels(InBuffer.NumChannels)
+		, BufferType(InBuffer.GetType())
+		, NumPrecacheFrames(InWave.NumPrecacheFrames)
 		, bInitialized(false)
 		, bBufferFinished(false)
 		, bPlayedCachedBuffer(false)
-		, bIsSeeking(false)
+		, bIsSeeking(bInIsSeeking)
 		, bLoopCallback(false)
-		, bProcedural(false)
-		, bIsBus(false)
+		, bProcedural(InWave.bProcedural)
+		, bIsBus(InWave.bIsBus)
 	{
+		InWave.AddPlayingSource();
+
+		const uint32 TotalSamples = MONO_PCM_BUFFER_SAMPLES * NumChannels;
+		for (int32 BufferIndex = 0; BufferIndex < Audio::MAX_BUFFERS_QUEUED; ++BufferIndex)
+		{
+			SourceVoiceBuffers.Add(MakeShared<FMixerSourceVoiceBuffer>());
+
+			// Prepare the memory to fit the max number of samples
+			SourceVoiceBuffers[BufferIndex]->AudioData.Reset(TotalSamples);
+			SourceVoiceBuffers[BufferIndex]->bRealTimeBuffer = true;
+			SourceVoiceBuffers[BufferIndex]->LoopCount = 0;
+		}
 	}
 
 	FMixerSourceBuffer::~FMixerSourceBuffer()
@@ -105,52 +134,6 @@ namespace Audio
 		}
 	}
 
-	bool FMixerSourceBuffer::PreInit(FMixerBuffer* InBuffer, USoundWave* InWave, ELoopingMode InLoopingMode, bool bInIsSeeking)
-	{
-		LLM_SCOPE(ELLMTag::AudioMixer);
-		if (!InWave)
-		{
-			return false;
-		}
-
-		NumChannels = InBuffer->NumChannels;
-		BufferType = InBuffer->GetType();
-		bIsBus = InWave->bIsBus;
-		bProcedural = InWave->bProcedural;
-		NumPrecacheFrames = InWave->NumPrecacheFrames;
-		
-		// Prevent double-triggering procedural soundwaves
-		if (bProcedural && InWave->IsGeneratingAudio())
-		{
-			UE_LOG(LogAudioMixer, Warning, TEXT("Procedural sound wave is reinitializing even though it is currently actively generating audio. Please stop sound before trying to play it again."));
-			return false;
-		}
-
-		check(SoundWave == nullptr);
-		// Only need to have a handle to a USoundWave for procedural sound waves
-		SoundWave = InWave;
-		check(SoundWave);
-		SoundWave->AddPlayingSource();
-
-		LoopingMode = InLoopingMode;
-		bIsSeeking = bInIsSeeking;
-		bLoopCallback = false;
-
-		BufferQueue.Empty();
-
-		const uint32 TotalSamples = MONO_PCM_BUFFER_SAMPLES * NumChannels;
-		for (int32 BufferIndex = 0; BufferIndex < Audio::MAX_BUFFERS_QUEUED; ++BufferIndex)
-		{
-			SourceVoiceBuffers.Add(TSharedPtr<FMixerSourceVoiceBuffer>(new FMixerSourceVoiceBuffer()));
-
-			// Prepare the memory to fit the max number of samples
-			SourceVoiceBuffers[BufferIndex]->AudioData.Reset(TotalSamples);
-			SourceVoiceBuffers[BufferIndex]->bRealTimeBuffer = true;
-			SourceVoiceBuffers[BufferIndex]->LoopCount = 0;
-		}
-		return true;
-	}
-
 	void FMixerSourceBuffer::SetDecoder(ICompressedAudioInfo* InCompressedAudioInfo)
 	{
 		if (DecompressionState == nullptr)
@@ -489,12 +472,6 @@ namespace Audio
 		}
 	}
 
-	void FMixerSourceBuffer::ClearSoundWave()
-	{
-		// Call on end generate right now, before destructor
-		OnEndGenerate();
-	}
-
 	void FMixerSourceBuffer::OnBeginGenerate()
 	{
 		if (bProcedural)
diff --git a/Engine/Source/Runtime/AudioMixer/Private/AudioMixerSourceBuffer.h b/Engine/Source/Runtime/AudioMixer/Private/AudioMixerSourceBuffer.h
--- a/Engine/Source/Runtime/AudioMixer/Private/AudioMixerSourceBuffer.h
+++ b/Engine/Source/Runtime/AudioMixer/Private/AudioMixerSourceBuffer.h
@@ -46,10 +46,10 @@ namespace Audio
 	class FMixerSourceBuffer
 	{
 	public:
-		FMixerSourceBuffer();
+		static TSharedPtr<FMixerSourceBuffer> Create(FMixerBuffer& InBuffer, USoundWave& InWave, ELoopingMode InLoopingMode, bool bInIsSeeking);
+
 		~FMixerSourceBuffer();
 
-		bool PreInit(FMixerBuffer* InBuffer, USoundWave* InWave, ELoopingMode InLoopingMode, bool bInIsSeeking);
 		bool Init();
 
 		// Sets the decoder to use for realtime async decoding
@@ -88,14 +88,12 @@ namespace Audio
 		// Ensures the async task finishes
 		void EnsureAsyncTaskFinishes();
 
-		// Clear the sound wave reference
-		void ClearSoundWave();
-
 		// Begin and end generation on the audio render thread (audio mixer only)
 		void OnBeginGenerate();
 		void OnEndGenerate();
 
 	private:
+		FMixerSourceBuffer(FMixerBuffer& InBuffer, USoundWave& InWave, ELoopingMode InLoopingMode, bool bInIsSeeking);
 
 		void SubmitInitialPCMBuffers();
 		void SubmitInitialRealtimeBuffers();
@@ -127,6 +125,5 @@ namespace Audio
 		uint32 bLoopCallback : 1;
 		uint32 bProcedural : 1;
 		uint32 bIsBus : 1;
-
 	};
 }
diff --git a/Engine/Source/Runtime/AudioMixer/Private/SoundWaveDecoder.cpp b/Engine/Source/Runtime/AudioMixer/Private/SoundWaveDecoder.cpp
--- a/Engine/Source/Runtime/AudioMixer/Private/SoundWaveDecoder.cpp
+++ b/Engine/Source/Runtime/AudioMixer/Private/SoundWaveDecoder.cpp
@@ -15,7 +15,6 @@ namespace Audio
 	{
 		SourceInfo.VolumeParam.Init();
 		SourceInfo.VolumeParam.SetValue(InitData.VolumeScale);
-
 		SourceInfo.PitchScale = InitData.PitchScale;
 
 		MixerBuffer = FMixerBuffer::Init(AudioDevice, InitData.SoundWave, true);
@@ -23,7 +22,10 @@ namespace Audio
 
 	FDecodingSoundSource::~FDecodingSoundSource()
 	{
-		MixerSourceBuffer.ClearSoundWave();
+		if (MixerSourceBuffer.IsValid())
+		{
+			MixerSourceBuffer->OnEndGenerate();
+		}
 	}
 
 	bool FDecodingSoundSource::PreInit(int32 InSampleRate)
@@ -35,20 +37,32 @@ namespace Audio
 		SineTone[1].Init(InSampleRate, 440.0f, 0.5f);
 #endif
 
-		ELoopingMode LoopingMode = SoundWave->bLooping ? ELoopingMode::LOOP_Forever : ELoopingMode::LOOP_Never;
+		if (!SoundWave || !MixerBuffer)
+		{
+			return nullptr;
+		}
 
-		return MixerSourceBuffer.PreInit(MixerBuffer, SoundWave, LoopingMode, SeekTime > 0.0f);
+		const ELoopingMode LoopingMode = SoundWave->bLooping ? ELoopingMode::LOOP_Forever : ELoopingMode::LOOP_Never;
+		const bool bIsSeeking = SeekTime > 0.0f;
+
+		MixerSourceBuffer = FMixerSourceBuffer::Create(*MixerBuffer, *SoundWave, LoopingMode, bIsSeeking);
+		return MixerSourceBuffer.IsValid();
 	}
 
 	bool FDecodingSoundSource::IsReadyToInit()
 	{
+		if (!MixerSourceBuffer.IsValid())
+		{
+			return false;
+		}
+
 		if (MixerBuffer && MixerBuffer->IsRealTimeSourceReady())
 		{
 			// Check if we have a realtime audio task already (doing first decode)
-			if (MixerSourceBuffer.IsAsyncTaskInProgress())
+			if (MixerSourceBuffer->IsAsyncTaskInProgress())
 			{
 				// not ready
-				return MixerSourceBuffer.IsAsyncTaskDone();
+				return MixerSourceBuffer->IsAsyncTaskDone();
 			}
 			else
 			{
@@ -67,7 +81,7 @@ namespace Audio
 
 						ICompressedAudioInfo* CompressedAudioInfo = MixerBuffer->GetDecompressionState(false);
 
-						MixerSourceBuffer.ReadMoreRealtimeData(CompressedAudioInfo, 0, EBufferReadMode::Asynchronous);
+						MixerSourceBuffer->ReadMoreRealtimeData(CompressedAudioInfo, 0, EBufferReadMode::Asynchronous);
 
 						// not ready
 						return false;
@@ -84,30 +98,34 @@ namespace Audio
 	{
 		if (MixerBuffer->GetNumChannels() > 0 && MixerBuffer->GetNumChannels() <= 2)
 		{
-			// Pass the decompression state off to the mixer source buffer if it hasn't already done so
-			ICompressedAudioInfo* Decoder = MixerBuffer->GetDecompressionState(false);
-			MixerSourceBuffer.SetDecoder(Decoder);
-			
-			if (!MixerSourceBuffer.IsAsyncTaskInProgress())
+			if (MixerSourceBuffer.IsValid())
 			{
-				MixerSourceBuffer.ReadMoreRealtimeData(Decoder, 0, EBufferReadMode::Asynchronous);
-			}
 
-			SourceInfo.NumSourceChannels = MixerBuffer->GetNumChannels();
-			SourceInfo.TotalNumFrames = MixerBuffer->GetNumFrames();
+				// Pass the decompression state off to the mixer source buffer if it hasn't already done so
+				ICompressedAudioInfo* Decoder = MixerBuffer->GetDecompressionState(false);
+				MixerSourceBuffer->SetDecoder(Decoder);
 
-			SourceInfo.CurrentFrameValues.AddZeroed(SourceInfo.NumSourceChannels);
-			SourceInfo.NextFrameValues.AddZeroed(SourceInfo.NumSourceChannels);
+				if (!MixerSourceBuffer->IsAsyncTaskInProgress())
+				{
+					MixerSourceBuffer->ReadMoreRealtimeData(Decoder, 0, EBufferReadMode::Asynchronous);
+				}
 
-			SourceInfo.BasePitchScale = MixerBuffer->GetSampleRate() / SampleRate;
-		
-			SourceInfo.PitchParam.Init();
-			SourceInfo.PitchParam.SetValue(SourceInfo.BasePitchScale * SourceInfo.PitchScale);
+				SourceInfo.NumSourceChannels = MixerBuffer->GetNumChannels();
+				SourceInfo.TotalNumFrames = MixerBuffer->GetNumFrames();
 
-			MixerSourceBuffer.Init();
-			MixerSourceBuffer.OnBeginGenerate();
+				SourceInfo.CurrentFrameValues.AddZeroed(SourceInfo.NumSourceChannels);
+				SourceInfo.NextFrameValues.AddZeroed(SourceInfo.NumSourceChannels);
 
-			bInitialized = true;
+				SourceInfo.BasePitchScale = MixerBuffer->GetSampleRate() / SampleRate;
+
+				SourceInfo.PitchParam.Init();
+				SourceInfo.PitchParam.SetValue(SourceInfo.BasePitchScale * SourceInfo.PitchScale);
+
+				MixerSourceBuffer->Init();
+				MixerSourceBuffer->OnBeginGenerate();
+
+				bInitialized = true;
+			}
 		}
 	}
 
@@ -153,12 +171,12 @@ namespace Audio
 					break;
 				}
 
-				MixerSourceBuffer.OnBufferEnd();
+				MixerSourceBuffer->OnBufferEnd();
 			}
 
-			if (MixerSourceBuffer.GetNumBuffersQueued() > 0)
+			if (MixerSourceBuffer->GetNumBuffersQueued() > 0)
 			{
-				SourceInfo.CurrentPCMBuffer = MixerSourceBuffer.GetNextBuffer();
+				SourceInfo.CurrentPCMBuffer = MixerSourceBuffer->GetNextBuffer();
 				SourceInfo.CurrentAudioChunkNumFrames = SourceInfo.CurrentPCMBuffer->AudioData.Num() / SourceInfo.NumSourceChannels;
 
 				if (bReadCurrentFrame)
@@ -399,7 +417,7 @@ namespace Audio
 			return true;
 		}
 
-		UE_LOG(LogAudioMixer, Warning, TEXT("Failed to initialize sound wave %s."), *InitData.SoundWave->GetName());
+		UE_LOG(LogAudioMixer, Warning, TEXT("Failed to initialize sound wave %s."), InitData.SoundWave ? *InitData.SoundWave->GetName() : TEXT("Unset"));
 		return false;
 	}
 
diff --git a/Engine/Source/Runtime/AudioMixer/Public/SoundWaveDecoder.h b/Engine/Source/Runtime/AudioMixer/Public/SoundWaveDecoder.h
--- a/Engine/Source/Runtime/AudioMixer/Public/SoundWaveDecoder.h
+++ b/Engine/Source/Runtime/AudioMixer/Public/SoundWaveDecoder.h
@@ -87,7 +87,7 @@ namespace Audio
 		FMixerBuffer* MixerBuffer;
 
 		// Object which handles bulk of decoding operations
-		FMixerSourceBuffer MixerSourceBuffer;
+		TSharedPtr<FMixerSourceBuffer> MixerSourceBuffer;
 
 		// Sample rate of the source
 		int32 SampleRate;
