Trivial: Fix implicit cast compiler error in iOS/Mac

diff --git a/Engine/Source/Runtime/AudioMixer/Private/SoundWaveDecoder.cpp b/Engine/Source/Runtime/AudioMixer/Private/SoundWaveDecoder.cpp
--- a/Engine/Source/Runtime/AudioMixer/Private/SoundWaveDecoder.cpp
+++ b/Engine/Source/Runtime/AudioMixer/Private/SoundWaveDecoder.cpp
@@ -39,7 +39,7 @@ namespace Audio
 
 		if (!SoundWave || !MixerBuffer)
 		{
-			return nullptr;
+			return false;
 		}
 
 		const ELoopingMode LoopingMode = SoundWave->bLooping ? ELoopingMode::LOOP_Forever : ELoopingMode::LOOP_Never;
