Bandaid fix for derefering a GCd SoundWave pointer during teardown.

diff --git a/Engine/Source/Runtime/AudioMixer/Private/AudioMixerSourceBuffer.cpp b/Engine/Source/Runtime/AudioMixer/Private/AudioMixerSourceBuffer.cpp
--- a/Engine/Source/Runtime/AudioMixer/Private/AudioMixerSourceBuffer.cpp
+++ b/Engine/Source/Runtime/AudioMixer/Private/AudioMixerSourceBuffer.cpp
@@ -474,7 +474,7 @@ namespace Audio
 
 	void FMixerSourceBuffer::OnBeginGenerate()
 	{
-		if (bProcedural)
+		if (SoundWave && bProcedural)
 		{
 			check(SoundWave && SoundWave->bProcedural);
 			SoundWave->OnBeginGenerate();
@@ -482,12 +482,12 @@ namespace Audio
 	}
 
 	void FMixerSourceBuffer::OnEndGenerate()
-	{
+	{		
 		// Make sure the async task finishes!
 		EnsureAsyncTaskFinishes();
-
+						
 		// Only need to call OnEndGenerate and access SoundWave here if we successfully initialized
-		if (bInitialized && bProcedural)
+		if (SoundWave && bInitialized && bProcedural)
 		{
 			check(SoundWave && SoundWave->bProcedural);
 			SoundWave->OnEndGenerate();
diff --git a/Engine/Source/Runtime/AudioMixer/Private/AudioMixerSourceBuffer.h b/Engine/Source/Runtime/AudioMixer/Private/AudioMixerSourceBuffer.h
--- a/Engine/Source/Runtime/AudioMixer/Private/AudioMixerSourceBuffer.h
+++ b/Engine/Source/Runtime/AudioMixer/Private/AudioMixerSourceBuffer.h
@@ -91,7 +91,7 @@ namespace Audio
 		// Begin and end generation on the audio render thread (audio mixer only)
 		void OnBeginGenerate();
 		void OnEndGenerate();
-
+		void ClearWave() { SoundWave = nullptr; }
 	private:
 		FMixerSourceBuffer(FMixerBuffer& InBuffer, USoundWave& InWave, ELoopingMode InLoopingMode, bool bInIsSeeking);
 
diff --git a/Engine/Source/Runtime/AudioMixer/Private/AudioMixerSourceManager.cpp b/Engine/Source/Runtime/AudioMixer/Private/AudioMixerSourceManager.cpp
--- a/Engine/Source/Runtime/AudioMixer/Private/AudioMixerSourceManager.cpp
+++ b/Engine/Source/Runtime/AudioMixer/Private/AudioMixerSourceManager.cpp
@@ -149,6 +149,16 @@ namespace Audio
 		}
 
 		FPlatformProcess::ReturnSynchEventToPool(CommandsProcessedEvent);
+
+		// We are tearing the whole system down, so clear the MixerSourceBuffers Wave pointer
+		// prior to them destructing, which prevents them doing the normal OnGenerateEnd flow.
+		for (FSourceInfo& i : SourceInfos)
+		{
+			if (FMixerSourceBuffer* Buffer = i.MixerSourceBuffer.Get())
+			{				
+				Buffer->ClearWave(); // NOTE; this does not decrement the active source count on the wave and assumes the pointer is unsafe.
+			}
+		}
 	}
 
 	void FMixerSourceManager::Init(const FSourceManagerInitParams& InitParams)
