Apply .cpp changes from CL# 8497701 in addition to the .h changes in CL# 9739964
#rnx

diff --git a/Engine/Source/Runtime/AudioMixer/Private/AudioMixerSource.cpp b/Engine/Source/Runtime/AudioMixer/Private/AudioMixerSource.cpp
--- a/Engine/Source/Runtime/AudioMixer/Private/AudioMixerSource.cpp
+++ b/Engine/Source/Runtime/AudioMixer/Private/AudioMixerSource.cpp
@@ -1,6 +1,7 @@
 // Copyright 1998-2019 Epic Games, Inc. All Rights Reserved.
 
 #include "AudioMixerSource.h"
+#include "AudioMixerSourceBuffer.h"
 #include "ActiveSound.h"
 #include "AudioMixerSourceBuffer.h"
 #include "AudioMixerDevice.h"
