Fix-up non-unity build forward declaration compiler errors

diff --git a/Engine/Source/Runtime/AudioMixer/Private/AudioMixerSubmix.cpp b/Engine/Source/Runtime/AudioMixer/Private/AudioMixerSubmix.cpp
--- a/Engine/Source/Runtime/AudioMixer/Private/AudioMixerSubmix.cpp
+++ b/Engine/Source/Runtime/AudioMixer/Private/AudioMixerSubmix.cpp
@@ -3,8 +3,9 @@
 #include "AudioMixerSubmix.h"
 #include "AudioMixerDevice.h"
 #include "AudioMixerSourceVoice.h"
-#include "Sound/SoundSubmix.h"
 #include "Sound/SoundEffectSubmix.h"
+#include "Sound/SoundSubmix.h"
+#include "Sound/SoundSubmixSend.h"
 #include "ProfilingDebugging/CsvProfiler.h"
 
 // Link to "Audio" profiling category
@@ -20,7 +21,7 @@ FAutoConsoleVariableRef CVarRecoverRecordingOnShutdown(
 
 namespace Audio
 {
-	// Unique IDs for mixer submixe's
+	// Unique IDs for mixer submixes
 	static uint32 GSubmixMixerIDs = 0;
 
 	FMixerSubmix::FMixerSubmix(FMixerDevice* InMixerDevice)
diff --git a/Engine/Source/Runtime/AudioMixer/Public/AudioMixerDevice.h b/Engine/Source/Runtime/AudioMixer/Public/AudioMixerDevice.h
--- a/Engine/Source/Runtime/AudioMixer/Public/AudioMixerDevice.h
+++ b/Engine/Source/Runtime/AudioMixer/Public/AudioMixerDevice.h
@@ -6,8 +6,12 @@
 #include "AudioMixerSourceManager.h"
 #include "AudioDevice.h"
 
+
+// Forward Declarations
+class FOnSubmixEnvelopeBP;
 class IAudioMixerPlatformInterface;
 
+
 namespace Audio
 {
 	class FMixerSourceVoice;
diff --git a/Engine/Source/Runtime/AudioMixer/Public/AudioMixerSubmix.h b/Engine/Source/Runtime/AudioMixer/Public/AudioMixerSubmix.h
--- a/Engine/Source/Runtime/AudioMixer/Public/AudioMixerSubmix.h
+++ b/Engine/Source/Runtime/AudioMixer/Public/AudioMixerSubmix.h
@@ -8,9 +8,13 @@
 #include "DSP/EnvelopeFollower.h"
 #include "DSP/SpectrumAnalyzer.h"
 
+
+// Forward Declarations
+class FOnSubmixEnvelopeBP;
 class USoundEffectSubmix;
 class USoundSubmix;
 
+
 namespace Audio
 {
 	class IAudioMixerEffect;
diff --git a/Engine/Source/Runtime/Engine/Classes/Sound/SoundSubmix.h b/Engine/Source/Runtime/Engine/Classes/Sound/SoundSubmix.h
--- a/Engine/Source/Runtime/Engine/Classes/Sound/SoundSubmix.h
+++ b/Engine/Source/Runtime/Engine/Classes/Sound/SoundSubmix.h
@@ -2,12 +2,12 @@
 #pragma once
 
 #include "CoreMinimal.h"
-#include "UObject/ObjectMacros.h"
-#include "UObject/Object.h"
+
+#include "IAmbisonicsMixer.h"
 #include "SampleBufferIO.h"
 #include "SoundEffectSubmix.h"
-#include "IAmbisonicsMixer.h"
-#include "Curves/CurveFloat.h"
+#include "UObject/Object.h"
+#include "UObject/ObjectMacros.h"
 
 #include "SoundSubmix.generated.h"
 
