[NO NON-PII COMMIT MESSAGE]

diff --git a/Engine/Source/Runtime/AudioMixer/Private/AudioMixerSource.cpp b/Engine/Source/Runtime/AudioMixer/Private/AudioMixerSource.cpp
--- a/Engine/Source/Runtime/AudioMixer/Private/AudioMixerSource.cpp
+++ b/Engine/Source/Runtime/AudioMixer/Private/AudioMixerSource.cpp
@@ -440,7 +440,7 @@ namespace Audio
 
 		if (!MixerBuffer)
 		{
-			MixerSourceBuffer.Reset();
+			FreeResources(); // APM: maybe need to call this here too? 
 			return false;
 		}
 
@@ -477,6 +477,12 @@ namespace Audio
 
 		check(!MixerSourceBuffer.IsValid());
 		MixerSourceBuffer = FMixerSourceBuffer::Create(*MixerBuffer, SoundWave, InWaveInstance->LoopingMode, bIsSeeking);
+		
+		if (!MixerSourceBuffer.IsValid())
+		{
+			FreeResources();
+		}
+		
 		return MixerSourceBuffer.IsValid();
 	}
 
