[NO NON-PII COMMIT MESSAGE]

diff --git a/Engine/Source/Runtime/AudioMixer/Private/DSP/MultithreadedPatching.cpp b/Engine/Source/Runtime/AudioMixer/Private/DSP/MultithreadedPatching.cpp
new file mode 100644
--- /dev/null
+++ b/Engine/Source/Runtime/AudioMixer/Private/DSP/MultithreadedPatching.cpp
@@ -0,0 +1,455 @@
+// Copyright 1998-2019 Epic Games, Inc. All Rights Reserved.
+
+#include "DSP/MultithreadedPatching.h"
+#include "DSP/BufferVectorOperations.h"
+
+
+namespace Audio
+{
+	TAtomic<int32> FPatchOutput::PatchIDCounter(0);
+
+	FPatchOutput::FPatchOutput(int32 MaxCapacity, float InGain /*= 1.0f*/)
+		: InternalBuffer(MaxCapacity)
+		, TargetGain(InGain)
+		, PreviousGain(InGain)
+		, PatchID(++PatchIDCounter)
+		, NumAliveInputs(0)
+	{
+
+	}
+
+
+	FPatchOutput::FPatchOutput()
+		: InternalBuffer(0)
+		, TargetGain(0.0f)
+		, PreviousGain(0.0f)
+		, PatchID(INDEX_NONE)
+		, NumAliveInputs(0)
+	{
+	}
+
+	int32 FPatchOutput::PopAudio(float* OutBuffer, int32 NumSamples, bool bUseLatestAudio)
+	{
+		if (IsInputStale())
+		{
+			return -1;
+		}
+
+		if (bUseLatestAudio && InternalBuffer.Num() > (uint32) NumSamples)
+		{
+			InternalBuffer.SetNum(((uint32)NumSamples));
+		}
+
+		int32 PopResult = InternalBuffer.Pop(OutBuffer, NumSamples);
+
+		// Apply gain stage.
+		if (FMath::IsNearlyEqual(TargetGain, PreviousGain))
+		{
+			MultiplyBufferByConstantInPlace(OutBuffer, PopResult, PreviousGain);
+		}
+		else
+		{
+			FadeBufferFast(OutBuffer, PopResult, PreviousGain, TargetGain);
+			PreviousGain = TargetGain;
+		}
+		
+		return PopResult;
+	}
+
+	bool FPatchOutput::IsInputStale() const
+	{
+		return NumAliveInputs == 0;
+	}
+
+	int32 FPatchOutput::MixInAudio(float* OutBuffer, int32 NumSamples, bool bUseLatestAudio)
+	{
+		if (IsInputStale())
+		{
+			return -1;
+		}
+
+		MixingBuffer.SetNumUninitialized(NumSamples, false);
+		int32 PopResult = 0;
+		
+		if (bUseLatestAudio && InternalBuffer.Num() > (uint32)NumSamples)
+		{
+			InternalBuffer.SetNum(((uint32)NumSamples));
+			PopResult = InternalBuffer.Peek(MixingBuffer.GetData(), NumSamples);
+		}
+		else
+		{
+			PopResult = InternalBuffer.Pop(MixingBuffer.GetData(), NumSamples);
+		}
+
+		if (FMath::IsNearlyEqual(TargetGain, PreviousGain))
+		{
+			MixInBufferFast(MixingBuffer.GetData(), OutBuffer, PopResult, PreviousGain);
+		}
+		else
+		{
+			MixInBufferFast(MixingBuffer.GetData(), OutBuffer, PopResult, PreviousGain, TargetGain);
+			PreviousGain = TargetGain;
+		}
+
+		return PopResult;
+	}
+
+
+	FPatchInput::FPatchInput(const FPatchOutputStrongPtr& InOutput)
+		: OutputHandle(InOutput)
+		, PushCallsCounter(0)
+	{
+		if (OutputHandle.IsValid())
+		{
+			OutputHandle->NumAliveInputs++;
+		}
+	}
+
+	FPatchInput::FPatchInput()
+		: PushCallsCounter(0)
+	{
+	}
+
+	FPatchInput::FPatchInput(const FPatchInput& Other)
+		: FPatchInput(Other.OutputHandle)
+	{
+	}
+
+	FPatchInput& FPatchInput::operator=(const FPatchInput& Other)
+	{
+		OutputHandle = Other.OutputHandle;
+		PushCallsCounter = 0;
+		
+		if (OutputHandle.IsValid())
+		{
+			OutputHandle->NumAliveInputs++;
+		}
+
+		return *this;
+	}
+
+	FPatchInput::~FPatchInput()
+	{
+		if (OutputHandle.IsValid())
+		{
+			OutputHandle->NumAliveInputs--;
+		}
+	}
+
+	int32 FPatchInput::PushAudio(const float* InBuffer, int32 NumSamples)
+	{
+		if (!OutputHandle.IsValid())
+		{
+			return -1;
+		}
+
+		int32 SamplesPushed = OutputHandle->InternalBuffer.Push(InBuffer, NumSamples);
+
+		// Every so often, we check to see if the output handle has been destroyed and clean it up.
+		static const int32 NumPushCallsUntilCleanupCheck = 256;
+		
+		PushCallsCounter = (PushCallsCounter + 1) % NumPushCallsUntilCleanupCheck;
+		if (PushCallsCounter == 0 && OutputHandle.IsUnique())
+		{
+			// Delete the output.
+			OutputHandle.Reset();
+		}
+
+		return SamplesPushed;
+	}
+
+	void FPatchInput::SetGain(float InGain)
+	{
+		if (!OutputHandle.IsValid())
+		{
+			return;
+		}
+
+		OutputHandle->TargetGain = InGain;
+	}
+
+	bool FPatchInput::IsOutputStillActive()
+	{
+		return OutputHandle.IsUnique() || OutputHandle.IsValid();
+	}
+
+	FPatchMixer::FPatchMixer()
+	{
+	}
+
+
+	FPatchInput FPatchMixer::AddNewInput(int32 MaxLatencyInSamples, float InGain)
+	{
+		FScopeLock ScopeLock(&PendingNewInputsCriticalSection);
+
+		int32 NewPatchIndex = PendingNewInputs.Emplace(new FPatchOutput(MaxLatencyInSamples, InGain));
+		return FPatchInput(PendingNewInputs[NewPatchIndex]);
+	}
+
+	void FPatchMixer::RemovePatch(const FPatchInput& PatchInput)
+	{
+		// If the output is already disconnected, early exit.
+		if (!PatchInput.OutputHandle.IsValid())
+		{
+			return;
+		}
+
+		FScopeLock ScopeLock(&InputDeletionCriticalSection);
+		DisconnectedInputs.Add(PatchInput.OutputHandle->PatchID);
+	}
+
+	int32 FPatchMixer::PopAudio(float* OutBuffer, int32 OutNumSamples, bool bUseLatestAudio)
+	{
+		FScopeLock ScopeLock(&CurrentPatchesCriticalSection);
+
+		CleanUpDisconnectedPatches();
+		ConnectNewPatches();
+
+		FMemory::Memzero(OutBuffer, OutNumSamples * sizeof(float));
+		int32 MaxPoppedSamples = 0;
+
+		for (FPatchOutputStrongPtr& OutputPtr : CurrentInputs)
+		{
+			const int32 NumPoppedSamples = OutputPtr->MixInAudio(OutBuffer, OutNumSamples, bUseLatestAudio);
+			MaxPoppedSamples = FMath::Max(NumPoppedSamples, MaxPoppedSamples);
+		}
+
+		return MaxPoppedSamples;
+	}
+
+	int32 FPatchMixer::Num()
+	{
+		FScopeLock ScopeLock(&CurrentPatchesCriticalSection);
+		return CurrentInputs.Num();
+	}
+
+	int32 FPatchMixer::MaxNumberOfSamplesThatCanBePopped()
+	{
+		FScopeLock ScopeLock(&CurrentPatchesCriticalSection);
+		ConnectNewPatches();
+
+		// Iterate through our inputs and see which input has the least audio buffered.
+		uint32 SmallestNumSamplesBuffered = TNumericLimits<uint32>::Max();
+
+		for (FPatchOutputStrongPtr& Output : CurrentInputs)
+		{
+			if (Output.IsValid())
+			{
+				SmallestNumSamplesBuffered = FMath::Min(SmallestNumSamplesBuffered, Output->InternalBuffer.Num());
+			}
+		}
+
+		if (SmallestNumSamplesBuffered == TNumericLimits<uint32>::Max())
+		{
+			return -1;
+		}
+		else
+		{
+			// If this check is hit, we need to either change this function to return an int64 or find a different way to notify the caller that all outputs have been disconeccted.
+			check(SmallestNumSamplesBuffered <= ((uint32)TNumericLimits<int32>::Max()));
+			return SmallestNumSamplesBuffered;
+		}
+	}
+
+	void FPatchMixer::ConnectNewPatches()
+	{
+		FScopeLock ScopeLock(&PendingNewInputsCriticalSection);
+
+		// If AddNewPatch is called in a separate thread, wait until the next PopAudio call to do this work.
+		// Todo: convert this to move semantics to avoid copying the shared pointer around.
+		for (FPatchOutputStrongPtr& Patch : PendingNewInputs)
+		{
+			CurrentInputs.Add(Patch);
+		}
+
+		PendingNewInputs.Reset();
+	}
+
+	void FPatchMixer::CleanUpDisconnectedPatches()
+	{
+		FScopeLock PendingInputDeletionScopeLock(&InputDeletionCriticalSection);
+
+		// Iterate through all of the PatchIDs we need to clean up.
+		for (const int32& PatchID : DisconnectedInputs)
+		{
+			bool bInputRemoved = false;
+
+			// First, make sure that the patch isn't in the pending new patchs we haven't added yet:
+			{
+				FScopeLock PendingNewInputsScopeLock(&PendingNewInputsCriticalSection);
+				for (int32 Index = 0; Index < PendingNewInputs.Num(); Index++)
+				{
+					checkSlow(CurrentInputs[Index].IsValid());
+
+					if (PatchID == PendingNewInputs[Index]->PatchID)
+					{
+						PendingNewInputs.RemoveAtSwap(Index);
+						bInputRemoved = true;
+						break;
+					}
+				}
+			}
+
+			if (bInputRemoved)
+			{
+				continue;
+			}
+
+			// Next, we check out current patchs.
+			for (int32 Index = 0; Index < CurrentInputs.Num(); Index++)
+			{
+				checkSlow(CurrentInputs[Index].IsValid());
+
+				if (PatchID == CurrentInputs[Index]->PatchID)
+				{
+					CurrentInputs.RemoveAtSwap(Index);
+					break;
+				}
+			}
+		}
+
+		DisconnectedInputs.Reset();
+	}
+
+	FPatchSplitter::FPatchSplitter()
+	{
+	}
+
+	FPatchSplitter::~FPatchSplitter()
+	{
+	}
+
+	Audio::FPatchOutputStrongPtr FPatchSplitter::AddNewPatch(int32 MaxLatencyInSamples, float InGain)
+	{
+		// Allocate a new FPatchOutput, then store a weak pointer to it in our PendingOutputs array to be added in our next call to PushAudio.
+		FPatchOutputStrongPtr StrongOutputPtr = MakeShareable(new FPatchOutput(MaxLatencyInSamples * 2, InGain));
+
+		{
+			FScopeLock ScopeLock(&PendingOutputsCriticalSection);
+			PendingOutputs.Emplace(StrongOutputPtr);
+		}
+
+		return StrongOutputPtr;
+	}
+
+	int32 FPatchSplitter::Num()
+	{
+		FScopeLock ScopeLock(&ConnectedOutputsCriticalSection);
+		AddPendingPatches();
+		return ConnectedOutputs.Num();
+	}
+
+	int32 FPatchSplitter::MaxNumberOfSamplesThatCanBePushed()
+	{
+		FScopeLock ScopeLock(&ConnectedOutputsCriticalSection);
+		AddPendingPatches();
+
+		// Iterate over our outputs and get the smallest remainder of all of our circular buffers.
+		uint32 SmallestRemainder = TNumericLimits<uint32>::Max();
+
+		for (FPatchInput& Input : ConnectedOutputs)
+		{
+			if (Input.OutputHandle.IsValid())
+			{
+				SmallestRemainder = FMath::Min(SmallestRemainder, Input.OutputHandle->InternalBuffer.Remainder());
+			}
+		}
+
+		if (SmallestRemainder == TNumericLimits<uint32>::Max())
+		{
+			return -1;
+		}
+		else
+		{
+			// If we hit this check, we need to either return an int64 or use some other method to notify the caller that all outputs are disconnected.
+			check(SmallestRemainder <= ((uint32)TNumericLimits<int32>::Max()));
+			return SmallestRemainder;
+		}
+	}
+
+	void FPatchSplitter::AddPendingPatches()
+	{
+		FScopeLock ScopeLock(&PendingOutputsCriticalSection);
+		ConnectedOutputs.Append(PendingOutputs);
+		PendingOutputs.Reset();
+	}
+
+	int32 FPatchSplitter::PushAudio(const float* InBuffer, int32 InNumSamples)
+	{
+		AddPendingPatches();
+
+		FScopeLock ScopeLock(&ConnectedOutputsCriticalSection);
+
+		int32 MinimumSamplesPushed = TNumericLimits<int32>::Max();
+
+		// Iterate through our array of connected outputs from the end, removing destroyed outputs as we go.
+		for (int32 Index = ConnectedOutputs.Num() - 1; Index >= 0; Index--)
+		{
+			int32 NumSamplesPushed = ConnectedOutputs[Index].PushAudio(InBuffer, InNumSamples);
+			if (NumSamplesPushed >= 0)
+			{
+				MinimumSamplesPushed = FMath::Min(MinimumSamplesPushed, NumSamplesPushed);
+			}
+			else
+			{
+				// If this output has been destroyed, remove it from our array of connected outputs.
+				ConnectedOutputs.RemoveAtSwap(Index);
+			}
+		}
+
+		// If we weren't able to push audio to any of our outputs, return -1.
+		if (MinimumSamplesPushed == TNumericLimits<int32>::Max())
+		{
+			MinimumSamplesPushed = -1;
+		}
+
+		return MinimumSamplesPushed;
+	}
+
+	FPatchMixerSplitter::FPatchMixerSplitter()
+	{
+	}
+
+	FPatchMixerSplitter::~FPatchMixerSplitter()
+	{
+	}
+
+	Audio::FPatchOutputStrongPtr FPatchMixerSplitter::AddNewOutput(int32 MaxLatencyInSamples, float InGain)
+	{
+		return Splitter.AddNewPatch(MaxLatencyInSamples, InGain);
+	}
+
+	Audio::FPatchInput FPatchMixerSplitter::AddNewInput(int32 MaxLatencyInSamples, float InGain)
+	{
+		return Mixer.AddNewInput(MaxLatencyInSamples, InGain);
+	}
+
+	void FPatchMixerSplitter::RemovePatch(const FPatchInput& TapInput)
+	{
+		Mixer.RemovePatch(TapInput);
+	}
+
+	void FPatchMixerSplitter::ProcessAudio()
+	{
+		int32 NumSamplesToForward = FMath::Min(Mixer.MaxNumberOfSamplesThatCanBePopped(), Splitter.MaxNumberOfSamplesThatCanBePushed());
+		
+		if (NumSamplesToForward <= 0)
+		{
+			// Likely there are either no inputs or no outputs connected, or one of the inputs has not pushed any audio yet. Early exit.
+			return;
+		}
+
+		IntermediateBuffer.Reset();
+		IntermediateBuffer.AddUninitialized(NumSamplesToForward);
+
+		// Mix down inputs:
+		int32 PopResult = Mixer.PopAudio(IntermediateBuffer.GetData(), NumSamplesToForward, false);
+		check(PopResult == NumSamplesToForward);
+		
+		OnProcessAudio(TArrayView<const float>(IntermediateBuffer));
+
+		// Push audio to outputs:
+		int32 PushResult = Splitter.PushAudio(IntermediateBuffer.GetData(), NumSamplesToForward);
+		check(PushResult == NumSamplesToForward);
+	}
+}
diff --git a/Engine/Source/Runtime/AudioMixer/Public/DSP/MultithreadedPatching.h b/Engine/Source/Runtime/AudioMixer/Public/DSP/MultithreadedPatching.h
new file mode 100644
--- /dev/null
+++ b/Engine/Source/Runtime/AudioMixer/Public/DSP/MultithreadedPatching.h
@@ -0,0 +1,223 @@
+// Copyright 1998-2019 Epic Games, Inc. All Rights Reserved.
+
+#pragma once
+
+#include "CoreMinimal.h"
+#include "DSP/BufferVectorOperations.h"
+#include "DSP/Dsp.h"
+
+namespace Audio
+{
+	/**
+	 * This class can be thought of as an output
+	 */
+	struct AUDIOMIXER_API FPatchOutput
+	{
+	public:
+		FPatchOutput(int32 MaxCapacity, float InGain = 1.0f);
+
+		/** The default constructor will result in an uninitialized, disconnected patch point, which  */
+		FPatchOutput();
+
+		/** Copies the minimum of NumSamples or however many samples are available into OutBuffer. Returns the number of samples copied, or -1 if this output's corresponding input has been destroyed. */
+		int32 PopAudio(float* OutBuffer, int32 NumSamples, bool bUseLatestAudio);
+
+		/** Sums the minimum of NumSamples or however many samples are available into OutBuffer. Returns the number of samples summed into OutBuffer. */
+		int32 MixInAudio(float* OutBuffer, int32 NumSamples, bool bUseLatestAudio);
+
+		/** Returns true if the input for this patch has been destroyed. */
+		bool IsInputStale() const;
+
+		friend class FPatchInput;
+		friend class FPatchMixer;
+		friend class FPatchSplitter;
+	private:
+
+		// Internal buffer.
+		TCircularAudioBuffer<float> InternalBuffer;
+
+		// For MixInAudio, audio is popped off of InternalBuffer onto here, and then mixed into OutBuffer in MixInAudio.
+		AlignedFloatBuffer MixingBuffer;
+
+		// This is applied in PopAudio or MixInAudio.
+		TAtomic<float> TargetGain;
+		float PreviousGain;
+
+		// This is used to breadcrumb the FPatchOutput when we want to delete it.
+		int32 PatchID;
+
+		// Counter that is incremented/decremented to allow FPatchInput to be copied around safely.
+		int32 NumAliveInputs;
+
+		static TAtomic<int32> PatchIDCounter;
+	};
+
+	/** Patch outputs are owned by the FPatchMixer, and are pinned by the FPatchInput. */
+	typedef TSharedPtr<FPatchOutput, ESPMode::ThreadSafe> FPatchOutputStrongPtr;
+	typedef TWeakPtr<FPatchOutput, ESPMode::ThreadSafe> FPatchOutputWeakPtr;
+
+	/**
+	 * Safe handle to a FPatchOutput. Should only be used by a single thread.
+	 */
+	class AUDIOMIXER_API FPatchInput
+	{
+	public:
+		/** PatchInputs can only be created from explicit outputs. */
+		FPatchInput(const FPatchOutputStrongPtr& InOutput);
+		FPatchInput(const FPatchInput& Other);
+		FPatchInput& operator=(const FPatchInput& Other);
+
+		/** Default constructed FPatchInput instances will always return -1 for PushAudio and true for IsOutputStillActive. */
+		FPatchInput();
+
+		~FPatchInput();
+
+		/** pushes audio from InBuffer to the corresponding FPatchOutput.
+		 *  Returns how many samples were able to be pushed, or -1 if the output was disconnected.
+		 */
+		int32 PushAudio(const float* InBuffer, int32 NumSamples);
+
+		void SetGain(float InGain);
+
+		/** Returns false if this output was removed, either because someone called FPatchMixer::RemoveTap with this FPatchInput, or the FPatchMixer was destroyed. */
+		bool IsOutputStillActive();
+
+		friend class FPatchMixer;
+		friend class FPatchSplitter;
+
+	private:
+		/** Weak pointer to our destination buffer. */
+		FPatchOutputStrongPtr OutputHandle;
+
+		/** Counter of the number of push calls until we check OutputHandle.IsUnique() and possibly clean up the output. */
+		int32 PushCallsCounter;
+	};
+
+	/**
+	 * This class is used for retrieving and mixing down audio from multiple threads.
+	 * Important to note that this is MPSC: while multiple threads can enqueue audio on an instance of FPatchMixer using instances of FPatchInput,
+	 * only one thread can call PopAudio safely.
+	 */
+	class AUDIOMIXER_API FPatchMixer
+	{
+	public:
+		/** Constructor. */
+		FPatchMixer();
+
+		/** Adds a new input to the tap collector. Calling this is thread safe, but individual instances of FPatchInput are only safe to be used from one thread. */
+		FORCEINLINE FPatchInput AddNewInput(int32 MaxLatencyInSamples, float InGain);
+
+		/** Removes a tap from the tap collector. Calling this is thread safe, though FPatchOutput will likely not be deleted until the next call of PopAudio. */
+		void RemovePatch(const FPatchInput& TapInput);
+
+		/** Mixes all inputs into a single buffer. This should only be called from a single thread. Returns the number of non-silent samples popped to OutBuffer. */
+		int32 PopAudio(float* OutBuffer, int32 OutNumSamples, bool bUseLatestAudio);
+
+		/** This returns the number of inputs currently connected to this patch mixer. Thread safe, but blocks for PopAudio. */
+		int32 Num();
+
+		/** This function call gets the maximum number of samples that's safe to pop, based on the thread with the least amount of samples buffered. Thread safe, but blocks for PopAudio. */
+		int32 MaxNumberOfSamplesThatCanBePopped();
+
+	private:
+		/** Called within PopAudio. Flushes the PendingNewPatches array into CurrentPatches. During this function, AddNewPatch is blocked. */
+		void ConnectNewPatches();
+
+		/** Called within PopAudio. Removes PendingTapsToDelete from CurrentPatches and ConnectNewPatches. During this function, RemoveTap and AddNewPatch are blocked. */
+		void CleanUpDisconnectedPatches();
+
+		/** New taps are added here in AddNewPatch, and then are moved to CurrentPatches in ConnectNewPatches. */
+		TArray<FPatchOutputStrongPtr> PendingNewInputs;
+		/** Contended by AddNewPatch, ConnectNewPatches and CleanUpDisconnectedTaps. */
+		FCriticalSection PendingNewInputsCriticalSection;
+
+		/** Patch IDs of individual audio taps that will be removed on the next call of CleanUpDisconnectedPatches. */
+		TArray<int32> DisconnectedInputs;
+		/** Contended by RemoveTap, AddNewPatch, and ConnectNewPatches. */
+		FCriticalSection InputDeletionCriticalSection;
+
+		/** Only accessed within PopAudio. Indirect array of taps that are mixed in during PopAudio. */
+		TArray<FPatchOutputStrongPtr> CurrentInputs;
+		FCriticalSection CurrentPatchesCriticalSection;
+	};
+
+	/**
+	 * This class is used to post audio from one source to multiple threads.
+	 * This class is SPMC: multiple threads can call FPatchOutputStrongPtr->PopAudio safely,
+	 * but only one thread can call PushAudio.
+	 */
+	class AUDIOMIXER_API FPatchSplitter
+	{
+	public:
+		/** Constructor. */
+		FPatchSplitter();
+
+		/** The destructor will mark every still connected FPatchOutput as stale. */
+		~FPatchSplitter();
+
+		/**
+		 * Adds a new output. Calling this is thread safe, but individual instances of FPatchOutput are only safe to be used from one thread.
+		 * the returned FPatchOutputPtr can be safely destroyed at any point.
+		 */
+		FORCEINLINE FPatchOutputStrongPtr AddNewPatch(int32 MaxLatencyInSamples, float InGain);
+
+		/** This call pushes audio to all outputs connected to this splitter. Only should be called from one thread. */
+		int32 PushAudio(const float* InBuffer, int32 InNumSamples);
+
+		/** This returns the number of outputs currently connected to this patch splitter. Thread safe, but blocks for PushAudio. */
+		int32 Num();
+
+		/** This function call gets the maximum number of samples that's safe to push. Thread safe, but blocks for PushAudio. */
+		int32 MaxNumberOfSamplesThatCanBePushed();
+
+	private:
+		// Called from PushAudio(). 
+		void AddPendingPatches();
+
+		TArray<FPatchInput> PendingOutputs;
+		FCriticalSection PendingOutputsCriticalSection;
+
+		TArray<FPatchInput> ConnectedOutputs;
+		FCriticalSection ConnectedOutputsCriticalSection;
+	};
+
+	/**
+	 * This class is used to mix multiple inputs from disparate threads to a single mixdown and deliver that mixdown to multiple outputs.
+	 * This class is MPCMC, but only one thread can and should call ProcessAudio().
+	 */
+	class AUDIOMIXER_API FPatchMixerSplitter
+	{
+	public:
+		/** Constructor. */
+		FPatchMixerSplitter();
+
+		/** The destructor will mark every still connected FPatchOutput as stale. */
+		virtual ~FPatchMixerSplitter();
+
+		/**
+		 * Adds a new output. Calling this is thread safe, but individual instances of FPatchOutput are only safe to be used from one thread.
+		 * the returned FPatchOutputPtr can be safely destroyed at any point.
+		 */
+		FORCEINLINE FPatchOutputStrongPtr AddNewOutput(int32 MaxLatencyInSamples, float InGain);
+
+		/** Adds a new input to the tap collector. Calling this is thread safe, but individual instances of FPatchInput are only safe to be used from one thread. */
+		FORCEINLINE FPatchInput AddNewInput(int32 MaxLatencyInSamples, float InGain);
+
+		/** Removes a tap from the tap collector. Calling this is thread safe, though FPatchOutput will likely not be deleted until the next call of PopAudio. */
+		void RemovePatch(const FPatchInput& TapInput);
+
+		/** Mixes audio from all inputs and pushes it to all outputs. Should be called regularly. */
+		void ProcessAudio();
+
+	protected:
+		/** This class can be subclassed with OnProcessAudio overridden. */
+		virtual void OnProcessAudio(TArrayView<const float> InAudio) {}
+
+	private:
+		FPatchMixer Mixer;
+		FPatchSplitter Splitter;
+
+		/** This buffer is used to pop audio from our Mixer and push to our splitter. */
+		AlignedFloatBuffer IntermediateBuffer;
+	};
+}
