Fixed bug where submix envelope follower would report zero's due to temp result being optimized away on certain platforms

diff --git a/Engine/Source/Runtime/AudioMixer/Private/AudioMixerSubmix.cpp b/Engine/Source/Runtime/AudioMixer/Private/AudioMixerSubmix.cpp
--- a/Engine/Source/Runtime/AudioMixer/Private/AudioMixerSubmix.cpp
+++ b/Engine/Source/Runtime/AudioMixer/Private/AudioMixerSubmix.cpp
@@ -782,10 +782,11 @@ namespace Audio
 			const int32 OutBufferSamples = OutAudioBuffer.Num();
 			const float* OutAudioBufferPtr = OutAudioBuffer.GetData();
 
-			float TempEnvelopeValues[AUDIO_MIXER_MAX_OUTPUT_CHANNELS];
-			FMemory::Memset(TempEnvelopeValues, sizeof(float)*AUDIO_MIXER_MAX_OUTPUT_CHANNELS);
 
 			// Perform envelope following per channel
+			FScopeLock EnvelopeScopeLock(&EnvelopeCriticalSection);
+			FMemory::Memset(EnvelopeValues, sizeof(float)*AUDIO_MIXER_MAX_OUTPUT_CHANNELS);
+
 			for (int32 ChannelIndex = 0; ChannelIndex < NumChannels; ++ChannelIndex)
 			{
 				// Get the envelope follower for the channel
@@ -798,14 +799,10 @@ namespace Audio
 					EnvFollower.ProcessAudio(SampleValue);
 				}
 
-				// Store the last value
-				TempEnvelopeValues[ChannelIndex] = EnvFollower.GetCurrentValue();
+				EnvelopeValues[ChannelIndex] = EnvFollower.GetCurrentValue();
 			}
 
-			FScopeLock EnvelopeScopeLock(&EnvelopeCriticalSection);
-
 			EnvelopeNumChannels = NumChannels;
-			FMemory::Memcpy(EnvelopeValues, TempEnvelopeValues, sizeof(float)*AUDIO_MIXER_MAX_OUTPUT_CHANNELS);
 		}
 
 		// Don't necessarily need to do this if the user isn't using this feature
