PR #6322: Added check if AudioComponent is allready beeing destroyed by owning . (Contributed by [REDACTED_PERSON])

diff --git a/Engine/Source/Runtime/AudioMixer/Private/Components/SynthComponent.cpp b/Engine/Source/Runtime/AudioMixer/Private/Components/SynthComponent.cpp
--- a/Engine/Source/Runtime/AudioMixer/Private/Components/SynthComponent.cpp
+++ b/Engine/Source/Runtime/AudioMixer/Private/Components/SynthComponent.cpp
@@ -314,7 +314,7 @@ void USynthComponent::OnUnregister()
 	}
 
 	// Make sure the audio component is destroyed during unregister
-	if (AudioComponent)
+	if (AudioComponent && !AudioComponent->IsBeingDestroyed())
 	{
 		if (Owner && Owner->GetWorld())
 		{
