Fix non-unity
#rnx

diff --git a/Engine/Source/Runtime/AudioMixer/Private/AudioMixerSourceBuffer.h b/Engine/Source/Runtime/AudioMixer/Private/AudioMixerSourceBuffer.h
--- a/Engine/Source/Runtime/AudioMixer/Private/AudioMixerSourceBuffer.h
+++ b/Engine/Source/Runtime/AudioMixer/Private/AudioMixerSourceBuffer.h
@@ -5,6 +5,7 @@
 #include "CoreMinimal.h"
 #include "AudioMixerBuffer.h"
 #include "AudioMixerSourceManager.h"
+#include "Sound/SoundWave.h"
 
 namespace Audio
 {
