Band-aid fix for high volume iOS crash.

diff --git a/Engine/Source/Runtime/AudioMixer/Private/AudioMixerSourceDecode.cpp b/Engine/Source/Runtime/AudioMixer/Private/AudioMixerSourceDecode.cpp
--- a/Engine/Source/Runtime/AudioMixer/Private/AudioMixerSourceDecode.cpp
+++ b/Engine/Source/Runtime/AudioMixer/Private/AudioMixerSourceDecode.cpp
@@ -56,7 +56,7 @@ public:
 			case EAudioTaskType::Procedural:
 			{
 				// Make sure we've been flagged as active
-				if (!ProceduralTaskData.ProceduralSoundWave->IsGeneratingAudio())
+				if (!ProceduralTaskData.ProceduralSoundWave || !ProceduralTaskData.ProceduralSoundWave->IsGeneratingAudio())
 				{
 					// Act as if we generated audio, but return silence.
 					FMemory::Memzero(ProceduralTaskData.AudioData, ProceduralTaskData.NumSamples * sizeof(float));
