fix for leaked FActiveSounds when a USoundWaveProcedural is played twice.

diff --git a/Engine/Source/Runtime/AudioMixer/Private/AudioMixerSource.cpp b/Engine/Source/Runtime/AudioMixer/Private/AudioMixerSource.cpp
--- a/Engine/Source/Runtime/AudioMixer/Private/AudioMixerSource.cpp
+++ b/Engine/Source/Runtime/AudioMixer/Private/AudioMixerSource.cpp
@@ -481,6 +481,14 @@ namespace Audio
 		if (!MixerSourceBuffer.IsValid())
 		{
 			FreeResources();
+
+			// Guarantee that this wave instance does not try to replay by disabling looping.
+			WaveInstance->LoopingMode = LOOP_Never;
+
+			if (ensure(WaveInstance->ActiveSound))
+			{
+				WaveInstance->ActiveSound->bShouldRemainActiveIfDropped = false;
+			}
 		}
 		
 		return MixerSourceBuffer.IsValid();
@@ -599,6 +607,13 @@ namespace Audio
 			return;
 		}
 
+		// Always stop procedural sounds immediately.
+		if (WaveInstance && WaveInstance->WaveData && WaveInstance->WaveData->bProcedural)
+		{
+			StopNow();
+			return;
+		}
+
 		if (bIsDone)
 		{
 			StopNow();
diff --git a/Engine/Source/Runtime/AudioMixer/Private/AudioMixerSourceBuffer.cpp b/Engine/Source/Runtime/AudioMixer/Private/AudioMixerSourceBuffer.cpp
--- a/Engine/Source/Runtime/AudioMixer/Private/AudioMixerSourceBuffer.cpp
+++ b/Engine/Source/Runtime/AudioMixer/Private/AudioMixerSourceBuffer.cpp
@@ -65,6 +65,9 @@ namespace Audio
 			UE_LOG(LogAudioMixer, Warning,
 				TEXT("Procedural USoundWave is reinitializing even though it is actively "
 				"generating audio. Sound must be stopped before playing again."));
+
+			// Need to set the procedural sound wave as not looping to allow it to get stopped during sound wave parsing
+			InWave.bLooping = false;
 			return nullptr;
 		}
 
diff --git a/Engine/Source/Runtime/AudioMixer/Private/Components/SynthComponent.cpp b/Engine/Source/Runtime/AudioMixer/Private/Components/SynthComponent.cpp
--- a/Engine/Source/Runtime/AudioMixer/Private/Components/SynthComponent.cpp
+++ b/Engine/Source/Runtime/AudioMixer/Private/Components/SynthComponent.cpp
@@ -484,11 +484,6 @@ void USynthComponent::Start()
 		AudioComponent->SetSound(Synth);
 		AudioComponent->Play(0);
 
-		// Copy sound base data to the sound
-		Synth->SourceEffectChain = SourceEffectChain;
-		Synth->SoundSubmixObject = SoundSubmix;
-		Synth->SoundSubmixSends = SoundSubmixSends;
-
 		SetActiveFlag(AudioComponent->IsActive());
 
 		if (IsActive())
@@ -507,6 +502,12 @@ void USynthComponent::Stop()
 		if (AudioComponent)
 		{
 			AudioComponent->Stop();
+
+			FAudioDevice* AudioDevice = AudioComponent->GetAudioDevice();
+			if (AudioDevice)
+			{
+				AudioDevice->StopSoundsUsingResource(Synth);
+			}
 		}
 
 		SetActiveFlag(false);
