Assetization of Master Submix Graph Prep 1:
- Clean-up channel matrix initialization and move out of InitSoundSubmixes

diff --git a/Engine/Source/Runtime/AudioMixer/Private/AudioMixerChannelMaps.cpp b/Engine/Source/Runtime/AudioMixer/Private/AudioMixerChannelMaps.cpp
--- a/Engine/Source/Runtime/AudioMixer/Private/AudioMixerChannelMaps.cpp
+++ b/Engine/Source/Runtime/AudioMixer/Private/AudioMixerChannelMaps.cpp
@@ -1,11 +1,72 @@
 // Copyright 1998-2019 Epic Games, Inc. All Rights Reserved.
 
+#include "Audio.h"
 #include "AudioMixer.h"
 #include "AudioMixerDevice.h"
 #include "Misc/ConfigCacheIni.h"
 
 namespace Audio
 {
+	/** Channel type matrix for submix speaker channel mappings. */
+	const TArray<TArray<EAudioMixerChannel::Type>> SubmixOutputChannelMatrix
+	{
+		// ESubmixChannelFormat::Device
+		// Placeholder: Should never be used as Device signifies dynamically set
+		{
+		},
+
+		// ESubmixChannelFormat::Stereo
+		{
+			EAudioMixerChannel::FrontLeft,
+			EAudioMixerChannel::FrontRight
+		},
+
+		// ESubmixChannelFormat::Quad
+		{
+			EAudioMixerChannel::FrontLeft,
+			EAudioMixerChannel::FrontRight,
+			EAudioMixerChannel::SideLeft,
+			EAudioMixerChannel::SideRight
+		},
+
+		// ESubmixChannelFormat::FiveDotOne
+		{
+			EAudioMixerChannel::FrontLeft,
+			EAudioMixerChannel::FrontRight,
+			EAudioMixerChannel::FrontCenter,
+			EAudioMixerChannel::LowFrequency,
+			EAudioMixerChannel::SideLeft,
+			EAudioMixerChannel::SideRight
+		},
+
+		// ESubmixChannelFormat::SevenDotOne
+		{
+			EAudioMixerChannel::FrontLeft,
+			EAudioMixerChannel::FrontRight,
+			EAudioMixerChannel::FrontCenter,
+			EAudioMixerChannel::LowFrequency,
+			EAudioMixerChannel::BackLeft,
+			EAudioMixerChannel::BackRight,
+			EAudioMixerChannel::SideLeft,
+			EAudioMixerChannel::SideRight
+		},
+
+		// ESubmixChannelFormat::Ambisonics
+		// Ambisonics output is encoded to max encoded channel (i.e. 7.1).
+		// To support ambisonic encoded output, will need to convert to
+		// Ambisonics_W/X/Y/Z alias values.
+		{
+			EAudioMixerChannel::FrontLeft,
+			EAudioMixerChannel::FrontRight,
+			EAudioMixerChannel::FrontCenter,
+			EAudioMixerChannel::LowFrequency,
+			EAudioMixerChannel::BackLeft,
+			EAudioMixerChannel::BackRight,
+			EAudioMixerChannel::SideLeft,
+			EAudioMixerChannel::SideRight
+		}, 
+	};
+
 	// Tables based on Ac-3 down-mixing
 	// Rows: output speaker configuration
 	// Cols: input source channels
@@ -542,10 +603,10 @@ namespace Audio
 
 		// ambisonics is special cased and uses a plugin.
 		TArray<FChannelPositionInfo> FirstOrderAmbisonicsPositions;
-		FirstOrderAmbisonicsPositions.Add(DefaultChannelAzimuthPosition[EAudioMixerChannel::FrontLeft]);
-		FirstOrderAmbisonicsPositions.Add(DefaultChannelAzimuthPosition[EAudioMixerChannel::FrontRight]);
-		FirstOrderAmbisonicsPositions.Add(DefaultChannelAzimuthPosition[EAudioMixerChannel::SideLeft]);
-		FirstOrderAmbisonicsPositions.Add(DefaultChannelAzimuthPosition[EAudioMixerChannel::SideRight]);
+		FirstOrderAmbisonicsPositions.Add(DefaultChannelAzimuthPosition[EAudioMixerChannel::Ambisonics_W]);
+		FirstOrderAmbisonicsPositions.Add(DefaultChannelAzimuthPosition[EAudioMixerChannel::Ambisonics_X]);
+		FirstOrderAmbisonicsPositions.Add(DefaultChannelAzimuthPosition[EAudioMixerChannel::Ambisonics_Y]);
+		FirstOrderAmbisonicsPositions.Add(DefaultChannelAzimuthPosition[EAudioMixerChannel::Ambisonics_Z]);
 		FirstOrderAmbisonicsPositions.Sort(FCompareByAzimuth());
 		ChannelAzimuthPositions.Add(ESubmixChannelFormat::Ambisonics, FirstOrderAmbisonicsPositions);
 		OutputChannels[int32(ESubmixChannelFormat::Ambisonics)] = FirstOrderAmbisonicsPositions.Num();
@@ -574,14 +635,16 @@ namespace Audio
 		return ESubmixChannelFormat::Device;
 	}
 
-	const TArray<EAudioMixerChannel::Type>& FMixerDevice::GetChannelArrayForSubmixChannelType(const ESubmixChannelFormat InSubmixChannelType) const
+	const TArray<EAudioMixerChannel::Type>& FMixerDevice::GetChannelArrayForSubmixChannelFormat(const ESubmixChannelFormat InSubmixChannelFormat) const
 	{
-		if (InSubmixChannelType == ESubmixChannelFormat::Device)
+		if (InSubmixChannelFormat == ESubmixChannelFormat::Device)
 		{
 			return PlatformInfo.OutputChannelArray;
 		}
-		const TArray<EAudioMixerChannel::Type>* ChannelArray = ChannelArrays.Find(InSubmixChannelType);
-		return *ChannelArray;
-	}
 
+		static_assert(static_cast<int32>(ESubmixChannelFormat::Count) == 6, "ESubmixChannelFormat count must match SubmixOutputChannelMatrix size");
+
+		const int32 ChannelFormatIndex = static_cast<int32>(InSubmixChannelFormat);
+		return SubmixOutputChannelMatrix[ChannelFormatIndex];
+	}
 }
diff --git a/Engine/Source/Runtime/AudioMixer/Private/AudioMixerDevice.cpp b/Engine/Source/Runtime/AudioMixer/Private/AudioMixerDevice.cpp
--- a/Engine/Source/Runtime/AudioMixer/Private/AudioMixerDevice.cpp
+++ b/Engine/Source/Runtime/AudioMixer/Private/AudioMixerDevice.cpp
@@ -286,7 +286,6 @@ namespace Audio
 	{
 		AUDIO_MIXER_CHECK_GAME_THREAD(this);
 
-		// Make sure all submixes are registered but not initialized
 		for (TObjectIterator<USoundSubmix> It; It; ++It)
 		{
 			UnregisterSoundSubmix(*It);
@@ -741,7 +740,7 @@ namespace Audio
 		// Make sure all submixes are registered but not initialized
 		for (TObjectIterator<USoundSubmix> It; It; ++It)
 		{
-			RegisterSoundSubmix(*It, false);
+			RegisterSoundSubmix(*It, false /* bInit */);
 		}
 
 		// Now setup the graph for all the submixes
@@ -790,49 +789,6 @@ namespace Audio
 			// Perform any other initialization on the submix instance
 			SubmixInstance->Init(SoundSubmix);
 		}
-
-		TArray<EAudioMixerChannel::Type> StereoChannelTypes;
-		StereoChannelTypes.Add(EAudioMixerChannel::FrontLeft);
-		StereoChannelTypes.Add(EAudioMixerChannel::FrontRight);
-		ChannelArrays.Add(ESubmixChannelFormat::Stereo, StereoChannelTypes);
-
-		TArray<EAudioMixerChannel::Type> QuadChannelTypes;
-		QuadChannelTypes.Add(EAudioMixerChannel::FrontLeft);
-		QuadChannelTypes.Add(EAudioMixerChannel::FrontRight);
-		QuadChannelTypes.Add(EAudioMixerChannel::SideLeft);
-		QuadChannelTypes.Add(EAudioMixerChannel::SideRight);
-		ChannelArrays.Add(ESubmixChannelFormat::Quad, QuadChannelTypes);
-
-		TArray<EAudioMixerChannel::Type> FiveDotOneTypes;
-		FiveDotOneTypes.Add(EAudioMixerChannel::FrontLeft);
-		FiveDotOneTypes.Add(EAudioMixerChannel::FrontRight);
-		FiveDotOneTypes.Add(EAudioMixerChannel::FrontCenter);
-		FiveDotOneTypes.Add(EAudioMixerChannel::LowFrequency);
-		FiveDotOneTypes.Add(EAudioMixerChannel::SideLeft);
-		FiveDotOneTypes.Add(EAudioMixerChannel::SideRight);
-		ChannelArrays.Add(ESubmixChannelFormat::FiveDotOne, FiveDotOneTypes);
-
-		TArray<EAudioMixerChannel::Type> SevenDotOneTypes;
-		SevenDotOneTypes.Add(EAudioMixerChannel::FrontLeft);
-		SevenDotOneTypes.Add(EAudioMixerChannel::FrontRight);
-		SevenDotOneTypes.Add(EAudioMixerChannel::FrontCenter);
-		SevenDotOneTypes.Add(EAudioMixerChannel::LowFrequency);
-		SevenDotOneTypes.Add(EAudioMixerChannel::BackLeft);
-		SevenDotOneTypes.Add(EAudioMixerChannel::BackRight);
-		SevenDotOneTypes.Add(EAudioMixerChannel::SideLeft);
-		SevenDotOneTypes.Add(EAudioMixerChannel::SideRight);
-		ChannelArrays.Add(ESubmixChannelFormat::SevenDotOne, SevenDotOneTypes);
-
-		TArray<EAudioMixerChannel::Type> FirstOrderAmbisonicsTypes;
-		FirstOrderAmbisonicsTypes.Add(EAudioMixerChannel::FrontLeft);
-		FirstOrderAmbisonicsTypes.Add(EAudioMixerChannel::FrontRight);
-		FirstOrderAmbisonicsTypes.Add(EAudioMixerChannel::FrontCenter);
-		FirstOrderAmbisonicsTypes.Add(EAudioMixerChannel::LowFrequency);
-		FirstOrderAmbisonicsTypes.Add(EAudioMixerChannel::BackLeft);
-		FirstOrderAmbisonicsTypes.Add(EAudioMixerChannel::BackRight);
-		FirstOrderAmbisonicsTypes.Add(EAudioMixerChannel::SideLeft);
-		FirstOrderAmbisonicsTypes.Add(EAudioMixerChannel::SideRight);
-		ChannelArrays.Add(ESubmixChannelFormat::Ambisonics, FirstOrderAmbisonicsTypes);
 	}
 	
  	FAudioPlatformSettings FMixerDevice::GetPlatformSettings() const
@@ -926,13 +882,13 @@ namespace Audio
 #if WITH_EDITOR
 		check(IsInAudioThread());
 
-		FMixerSubmixPtr* MixerSubmix = Submixes.Find(InSoundSubmix);
-		if (MixerSubmix)
+		const FMixerSubmixPtr MixerSubmix = Submixes.FindRef(InSoundSubmix);
+		if (MixerSubmix.IsValid())
 		{
-			float NewVolume = InSoundSubmix->OutputVolume;
+			const float NewVolume = InSoundSubmix->OutputVolume;
 			AudioRenderThreadCommand([MixerSubmix, NewVolume]()
 			{
-				(*MixerSubmix)->SetOutputVolume(NewVolume);
+				MixerSubmix->SetOutputVolume(NewVolume);
 			});
 		}
 #endif // WITH_EDITOR
@@ -953,12 +909,12 @@ namespace Audio
 			return;
 		}
 
-		FMixerSubmixPtr* MixerSubmix = Submixes.Find(InSoundSubmix);
-		if (MixerSubmix)
+		const FMixerSubmixPtr MixerSubmix = Submixes.FindRef(InSoundSubmix);
+		if (MixerSubmix.IsValid())
 		{
 			AudioRenderThreadCommand([MixerSubmix, NewVolume]()
 			{
-				(*MixerSubmix)->SetDynamicOutputVolume(NewVolume);
+				MixerSubmix->SetDynamicOutputVolume(NewVolume);
 			});
 		}
 	}
@@ -1187,7 +1143,7 @@ namespace Audio
 		if (InWaveInstance->bCenterChannelOnly)
 		{
 			int32 NumOutputChannels = GetNumChannelsForSubmixFormat(InSubmixType);
-			const TArray<EAudioMixerChannel::Type>& ChannelArray = GetChannelArrayForSubmixChannelType(InSubmixType);
+			const TArray<EAudioMixerChannel::Type>& ChannelArray = GetChannelArrayForSubmixChannelFormat(InSubmixType);
 
 			// If we are only spatializing to stereo output
 			if (NumOutputChannels == 2)
@@ -1298,7 +1254,7 @@ namespace Audio
 		float OmniPanFactor = 1.0f / NumSpatialChannels;
 
 		float DefaultEffectivePan = !OmniAmount ? 0.0f : FMath::Lerp(0.0f, OmniPanFactor, OmniAmount);
-		const TArray<EAudioMixerChannel::Type>& ChannelArray = GetChannelArrayForSubmixChannelType(InSubmixType);
+		const TArray<EAudioMixerChannel::Type>& ChannelArray = GetChannelArrayForSubmixChannelFormat(InSubmixType);
 
 		for (EAudioMixerChannel::Type Channel : ChannelArray)
 		{
diff --git a/Engine/Source/Runtime/AudioMixer/Public/AudioMixerDevice.h b/Engine/Source/Runtime/AudioMixer/Public/AudioMixerDevice.h
--- a/Engine/Source/Runtime/AudioMixer/Public/AudioMixerDevice.h
+++ b/Engine/Source/Runtime/AudioMixer/Public/AudioMixerDevice.h
@@ -2,6 +2,7 @@
 
 #pragma once
 
+#include "Audio.h"
 #include "AudioMixer.h"
 #include "AudioMixerSourceManager.h"
 #include "AudioDevice.h"
@@ -204,7 +205,7 @@ namespace Audio
 		uint32 GetNewUniqueAmbisonicsStreamID();
 
 		// Returns the channel array for the given submix channel type
-		const TArray<EAudioMixerChannel::Type>& GetChannelArrayForSubmixChannelType(const ESubmixChannelFormat InSubmixChannelType) const;
+		const TArray<EAudioMixerChannel::Type>& GetChannelArrayForSubmixChannelFormat(const ESubmixChannelFormat InSubmixChannelFormat) const;
 
 		// Retrieves the listener transforms
 		const TArray<FTransform>* GetListenerTransforms();
@@ -260,9 +261,6 @@ namespace Audio
 
 		int32 OutputChannels[(int32)ESubmixChannelFormat::Count];
 
-		/** Channel type arrays for submix channel types. */
-		TMap<ESubmixChannelFormat, TArray<EAudioMixerChannel::Type>> ChannelArrays;
-
 		/** What upmix method to use for mono channel upmixing. */
 		EMonoChannelUpmixMethod MonoChannelUpmixMethod;
 
diff --git a/Engine/Source/Runtime/Engine/Classes/Sound/SoundSubmix.h b/Engine/Source/Runtime/Engine/Classes/Sound/SoundSubmix.h
--- a/Engine/Source/Runtime/Engine/Classes/Sound/SoundSubmix.h
+++ b/Engine/Source/Runtime/Engine/Classes/Sound/SoundSubmix.h
@@ -44,7 +44,7 @@ enum class ESubmixChannelFormat : uint8
 	// Sets the submix to mix audio to 7.1 (FL, FR, FC, LF, BL, BR, SL, SR)
 	SevenDotOne UMETA(DisplayName = "7.1"),
 
-	// Sets the submix to render audio as an ambisonics bed.
+	// Sets the submix to render audio to first-order ambisonics.
 	Ambisonics UMETA(DisplayName = "Ambisonics"),
 
 	Count UMETA(Hidden)
