Updating logic to recalculate source channel maps based on whether source HRTF is applied on source init.

diff --git a/Engine/Source/Runtime/AudioMixer/Private/AudioMixerSource.cpp b/Engine/Source/Runtime/AudioMixer/Private/AudioMixerSource.cpp
--- a/Engine/Source/Runtime/AudioMixer/Private/AudioMixerSource.cpp
+++ b/Engine/Source/Runtime/AudioMixer/Private/AudioMixerSource.cpp
@@ -1091,7 +1091,7 @@ namespace Audio
 
 	bool FMixerSource::ComputeMonoChannelMap(const ESubmixChannelFormat SubmixChannelType, Audio::AlignedFloatBuffer& OutChannelMap)
 	{
-		if (UseObjectBasedSpatialization())
+		if (IsUsingObjectBasedSpatialization())
 		{
 			if (WaveInstance->SpatializationMethod != ESoundSpatializationAlgorithm::SPATIALIZATION_HRTF && !bEditorWarnedChangedSpatialization)
 			{
@@ -1132,7 +1132,8 @@ namespace Audio
 				UpdateStereoEmitterPositions();
 			}
 
-			if (!UseObjectBasedSpatialization())
+			// Check whether voice is currently using 
+			if (!IsUsingObjectBasedSpatialization())
 			{
 				float AzimuthOffset = 0.0f;
 
@@ -1203,6 +1204,22 @@ namespace Audio
 				WaveInstance->SpatializationMethod == ESoundSpatializationAlgorithm::SPATIALIZATION_HRTF);
 	}
 
+	bool FMixerSource::IsUsingObjectBasedSpatialization() const
+	{
+		bool bIsUsingObjectBaseSpatialization = UseObjectBasedSpatialization();
+
+		if (MixerSourceVoice)
+		{
+			// If it is currently playing, check whether it actively uses HRTF spatializer.
+			// HRTF spatialization cannot be altered on currently playing source. So this handles
+			// the case where the source was initialized without HRTF spatialization before HRTF
+			// spatialization is enabled. 
+			bool bDefaultIfNoSourceId = true;
+			bIsUsingObjectBaseSpatialization &= MixerSourceVoice->IsUsingHRTFSpatializer(bDefaultIfNoSourceId);
+		}
+		return bIsUsingObjectBaseSpatialization;
+	}
+
 	bool FMixerSource::UseSpatializationPlugin() const
 	{
 		return (Buffer->NumChannels <= MixerDevice->MaxChannelsSupportedBySpatializationPlugin) &&
diff --git a/Engine/Source/Runtime/AudioMixer/Private/AudioMixerSource.h b/Engine/Source/Runtime/AudioMixer/Private/AudioMixerSource.h
--- a/Engine/Source/Runtime/AudioMixer/Private/AudioMixerSource.h
+++ b/Engine/Source/Runtime/AudioMixer/Private/AudioMixerSource.h
@@ -96,6 +96,9 @@ namespace Audio
 
 		/** Whether or not we should create the source voice with the HRTF spatializer. */
 		bool UseObjectBasedSpatialization() const;
+		
+		/** Whether or not existing or new sources will use the HRTF spatializer. */
+		bool IsUsingObjectBasedSpatialization() const;
 
 		/** Whether or not to use the spatialization plugin. */
 		bool UseSpatializationPlugin() const;
diff --git a/Engine/Source/Runtime/AudioMixer/Private/AudioMixerSourceManager.cpp b/Engine/Source/Runtime/AudioMixer/Private/AudioMixerSourceManager.cpp
--- a/Engine/Source/Runtime/AudioMixer/Private/AudioMixerSourceManager.cpp
+++ b/Engine/Source/Runtime/AudioMixer/Private/AudioMixerSourceManager.cpp
@@ -256,6 +256,7 @@ namespace Audio
 		GameThreadInfo.bIsBusy.AddDefaulted(NumTotalSources);
 		GameThreadInfo.bNeedsSpeakerMap.AddDefaulted(NumTotalSources);
 		GameThreadInfo.bIsDebugMode.AddDefaulted(NumTotalSources);
+		GameThreadInfo.bIsUsingHRTFSpatializer.AddDefaulted(NumTotalSources);
 		GameThreadInfo.FreeSourceIndices.Reset(NumTotalSources);
 		for (int32 i = NumTotalSources - 1; i >= 0; --i)
 		{
@@ -629,6 +630,8 @@ namespace Audio
 		// Make sure we flag that this source needs a speaker map to at least get one
 		GameThreadInfo.bNeedsSpeakerMap[SourceId] = true;
 
+		GameThreadInfo.bIsUsingHRTFSpatializer[SourceId] = InitParams.bUseHRTFSpatialization;
+
 		// Create the modulation plugin source effect
 		if (InitParams.ModulationPluginSettings != nullptr)
 		{
@@ -1275,6 +1278,12 @@ namespace Audio
 		return SourceInfos[SourceId].SourceEnvelopeValue;
 	}
 
+	bool FMixerSourceManager::IsUsingHRTFSpatializer(const int32 SourceId) const
+	{
+		AUDIO_MIXER_CHECK_GAME_THREAD(MixerDevice);
+		return GameThreadInfo.bIsUsingHRTFSpatializer[SourceId];
+	}
+
 	bool FMixerSourceManager::NeedsSpeakerMap(const int32 SourceId) const
 	{
 		AUDIO_MIXER_CHECK_GAME_THREAD(MixerDevice);
diff --git a/Engine/Source/Runtime/AudioMixer/Private/AudioMixerSourceManager.h b/Engine/Source/Runtime/AudioMixer/Private/AudioMixerSourceManager.h
--- a/Engine/Source/Runtime/AudioMixer/Private/AudioMixerSourceManager.h
+++ b/Engine/Source/Runtime/AudioMixer/Private/AudioMixerSourceManager.h
@@ -236,6 +236,7 @@ namespace Audio
 
 		int64 GetNumFramesPlayed(const int32 SourceId) const;
 		float GetEnvelopeValue(const int32 SourceId) const;
+		bool IsUsingHRTFSpatializer(const int32 SourceId) const;
 		bool NeedsSpeakerMap(const int32 SourceId) const;
 		void ComputeNextBlockOfSamples();
 		void ClearStoppingSounds();
@@ -559,6 +560,7 @@ namespace Audio
 			TArray<bool> bIsBusy;
 			TArray<bool> bNeedsSpeakerMap;
 			TArray<bool> bIsDebugMode;
+			TArray<bool> bIsUsingHRTFSpatializer;
 		} GameThreadInfo;
 
 		int32 NumActiveSources;
diff --git a/Engine/Source/Runtime/AudioMixer/Private/AudioMixerSourceVoice.cpp b/Engine/Source/Runtime/AudioMixer/Private/AudioMixerSourceVoice.cpp
--- a/Engine/Source/Runtime/AudioMixer/Private/AudioMixerSourceVoice.cpp
+++ b/Engine/Source/Runtime/AudioMixer/Private/AudioMixerSourceVoice.cpp
@@ -223,6 +223,18 @@ namespace Audio
 		return SourceManager->NeedsSpeakerMap(SourceId);
 	}
 
+	bool FMixerSourceVoice::IsUsingHRTFSpatializer(bool bDefaultValue) const
+	{
+		AUDIO_MIXER_CHECK_GAME_THREAD(MixerDevice);
+
+		if (SourceId != INDEX_NONE)
+		{
+			return SourceManager->IsUsingHRTFSpatializer(SourceId);
+		}
+
+		return bDefaultValue;
+	}
+
 	int64 FMixerSourceVoice::GetNumFramesPlayed() const
 	{
 		AUDIO_MIXER_CHECK_GAME_THREAD(MixerDevice);
diff --git a/Engine/Source/Runtime/AudioMixer/Private/AudioMixerSourceVoice.h b/Engine/Source/Runtime/AudioMixer/Private/AudioMixerSourceVoice.h
--- a/Engine/Source/Runtime/AudioMixer/Private/AudioMixerSourceVoice.h
+++ b/Engine/Source/Runtime/AudioMixer/Private/AudioMixerSourceVoice.h
@@ -83,6 +83,11 @@ namespace Audio
 		// Whether or not the device changed and needs another speaker map sent
 		bool NeedsSpeakerMap() const;
 
+		// Whether or not the voice is currently using HRTF spatialization.
+		//
+		// @param bDefaultValue - This value will be returned if voice does not have a valid source id.
+		bool IsUsingHRTFSpatializer(bool bDefaultValue) const;
+
 		// Retrieves the total number of samples played.
 		int64 GetNumFramesPlayed() const;
 
