fix for crash on pc shutdown

diff --git a/Engine/Source/Runtime/AudioMixer/Private/AudioMixerSourceManager.cpp b/Engine/Source/Runtime/AudioMixer/Private/AudioMixerSourceManager.cpp
--- a/Engine/Source/Runtime/AudioMixer/Private/AudioMixerSourceManager.cpp
+++ b/Engine/Source/Runtime/AudioMixer/Private/AudioMixerSourceManager.cpp
@@ -149,16 +149,6 @@ namespace Audio
 		}
 
 		FPlatformProcess::ReturnSynchEventToPool(CommandsProcessedEvent);
-
-		// We are tearing the whole system down, so clear the MixerSourceBuffers Wave pointer
-		// prior to them destructing, which prevents them doing the normal OnGenerateEnd flow.
-		for (FSourceInfo& i : SourceInfos)
-		{
-			if (FMixerSourceBuffer* Buffer = i.MixerSourceBuffer.Get())
-			{				
-				Buffer->ClearWave(); // NOTE; this does not decrement the active source count on the wave and assumes the pointer is unsafe.
-			}
-		}
 	}
 
 	void FMixerSourceManager::Init(const FSourceManagerInitParams& InitParams)
