Fixed race condiotion where calling stop() then start() on a procedural sound updated state across threads in a way where the start() would be acknowlegded before the stop(), and the start() would then be ignored.
Causing issues with VoidListenerSynthComponent where forcing resync would call stop() then start() in fairly quick succession.
Moved state to be on the same thread.

diff --git a/Engine/Source/Runtime/AudioMixer/Private/AudioMixerSource.cpp b/Engine/Source/Runtime/AudioMixer/Private/AudioMixerSource.cpp
--- a/Engine/Source/Runtime/AudioMixer/Private/AudioMixerSource.cpp
+++ b/Engine/Source/Runtime/AudioMixer/Private/AudioMixerSource.cpp
@@ -582,6 +582,11 @@ namespace Audio
 		// we'll just not actually play the source here. Instead we'll call play when the sound finishes loading.
 		if (MixerSourceVoice && InitializationState == EMixerSourceInitializationState::Initialized)
 		{
+			if (WaveInstance && WaveInstance->WaveData && WaveInstance->WaveData->bProcedural)
+			{
+				WaveInstance->WaveData->bPlayingProcedural = true;
+			}
+
 			MixerSourceVoice->Play();
 		}
 
@@ -610,6 +615,7 @@ namespace Audio
 		// Always stop procedural sounds immediately.
 		if (WaveInstance && WaveInstance->WaveData && WaveInstance->WaveData->bProcedural)
 		{
+			WaveInstance->WaveData->bPlayingProcedural = false;
 			StopNow();
 			return;
 		}
diff --git a/Engine/Source/Runtime/AudioMixer/Private/Components/SynthComponent.cpp b/Engine/Source/Runtime/AudioMixer/Private/Components/SynthComponent.cpp
--- a/Engine/Source/Runtime/AudioMixer/Private/Components/SynthComponent.cpp
+++ b/Engine/Source/Runtime/AudioMixer/Private/Components/SynthComponent.cpp
@@ -502,12 +502,6 @@ void USynthComponent::Stop()
 		if (AudioComponent)
 		{
 			AudioComponent->Stop();
-
-			FAudioDevice* AudioDevice = AudioComponent->GetAudioDevice();
-			if (AudioDevice)
-			{
-				AudioDevice->StopSoundsUsingResource(Synth);
-			}
 		}
 
 		SetActiveFlag(false);
diff --git a/Engine/Source/Runtime/Engine/Classes/Sound/SoundWave.h b/Engine/Source/Runtime/Engine/Classes/Sound/SoundWave.h
--- a/Engine/Source/Runtime/Engine/Classes/Sound/SoundWave.h
+++ b/Engine/Source/Runtime/Engine/Classes/Sound/SoundWave.h
@@ -334,6 +334,9 @@ public:
 
 	/** Set to true for programmatically generated audio. */
 	uint8 bProcedural:1;
+	
+	/** Set to true if the source is procedural and currently playing */
+	uint8 bPlayingProcedural : 1;
 
 	/** Set to true of this is a bus sound source. This will result in the sound wave not generating audio for itself, but generate audio through instances. Used only in audio mixer. */
 	uint8 bIsBus:1;
diff --git a/Engine/Source/Runtime/Engine/Private/AudioDevice.cpp b/Engine/Source/Runtime/Engine/Private/AudioDevice.cpp
--- a/Engine/Source/Runtime/Engine/Private/AudioDevice.cpp
+++ b/Engine/Source/Runtime/Engine/Private/AudioDevice.cpp
@@ -4502,7 +4502,7 @@ void FAudioDevice::AddNewActiveSoundInternal(const FActiveSound& NewActiveSound,
 #endif // !UE_BUILD_SHIPPING
 
 	const USoundWave* SoundWave = Cast<USoundWave>(Sound);
-	if (SoundWave && SoundWave->bProcedural && SoundWave->IsGeneratingAudio())
+	if (SoundWave && SoundWave->bProcedural && SoundWave->bPlayingProcedural)
 	{
 		FString SoundWaveName;
 		SoundWave->GetName(SoundWaveName);
