- Asset-ize all master submix effects
set in AudioSettings)
- Fix memory leak potential by using shared ptr for FSoundEffectSubmix and only create on first request from reverb plugin GetEffectSubmix implementations.
- Clean up reverb fast UX a bit
- Add bypass feature on new reverb to allow for disabling gracefully

- Upgrade Notes:
    - IAudioReverb child classes requires signature change
        - GetEffectSubmix has been updated in favor of reverb plugins internally managing USubmix asset and respective lifetime to remove unneeded engine dependencies on submix existence when plugin is not active.
        - GetSubmix has been added so plugin is in charge of maintaining and providing global submix setting internally
    - Clean-up and remove legacy reverb swap ability on init for 4.25. If users still want to use the legacy reverb, they should override the master reverb submix in AudioSettings with a game-specific submix that uses a legacy reverb effect chain, and set their override Submix version accordingly in the global Audio Settings
    - bEQFilter Applied has been removed from ActiveSound in favor of setting EQ Submix directly (or by using legacy path via check box on SoundClass that uses value

diff --git a/Engine/Plugins/Runtime/Oculus/OculusAudio/Source/OculusAudio/Private/OculusAudioReverb.cpp b/Engine/Plugins/Runtime/Oculus/OculusAudio/Source/OculusAudio/Private/OculusAudioReverb.cpp
--- a/Engine/Plugins/Runtime/Oculus/OculusAudio/Source/OculusAudio/Private/OculusAudioReverb.cpp
+++ b/Engine/Plugins/Runtime/Oculus/OculusAudio/Source/OculusAudio/Private/OculusAudioReverb.cpp
@@ -2,6 +2,9 @@
 
 #include "OculusAudioReverb.h"
 #include "OculusAudioMixer.h"
+#include "OculusAudioSettings.h"
+
+#include "Sound/SoundSubmix.h"
 
 
 void FSubmixEffectOculusReverbPlugin::SetContext(ovrAudioContext* SharedContext)
@@ -44,11 +47,8 @@ void OculusAudioReverb::SetContext(ovrAudioContext* SharedContext)
 	{
 		Context = SharedContext;
 	}
-	
-	for (FSubmixEffectOculusReverbPlugin* Submix : Submixes)
-	{
-		Submix->SetContext(SharedContext);
-	}
+
+	Submix->SetContext(SharedContext);
 }
 
 void OculusAudioReverb::ClearContext()
@@ -60,9 +60,32 @@ void OculusAudioReverb::ClearContext()
 	}
 }
 
-FSoundEffectSubmix* OculusAudioReverb::GetEffectSubmix(class USoundSubmix* Submix)
+FSoundEffectSubmixPtr OculusAudioReverb::GetEffectSubmix()
 {
-	FSubmixEffectOculusReverbPlugin* ReverbSubmix = new FSubmixEffectOculusReverbPlugin();
-	Submixes.Add(ReverbSubmix);
-	return ReverbSubmix;
+	if (!Submix.IsValid())
+	{
+		Submix = MakeShared<FSubmixEffectOculusReverbPlugin>();
+		Submix->SetEnabled(true);
+	}
+
+	return StaticCastSharedPtr<FSoundEffectSubmix, FSubmixEffectOculusReverbPlugin, ESPMode::ThreadSafe>(Submix);
 }
+
+USoundSubmix* OculusAudioReverb::GetSubmix()
+{
+	const UOculusAudioSettings* Settings = GetDefault<UOculusAudioSettings>();
+	check(Settings);
+	USoundSubmix* ReverbSubmix = Cast<USoundSubmix>(Settings->OutputSubmix.TryLoad());
+	if (!ReverbSubmix)
+	{
+		static const FString DefaultSubmixName = TEXT("Oculus Reverb Submix");
+		UE_LOG(LogAudio, Error, TEXT("Failed to load Oculus Reverb Submix from object path '%s' in OculusSettings. Creating '%s' as stub."),
+			*Settings->OutputSubmix.GetAssetPathString(),
+			*DefaultSubmixName);
+
+		ReverbSubmix = NewObject<USoundSubmix>(USoundSubmix::StaticClass(), *DefaultSubmixName);
+		ReverbSubmix->bMuteWhenBackgrounded = true;
+	}
+
+	return ReverbSubmix;
+}
\ No newline at end of file
diff --git a/Engine/Plugins/Runtime/Oculus/OculusAudio/Source/OculusAudio/Private/OculusAudioReverb.h b/Engine/Plugins/Runtime/Oculus/OculusAudio/Source/OculusAudio/Private/OculusAudioReverb.h
--- a/Engine/Plugins/Runtime/Oculus/OculusAudio/Source/OculusAudio/Private/OculusAudioReverb.h
+++ b/Engine/Plugins/Runtime/Oculus/OculusAudio/Source/OculusAudio/Private/OculusAudioReverb.h
@@ -5,6 +5,7 @@
 #include "IAudioExtensionPlugin.h"
 #include "Sound/SoundEffectSubmix.h"
 #include "OVR_Audio.h"
+#include "Templates/UniquePtr.h"
 
 
 class FSubmixEffectOculusReverbPlugin : public FSoundEffectSubmix
@@ -61,7 +62,9 @@ public:
 		return; // PAS
 	}
 
-	virtual class FSoundEffectSubmix* GetEffectSubmix(class USoundSubmix* Submix) override;
+	virtual TSharedPtr<FSoundEffectSubmix> GetEffectSubmix() override;
+
+	virtual USoundSubmix* GetSubmix() override;
 
 	virtual void ProcessSourceAudio(const FAudioPluginSourceInputData& InputData, FAudioPluginSourceOutputData& OutputData) override
 	{
@@ -69,5 +72,5 @@ public:
 	}
 private:
 	ovrAudioContext* Context;
-	TArray<FSubmixEffectOculusReverbPlugin*> Submixes;
+	TSharedPtr<FSubmixEffectOculusReverbPlugin> Submix;
 };
\ No newline at end of file
diff --git a/Engine/Plugins/Runtime/Oculus/OculusAudio/Source/OculusAudio/Private/OculusAudioSettings.cpp b/Engine/Plugins/Runtime/Oculus/OculusAudio/Source/OculusAudio/Private/OculusAudioSettings.cpp
--- a/Engine/Plugins/Runtime/Oculus/OculusAudio/Source/OculusAudio/Private/OculusAudioSettings.cpp
+++ b/Engine/Plugins/Runtime/Oculus/OculusAudio/Source/OculusAudio/Private/OculusAudioSettings.cpp
@@ -17,7 +17,7 @@ UOculusAudioSettings::UOculusAudioSettings() :
 	, ReflectionCoefBack(0.25f)
 	, ReflectionCoefFront(0.25f)
 {
-
+	OutputSubmix = FString(TEXT("/OculusAudio/OculusSubmixDefault.OculusSubmixDefault"));
 }
 
 #if WITH_EDITOR
diff --git a/Engine/Plugins/Runtime/Oculus/OculusAudio/Source/OculusAudio/Private/OculusAudioSettings.h b/Engine/Plugins/Runtime/Oculus/OculusAudio/Source/OculusAudio/Private/OculusAudioSettings.h
--- a/Engine/Plugins/Runtime/Oculus/OculusAudio/Source/OculusAudio/Private/OculusAudioSettings.h
+++ b/Engine/Plugins/Runtime/Oculus/OculusAudio/Source/OculusAudio/Private/OculusAudioSettings.h
@@ -19,6 +19,9 @@ public:
 	virtual bool CanEditChange(const FProperty* InProperty) const override;
 #endif
 
+	UPROPERTY(GlobalConfig, EditAnywhere, Category = "Reverb", meta = (AllowedClasses = "SoundSubmix"))
+	FSoftObjectPath OutputSubmix;
+
 	UPROPERTY(GlobalConfig, EditAnywhere, Category = "Reverb", meta = (ClampMin = "-60.0", ClampMax = "20.0", UIMin = "-60.0", UIMax = "20.0"))
 	float ReverbWetLevel;
 
diff --git a/Engine/Plugins/Runtime/ResonanceAudio/Source/ResonanceAudio/Private/ResonanceAudioReverb.cpp b/Engine/Plugins/Runtime/ResonanceAudio/Source/ResonanceAudio/Private/ResonanceAudioReverb.cpp
--- a/Engine/Plugins/Runtime/ResonanceAudio/Source/ResonanceAudio/Private/ResonanceAudioReverb.cpp
+++ b/Engine/Plugins/Runtime/ResonanceAudio/Source/ResonanceAudio/Private/ResonanceAudioReverb.cpp
@@ -108,8 +108,13 @@ namespace ResonanceAudio
 		}
 	}
 
-	FSoundEffectSubmix* FResonanceAudioReverb::GetEffectSubmix(USoundSubmix* Submix)
+	FSoundEffectSubmixPtr FResonanceAudioReverb::GetEffectSubmix()
 	{
+		if (SubmixEffect.IsValid())
+		{
+			return SubmixEffect;
+		}
+
 		// Load the global reverb preset settings:
 		const FSoftObjectPath ReverbPluginPresetName = GetDefault<UResonanceAudioSettings>()->GlobalReverbPreset;
 		if (ReverbPluginPresetName.IsValid())
@@ -120,7 +125,7 @@ namespace ResonanceAudio
 		// If loading of the Reverb Plugin Preset asset fails, create a temporary preset. No reverb will be applied.
 		if (GlobalReverbPluginPreset == nullptr)
 		{
-			GlobalReverbPluginPreset = NewObject<UResonanceAudioReverbPluginPreset>(Submix, TEXT("Resonance Audio Reverb Plugin Preset"));
+			GlobalReverbPluginPreset = NewObject<UResonanceAudioReverbPluginPreset>(UResonanceAudioReverbPluginPreset::StaticClass(), TEXT("Resonance Audio Reverb Plugin Preset"));
 		}
 
 		if (GlobalReverbPluginPreset)
@@ -128,16 +133,62 @@ namespace ResonanceAudio
 			GlobalReverbPluginPreset->AddToRoot();
 			ReverbPluginPreset = GlobalReverbPluginPreset;
 
-			FResonanceAudioReverbPlugin* Effect = static_cast<FResonanceAudioReverbPlugin*>(GlobalReverbPluginPreset->CreateNewEffect());
-			Effect->SetPreset(GlobalReverbPluginPreset);
-			Effect->SetResonanceAudioReverbPlugin(this);
-			return static_cast<FSoundEffectSubmix*>(Effect);
+			FResonanceAudioReverbPlugin* ReverbSubmixEffect = static_cast<FResonanceAudioReverbPlugin*>(ReverbPluginPreset->CreateNewEffect());
+			ReverbSubmixEffect->SetPreset(GlobalReverbPluginPreset);
+			ReverbSubmixEffect->SetResonanceAudioReverbPlugin(this);
+
+			ReverbSubmixEffect->SetEnabled(true);
+			SubmixEffect = MakeShareable(ReverbSubmixEffect);
+
 		}
 		else
 		{
 			ReverbPluginPreset = nullptr;
-			return nullptr;
 		}
+
+		return SubmixEffect;
+	}
+
+	USoundSubmix* FResonanceAudioReverb::GetSubmix()
+	{
+		const UResonanceAudioSettings* Settings = GetDefault<UResonanceAudioSettings>();
+		check(Settings);
+
+		USoundSubmix* ReverbSubmix = Cast<USoundSubmix>(Settings->OutputSubmix.TryLoad());
+		if (!ReverbSubmix)
+		{
+			static const FString DefaultSubmixName = TEXT("Resonance Reverb Submix");
+			UE_LOG(LogResonanceAudio, Error, TEXT("Failed to load Resonance Reverb Submix from object path '%s' in ResonanceSettings. Creating '%s' as stub."),
+				*Settings->OutputSubmix.GetAssetPathString(),
+				*DefaultSubmixName);
+
+			ReverbSubmix = NewObject<USoundSubmix>(USoundSubmix::StaticClass(), *DefaultSubmixName);
+			ReverbSubmix->bMuteWhenBackgrounded = true;
+		}
+
+		if (ReverbPluginPreset)
+		{
+			bool bFoundPreset = false;
+			for (USoundEffectSubmixPreset* Preset : ReverbSubmix->SubmixEffectChain)
+			{
+				if (UResonanceAudioReverbPluginPreset* PluginPreset = Cast<UResonanceAudioReverbPluginPreset>(Preset))
+				{
+					bFoundPreset = true;
+					break;
+				}
+			}
+
+			if (!bFoundPreset)
+			{
+				static const FString DefaultPresetName = TEXT("ResonanceReverbDefault_0");
+				UE_LOG(LogResonanceAudio, Error, TEXT("Failed to find Resonance UResonanceAudioReverbPluginPreset on default reverb submix. Creating stub '%s'."),
+					*Settings->OutputSubmix.GetAssetPathString(),
+					*DefaultPresetName);
+				ReverbSubmix->SubmixEffectChain.Add(ReverbPluginPreset);
+			}
+		}
+
+		return ReverbSubmix;
 	}
 
 	void FResonanceAudioReverb::SetPreset(UResonanceAudioReverbPluginPreset* InPreset)
diff --git a/Engine/Plugins/Runtime/ResonanceAudio/Source/ResonanceAudio/Private/ResonanceAudioReverb.h b/Engine/Plugins/Runtime/ResonanceAudio/Source/ResonanceAudio/Private/ResonanceAudioReverb.h
--- a/Engine/Plugins/Runtime/ResonanceAudio/Source/ResonanceAudio/Private/ResonanceAudioReverb.h
+++ b/Engine/Plugins/Runtime/ResonanceAudio/Source/ResonanceAudio/Private/ResonanceAudioReverb.h
@@ -111,9 +111,9 @@ namespace ResonanceAudio
 		virtual void Initialize(const FAudioPluginInitializationParams InitializationParams) override;
 		virtual void OnInitSource(const uint32 SourceId, const FName& AudioComponentUserId, const uint32 NumChannels, UReverbPluginSourceSettingsBase* InSettings) override;
 		virtual void OnReleaseSource(const uint32 SourceId) override;
-		virtual class FSoundEffectSubmix* GetEffectSubmix(class USoundSubmix* Submix) override;
+		virtual FSoundEffectSubmixPtr GetEffectSubmix() override;
+		virtual USoundSubmix* GetSubmix() override;
 		virtual void ProcessSourceAudio(const FAudioPluginSourceInputData& InputData, FAudioPluginSourceOutputData& OutputData) override;
-		virtual bool DoesReverbOverrideMasterReverb() const override { return false; }
 
 		void SetResonanceAudioApi(vraudio::ResonanceAudioApi* InResonanceAudioApi) { ResonanceAudioApi = InResonanceAudioApi; };
 		void SetPreset(UResonanceAudioReverbPluginPreset* InPreset);
@@ -145,6 +145,8 @@ namespace ResonanceAudio
 		UResonanceAudioReverbPluginPreset* ReverbPluginPreset;
 		UResonanceAudioReverbPluginPreset* GlobalReverbPluginPreset;
 		TArray<float, TAlignedHeapAllocator<AUDIO_BUFFER_ALIGNMENT>>  TemporaryStereoBuffer;
+
+		FSoundEffectSubmixPtr SubmixEffect;
 	};
 
 } // namespace ResonanceAudio
diff --git a/Engine/Plugins/Runtime/ResonanceAudio/Source/ResonanceAudio/Private/ResonanceAudioSettings.cpp b/Engine/Plugins/Runtime/ResonanceAudio/Source/ResonanceAudio/Private/ResonanceAudioSettings.cpp
--- a/Engine/Plugins/Runtime/ResonanceAudio/Source/ResonanceAudio/Private/ResonanceAudioSettings.cpp
+++ b/Engine/Plugins/Runtime/ResonanceAudio/Source/ResonanceAudio/Private/ResonanceAudioSettings.cpp
@@ -8,4 +8,5 @@ UResonanceAudioSettings::UResonanceAudioSettings()
 	: QualityMode(ERaQualityMode::BINAURAL_HIGH)
 	, GlobalReverbPreset(nullptr)
 {
+	OutputSubmix = FString(TEXT("/ResonanceAudio/ResonanceSubmixDefault.ResonanceSubmixDefault"));
 }
diff --git a/Engine/Plugins/Runtime/ResonanceAudio/Source/ResonanceAudio/Private/ResonanceAudioSettings.h b/Engine/Plugins/Runtime/ResonanceAudio/Source/ResonanceAudio/Private/ResonanceAudioSettings.h
--- a/Engine/Plugins/Runtime/ResonanceAudio/Source/ResonanceAudio/Private/ResonanceAudioSettings.h
+++ b/Engine/Plugins/Runtime/ResonanceAudio/Source/ResonanceAudio/Private/ResonanceAudioSettings.h
@@ -16,6 +16,10 @@ public:
 
 	UResonanceAudioSettings();
 
+	// Reference to submix where reverb plugin audio is routed.
+	UPROPERTY(GlobalConfig, EditAnywhere, Category = "Reverb", meta = (AllowedClasses = "SoundSubmix"))
+	FSoftObjectPath OutputSubmix;
+
 	// Global Quality mode to use for directional sound sources.
 	UPROPERTY(GlobalConfig, EditAnywhere, Category = General)
 	ERaQualityMode QualityMode;
diff --git a/Engine/Plugins/Runtime/Steam/SteamAudio/Source/SteamAudio/Private/PhononReverb.cpp b/Engine/Plugins/Runtime/Steam/SteamAudio/Source/SteamAudio/Private/PhononReverb.cpp
--- a/Engine/Plugins/Runtime/Steam/SteamAudio/Source/SteamAudio/Private/PhononReverb.cpp
+++ b/Engine/Plugins/Runtime/Steam/SteamAudio/Source/SteamAudio/Private/PhononReverb.cpp
@@ -391,12 +391,58 @@ namespace SteamAudio
 		bListenerInitialized = true;
 	}
 
-	FSoundEffectSubmix* FPhononReverb::GetEffectSubmix(USoundSubmix* Submix)
+	FSoundEffectSubmixPtr FPhononReverb::GetEffectSubmix()
 	{
-		USubmixEffectReverbPluginPreset* ReverbPluginPreset = NewObject<USubmixEffectReverbPluginPreset>(Submix, TEXT("Master Reverb Plugin Effect Preset"));
-		auto Effect = static_cast<FSubmixEffectReverbPlugin*>(ReverbPluginPreset->CreateNewEffect());
-		Effect->SetPhononReverbPlugin(this);
-		return static_cast<FSoundEffectSubmix*>(Effect);
+		if (!SubmixEffect.IsValid())
+		{
+			USubmixEffectReverbPluginPreset* ReverbPluginPreset = NewObject<USubmixEffectReverbPluginPreset>(GetSubmix(), TEXT("Master Reverb Plugin Effect Preset"));
+			FSubmixEffectReverbPlugin* ReverbSubmixEffect = static_cast<FSubmixEffectReverbPlugin*>(ReverbPluginPreset->CreateNewEffect());
+			ReverbSubmixEffect->SetPhononReverbPlugin(this);
+			USubmixEffectReverbPluginPreset
+			ReverbSubmixEffect->SetEnabled(true);
+
+			SubmixEffect = MakeShareable(static_cast<FSoundEffectSubmix>(ReverbSubmixEffect));
+		}
+
+		return SubmixEffect;
+	}
+
+	USoundSubmix* FPhononReverb::GetSubmix()
+	{
+		const USteamAudioSettings* Settings = GetDefault<USteamAudioSettings>();
+		check(Settings);
+		USoundSubmix* ReverbSubmix = Cast<USoundSubmix>(Settings->OutputSubmix.TryLoad());
+		if (!ReverbSubmix)
+		{
+			static const FString DefaultSubmixName = TEXT("Phonon Reverb Submix");
+			UE_LOG(LogSteamAudio, Error, TEXT("Failed to load Phonon Reverb Submix from object path '%s' in PhononSettings. Creating '%s' as stub."),
+				*Settings->OutputSubmix.GetAssetPathString(),
+				*DefaultSubmixName);
+
+			ReverbSubmix = NewObject<USoundSubmix>(USoundSubmix::StaticClass(), *DefaultSubmixName);
+			ReverbSubmix->bMuteWhenBackgrounded = true;
+		}
+
+		bool bFoundPreset = false;
+		for (USubmixEffectReverbPluginPreset* Preset : ReverbSubmix->SubmixEffectChain)
+		{
+			if (USubmixEffectReverbPluginPreset* PluginPreset = Cast<USubmixEffectReverbPluginPreset>(Preset))
+			{
+				bFoundPreset = true;
+				break;
+			}
+		}
+
+		if (!bFoundPreset)
+		{
+			static const FString DefaultPresetName = TEXT("PhononReverbDefault_0");
+			UE_LOG(LogSteamAudio, Error, TEXT("Failed to find Phonon USubmixEffectReverbPluginPreset on default reverb submix. Creating stub '%s'."),
+				*Settings->OutputSubmix.GetAssetPathString(),
+				*DefaultPresetName);
+			ReverbSubmix->SubmixEffectChain.Add(NewObject<USubmixEffectReverbPluginPreset>(USubmixEffectReverbPluginPreset::StaticClass(), *DefaultPresetName));
+		}
+
+		return ReverbSubmix;
 	}
 
 	//==============================================================================================================================================
@@ -421,7 +467,7 @@ FSubmixEffectReverbPlugin::FSubmixEffectReverbPlugin()
 	: PhononReverbPlugin(nullptr)
 {}
 
-void FSubmixEffectReverbPlugin::Init(const FSoundEffectSubmixInitData& InSampleRate)
+void FSubmixEffectReverbPlugin::Init(const FSoundEffectSubmixInitData& InInitData)
 {
 }
 
diff --git a/Engine/Plugins/Runtime/Steam/SteamAudio/Source/SteamAudio/Private/PhononReverb.h b/Engine/Plugins/Runtime/Steam/SteamAudio/Source/SteamAudio/Private/PhononReverb.h
--- a/Engine/Plugins/Runtime/Steam/SteamAudio/Source/SteamAudio/Private/PhononReverb.h
+++ b/Engine/Plugins/Runtime/Steam/SteamAudio/Source/SteamAudio/Private/PhononReverb.h
@@ -11,6 +11,11 @@
 #include "phonon.h"
 #include "PhononReverb.generated.h"
 
+
+// Forward Declaration
+class FSubmixEffectReverbPlugin;
+
+
 namespace SteamAudio
 {
 	class FEnvironment;
@@ -39,7 +44,8 @@ namespace SteamAudio
 		virtual void Initialize(const FAudioPluginInitializationParams InitializationParams) override;
 		virtual void OnInitSource(const uint32 SourceId, const FName& AudioComponentUserId, const uint32 NumChannels, UReverbPluginSourceSettingsBase* InSettings) override;
 		virtual void OnReleaseSource(const uint32 SourceId) override;
-		virtual class FSoundEffectSubmix* GetEffectSubmix(class USoundSubmix* Submix) override;
+		virtual FSoundEffectSubmixPtr GetEffectSubmix() override;
+		virtual USoundSubmix* GetSubmix() override;
 		virtual void ProcessSourceAudio(const FAudioPluginSourceInputData& InputData, FAudioPluginSourceOutputData& OutputData) override;
 		
 		void ProcessMixedAudio(const FSoundEffectSubmixInputData& InData, FSoundEffectSubmixOutputData& OutData);
@@ -48,6 +54,8 @@ namespace SteamAudio
 		void UpdateListener(const FVector& Position, const FVector& Forward, const FVector& Up, const FVector& Right);
 
 	private:
+		FSoundEffectSubmixPtr SubmixEffect;
+
 		IPLhandle BinauralRenderer;
 		IPLhandle IndirectBinauralEffect;
 		IPLhandle IndirectPanningEffect;
diff --git a/Engine/Plugins/Runtime/Steam/SteamAudio/Source/SteamAudio/Private/SteamAudioSettings.cpp b/Engine/Plugins/Runtime/Steam/SteamAudio/Source/SteamAudio/Private/SteamAudioSettings.cpp
--- a/Engine/Plugins/Runtime/Steam/SteamAudio/Source/SteamAudio/Private/SteamAudioSettings.cpp
+++ b/Engine/Plugins/Runtime/Steam/SteamAudio/Source/SteamAudio/Private/SteamAudioSettings.cpp
@@ -38,6 +38,8 @@ USteamAudioSettings::USteamAudioSettings()
 	, TANMaxSources(32)
 	, RadeonRaysBakingBatchSize(1)
 {
+	OutputSubmix = FString(TEXT("/PhononAudio/PhononSubmixDefault.PhononSubmixDefault"));
+
 	auto MaterialPreset = SteamAudio::MaterialPresets[StaticMeshMaterialPreset];
 	StaticMeshLowFreqAbsorption = MaterialPreset.lowFreqAbsorption;
 	StaticMeshMidFreqAbsorption = MaterialPreset.midFreqAbsorption;
diff --git a/Engine/Plugins/Runtime/Steam/SteamAudio/Source/SteamAudio/Private/SteamAudioSettings.h b/Engine/Plugins/Runtime/Steam/SteamAudio/Source/SteamAudio/Private/SteamAudioSettings.h
--- a/Engine/Plugins/Runtime/Steam/SteamAudio/Source/SteamAudio/Private/SteamAudioSettings.h
+++ b/Engine/Plugins/Runtime/Steam/SteamAudio/Source/SteamAudio/Private/SteamAudioSettings.h
@@ -33,6 +33,9 @@ public:
 	virtual bool CanEditChange(const FProperty* InProperty) const override;
 #endif
 
+	UPROPERTY(GlobalConfig, EditAnywhere, Category = General, meta = (AllowedClasses = "SoundSubmix"))
+	FSoftObjectPath OutputSubmix;
+
 	// Which audio engine to use.
 	UPROPERTY(GlobalConfig, EditAnywhere, Category = General)
 	EIplAudioEngine AudioEngine;
diff --git a/Engine/Source/Runtime/AudioMixer/Classes/SubmixEffects/AudioMixerSubmixEffectReverbFast.h b/Engine/Source/Runtime/AudioMixer/Classes/SubmixEffects/AudioMixerSubmixEffectReverbFast.h
--- a/Engine/Source/Runtime/AudioMixer/Classes/SubmixEffects/AudioMixerSubmixEffectReverbFast.h
+++ b/Engine/Source/Runtime/AudioMixer/Classes/SubmixEffects/AudioMixerSubmixEffectReverbFast.h
@@ -12,72 +12,77 @@ struct AUDIOMIXER_API FSubmixEffectReverbFastSettings
 {
 	GENERATED_USTRUCT_BODY()
 
+	/** Bypasses reverb */
+	UPROPERTY(EditAnywhere, BlueprintReadWrite, Category = General)
+	bool bBypass;
+
 	/** Density - 0.0 < 0.85 < 1.0 - Coloration of the late reverb - lower value is more grainy */
-	UPROPERTY(EditAnywhere, BlueprintReadWrite, Category = ReverbParameters, meta = (ClampMin = "0.0", ClampMax = "1.0"))
+	UPROPERTY(EditAnywhere, BlueprintReadWrite, Category = General, meta = (ClampMin = "0.0", ClampMax = "1.0", EditCondition = "!bBypass"))
 	float Density;
 
 	/** Diffusion - 0.0 < 0.85 < 1.0 - Echo density in the reverberation decay - lower is more grainy */
-	UPROPERTY(EditAnywhere, BlueprintReadWrite, Category = ReverbParameters, meta = (ClampMin = "0.0", ClampMax = "1.0"))
+	UPROPERTY(EditAnywhere, BlueprintReadWrite, Category = General, meta = (ClampMin = "0.0", ClampMax = "1.0", EditCondition = "!bBypass"))
 	float Diffusion;
 
 	/** Reverb Gain - 0.0 < 0.32 < 1.0 - overall reverb gain - master volume control */
-	UPROPERTY(EditAnywhere, BlueprintReadWrite, Category = ReverbParameters, meta = (ClampMin = "0.0", ClampMax = "1.0"))
+	UPROPERTY(EditAnywhere, BlueprintReadWrite, Category = General, meta = (ClampMin = "0.0", ClampMax = "1.0", EditCondition = "!bBypass"))
 	float Gain;
 
 	/** Reverb Gain High Frequency - 0.0 < 0.89 < 1.0 - attenuates the high frequency reflected sound */
-	UPROPERTY(EditAnywhere, BlueprintReadWrite, Category = ReverbParameters, meta = (ClampMin = "0.0", ClampMax = "1.0"))
+	UPROPERTY(EditAnywhere, BlueprintReadWrite, Category = General, meta = (ClampMin = "0.0", ClampMax = "1.0", EditCondition = "!bBypass"))
 	float GainHF;
 
 	/** Decay Time - 0.1 < 1.49 < 20.0 Seconds - larger is more reverb */
-	UPROPERTY(EditAnywhere, BlueprintReadWrite, Category = ReverbParameters, meta = (ClampMin = "0.1", ClampMax = "20.0"))
+	UPROPERTY(EditAnywhere, BlueprintReadWrite, Category = Decay, meta = (ClampMin = "0.1", ClampMax = "20.0", EditCondition = "!bBypass"))
 	float DecayTime;
 
 	/** Decay High Frequency Ratio - 0.1 < 0.83 < 2.0 - how much the quicker or slower the high frequencies decay relative to the lower frequencies. */
-	UPROPERTY(EditAnywhere, BlueprintReadWrite, Category = ReverbParameters, meta = (ClampMin = "0.1", ClampMax = "2.0"))
+	UPROPERTY(EditAnywhere, BlueprintReadWrite, Category = Decay, meta = (ClampMin = "0.1", ClampMax = "2.0", EditCondition = "!bBypass"))
 	float DecayHFRatio;
 
 	/** Reflections Gain - 0.0 < 0.05 < 3.16 - controls the amount of initial reflections */
-	UPROPERTY(EditAnywhere, BlueprintReadWrite, Category = ReverbParameters, meta = (ClampMin = "0.0", ClampMax = "3.16"))
+	UPROPERTY(EditAnywhere, BlueprintReadWrite, Category = EarlyReflections, meta = (ClampMin = "0.0", ClampMax = "3.16", EditCondition = "!bBypass"))
 	float ReflectionsGain;
 
 	/** Reflections Delay - 0.0 < 0.007 < 0.3 Seconds - the time between the listener receiving the direct path sound and the first reflection */
-	UPROPERTY(EditAnywhere, BlueprintReadWrite, Category = ReverbParameters, meta = (ClampMin = "0.0", ClampMax = "0.3"))
+	UPROPERTY(EditAnywhere, BlueprintReadWrite, Category = EarlyReflections, meta = (ClampMin = "0.0", ClampMax = "0.3", EditCondition = "!bBypass"))
 	float ReflectionsDelay;
 
 	/** Late Reverb Gain - 0.0 < 1.26 < 10.0 - gain of the late reverb */
-	UPROPERTY(EditAnywhere, BlueprintReadWrite, Category = ReverbParameters, meta = (ClampMin = "0.0", ClampMax = "10.0"))
+	UPROPERTY(EditAnywhere, BlueprintReadWrite, Category = LateReflections, meta = (ClampMin = "0.0", ClampMax = "10.0", EditCondition = "!bBypass"))
 	float LateGain;
 
 	/** Late Reverb Delay - 0.0 < 0.011 < 0.1 Seconds - time difference between late reverb and first reflections */
-	UPROPERTY(EditAnywhere, BlueprintReadWrite, Category = ReverbParameters, meta = (ClampMin = "0.0", ClampMax = "0.1"))
+	UPROPERTY(EditAnywhere, BlueprintReadWrite, Category = LateReflections, meta = (ClampMin = "0.0", ClampMax = "0.1", EditCondition = "!bBypass"))
 	float LateDelay;
 
 	/** Air Absorption - 0.0 < 0.994 < 1.0 - lower value means more absorption */
-	UPROPERTY(EditAnywhere, BlueprintReadWrite, Category = ReverbParameters, meta = (ClampMin = "0.0", ClampMax = "1.0"))
+	UPROPERTY(EditAnywhere, BlueprintReadWrite, Category = General, meta = (ClampMin = "0.0", ClampMax = "1.0", EditCondition = "!bBypass"))
 	float AirAbsorptionGainHF;
 
 	// Overall wetlevel of the reverb effect
-	UPROPERTY(EditAnywhere, BlueprintReadWrite, Category = Reverb)
+	UPROPERTY(EditAnywhere, BlueprintReadWrite, Category = Routing, meta = (EditCondition = "!bBypass"))
 	float WetLevel;
 
 	// Overall drylevel of the reverb effect
-	UPROPERTY(EditAnywhere, BlueprintReadWrite, Category = Reverb)
+	UPROPERTY(EditAnywhere, BlueprintReadWrite, Category = Routing, meta = (EditCondition = "!bBypass"))
 	float DryLevel;
 
 	FSubmixEffectReverbFastSettings()
-	: Density(0.85f)
-	, Diffusion(0.85f)
-	, Gain(0.0f)
-	, GainHF(0.89f)
-	, DecayTime(1.49f)
-	, DecayHFRatio(0.83f)
-	, ReflectionsGain(0.05f)
-	, ReflectionsDelay(0.007f)
-	, LateGain(1.26f)
-	, LateDelay(0.1f)
-	, AirAbsorptionGainHF(0.994f)
-	, WetLevel(0.3f)
-	, DryLevel(0.0f)
+		: bBypass(false)
+		, Density(0.85f)
+		, Diffusion(0.85f)
+		, Gain(0.0f)
+		, GainHF(0.89f)
+		, DecayTime(1.49f)
+		, DecayHFRatio(0.83f)
+		, ReflectionsGain(0.05f)
+		, ReflectionsDelay(0.007f)
+		, LateGain(1.26f)
+		, LateDelay(0.1f)
+		, AirAbsorptionGainHF(0.994f)
+		, WetLevel(0.3f)
+		, DryLevel(0.0f)
 	{
 	}
 };
@@ -114,6 +119,7 @@ private:
 	// Curve which maps old reverb times to new decay value
 	FRichCurve DecayCurve;
 
+	bool bBypass;
 };
 
 UCLASS()
diff --git a/Engine/Source/Runtime/AudioMixer/Private/AudioMixerBlueprintLibrary.cpp b/Engine/Source/Runtime/AudioMixer/Private/AudioMixerBlueprintLibrary.cpp
--- a/Engine/Source/Runtime/AudioMixer/Private/AudioMixerBlueprintLibrary.cpp
+++ b/Engine/Source/Runtime/AudioMixer/Private/AudioMixerBlueprintLibrary.cpp
@@ -52,10 +52,7 @@ void UAudioMixerBlueprintLibrary::AddMasterSubmixEffect(const UObject* WorldCont
 	if (Audio::FMixerDevice* MixerDevice = GetAudioMixerDeviceFromWorldContext(WorldContextObject))
 	{
 		// Immediately create a new sound effect base here before the object becomes potentially invalidated
-		FSoundEffectBase* SoundEffectBase = SubmixEffectPreset->CreateNewEffect();
-
-		// Cast it to a sound effect submix type
-		FSoundEffectSubmix* SoundEffectSubmix = static_cast<FSoundEffectSubmix*>(SoundEffectBase);
+		FSoundEffectSubmixPtr SoundEffectSubmix = MakeShareable(static_cast<FSoundEffectSubmix*>(SubmixEffectPreset->CreateNewEffect()));
 
 		FSoundEffectSubmixInitData InitData;
 		InitData.SampleRate = MixerDevice->GetSampleRate();
diff --git a/Engine/Source/Runtime/AudioMixer/Private/AudioMixerDevice.cpp b/Engine/Source/Runtime/AudioMixer/Private/AudioMixerDevice.cpp
--- a/Engine/Source/Runtime/AudioMixer/Private/AudioMixerDevice.cpp
+++ b/Engine/Source/Runtime/AudioMixer/Private/AudioMixerDevice.cpp
@@ -1,19 +1,21 @@
 // Copyright 1998-2019 Epic Games, Inc. All Rights Reserved.
 
 #include "AudioMixerDevice.h"
+
 #include "AudioMixerSource.h"
 #include "AudioMixerSubmix.h"
 #include "AudioMixerSourceVoice.h"
 #include "AudioPluginUtilities.h"
 #include "AudioMixerEffectsManager.h"
+#include "DSP/Noise.h"
+#include "DSP/SinOsc.h"
 #include "Sound/AudioSettings.h"
+#include "Sound/SoundSubmix.h"
 #include "Sound/SoundSubmixSend.h"
 #include "SubmixEffects/AudioMixerSubmixEffectReverb.h"
 #include "SubmixEffects/AudioMixerSubmixEffectReverbFast.h"
 #include "SubmixEffects/AudioMixerSubmixEffectEQ.h"
 #include "SubmixEffects/AudioMixerSubmixEffectDynamicsProcessor.h"
-#include "DSP/Noise.h"
-#include "DSP/SinOsc.h"
 #include "UObject/UObjectHash.h"
 #include "UObject/UObjectIterator.h"
 #include "Runtime/HeadMountedDisplay/Public/IHeadMountedDisplayModule.h"
@@ -23,7 +25,8 @@
 
 #if WITH_EDITOR
 #include "AudioEditorModule.h"
-#endif
+#endif // WITH_EDITOR
+
 
 static int32 DisableSubmixEffectEQCvar = 0;
 FAutoConsoleVariableRef CVarDisableSubmixEQ(
@@ -38,8 +41,6 @@ CSV_DECLARE_CATEGORY_MODULE_EXTERN(AUDIOMIXERCORE_API, Audio);
 
 namespace Audio
 {
-	TArray<USoundSubmix*> FMixerDevice::MasterSubmixes;
-
 	FMixerDevice::FMixerDevice(IAudioMixerPlatformInterface* InAudioMixerPlatform)
 		: AudioMixerPlatform(InAudioMixerPlatform)
 		, AudioClockDelta(0.0)
@@ -49,6 +50,7 @@ namespace Audio
 		, GameOrAudioThreadId(INDEX_NONE)
 		, AudioPlatformThreadId(INDEX_NONE)
 		, bDebugOutputEnabled(false)
+		, bSubmixRegistrationDisabled(false)
 	{
 		// This audio device is the audio mixer
 		bAudioMixerModuleLoaded = true;
@@ -373,19 +375,8 @@ namespace Audio
 			if (SoundSubmix)
 			{
 				// Retrieve the submix instance and the envelope data
-				Audio::FMixerSubmixWeakPtr SubmixPtr;
-
-				// First see list of submixes
-				Audio::FMixerSubmixPtr* FoundSubmix = Submixes.Find(SoundSubmix);
-				if (FoundSubmix)
-				{
-					SubmixPtr = *FoundSubmix;
-				}
-				// Fallback to master submix
-				else
-				{
-					SubmixPtr = GetMasterSubmix();
-				}
+				Audio::FMixerSubmixWeakPtr SubmixPtr = GetSubmixInstance(SoundSubmix);
+				check(SubmixPtr.IsValid());
 
 				// On the audio thread, do the broadcast.
 				FAudioThread::RunCommandOnGameThread([this, SubmixPtr]()
@@ -410,17 +401,8 @@ namespace Audio
 			{
 				if (*It && It->bMuteWhenBackgrounded)
 				{
-					FMixerSubmix* SubmixInstance = GetMasterSubmixInstance(*It);
-					if (!SubmixInstance)
-					{
-						FMixerSubmixPtr* NonMasterSubmix = Submixes.Find(*It);
-						if (NonMasterSubmix)
-						{
-							SubmixInstance = (*NonMasterSubmix).Get();
-						}
-					}
-
-					if (SubmixInstance)
+					FMixerSubmixPtr SubmixInstance = GetSubmixInstance(*It).Pin();
+					if (SubmixInstance.IsValid())
 					{
 						SubmixInstance->SetBackgroundMuted(IsMuted);
 					}
@@ -569,173 +551,145 @@ namespace Audio
 		SourceManager.UpdatePendingReleaseData(true);
 	}
 
-	void FMixerDevice::InitSoundSubmixes()
+	void FMixerDevice::InitMasterSoundSubmix(EMasterSubmixType::Type InType, const FString& InDefaultName, bool bInDefaultMuteWhenBackgrounded, FSoftObjectPath& InObjectPath)
 	{
-		if (!IsInAudioThread())
-		{
-			DECLARE_CYCLE_STAT(TEXT("FAudioThreadTask.InitSoundSubmixes"), STAT_InitSoundSubmixes, STATGROUP_AudioThreadCommands);
+		check(IsInAudioThread());
 
-			FAudioThread::RunCommandOnAudioThread([this]()
-			{
-				CSV_SCOPED_TIMING_STAT(Audio, InitSubmix);
-				InitSoundSubmixes();
-			}, GET_STATID(STAT_InitSoundSubmixes));
-			return;
+		const int32 MasterSubmixCount = static_cast<int32>(EMasterSubmixType::Type::Count);
+		if(MasterSubmixes.Num() < MasterSubmixCount)
+		{
+			MasterSubmixes.AddZeroed(MasterSubmixCount - MasterSubmixes.Num());
 		}
 
-		// Create the master, master reverb, and master eq sound submixes
-		if (!FMixerDevice::MasterSubmixes.Num())
+		if (MasterSubmixInstances.Num() < MasterSubmixCount)
 		{
-			// Master
-			USoundSubmix* MasterSubmix = NewObject<USoundSubmix>(USoundSubmix::StaticClass(), TEXT("Master Submix"));
-			MasterSubmix->AddToRoot();
-			FMixerDevice::MasterSubmixes.Add(MasterSubmix);
-
-			// Master Reverb Plugin
-			USoundSubmix* ReverbPluginSubmix = NewObject<USoundSubmix>(USoundSubmix::StaticClass(), TEXT("Master Reverb Plugin Submix"));
-			// Make the master reverb mute when backgrounded
-			ReverbPluginSubmix->bMuteWhenBackgrounded = true;
-
-			ReverbPluginSubmix->AddToRoot();
-			FMixerDevice::MasterSubmixes.Add(ReverbPluginSubmix);
-
-			// Master Reverb
-			USoundSubmix* ReverbSubmix = NewObject<USoundSubmix>(USoundSubmix::StaticClass(), TEXT("Master Reverb Submix"));
-			// Make the master reverb mute when backgrounded
-			ReverbSubmix->bMuteWhenBackgrounded = true;
-
-			ReverbSubmix->AddToRoot();
-			FMixerDevice::MasterSubmixes.Add(ReverbSubmix);
+			MasterSubmixInstances.AddZeroed(MasterSubmixCount - MasterSubmixInstances.Num());
+		}
 
-			// Master EQ
-			if (DisableSubmixEffectEQCvar)
+		const int32 TypeIndex = static_cast<int32>(InType);
+		if (USoundSubmix* OldSubmix = MasterSubmixes[TypeIndex])
+		{
+			// Don't bother swapping if new path is invalid...
+			if (!InObjectPath.IsValid())
 			{
-				FMixerDevice::MasterSubmixes.Add(nullptr);
+				return;
 			}
-			else
+
+			// or is same object already initialized.
+			if (InObjectPath.GetAssetPathString() == OldSubmix->GetPathName())
 			{
-				USoundSubmix* EQSubmix = NewObject<USoundSubmix>(USoundSubmix::StaticClass(), TEXT("Master EQ Submix"));
-				EQSubmix->AddToRoot();
-				FMixerDevice::MasterSubmixes.Add(EQSubmix);
+				return;
 			}
 
-			// Master ambisonics
-			USoundSubmix* AmbisonicsSubmix = NewObject<USoundSubmix>(USoundSubmix::StaticClass(), TEXT("Master Ambisonics Submix"));
-			AmbisonicsSubmix->AddToRoot();
-			AmbisonicsSubmix->ChannelFormat = ESubmixChannelFormat::Ambisonics;
-			if (AmbisonicsMixer.IsValid())
+			OldSubmix->RemoveFromRoot();
+			FMixerSubmixPtr OldSubmixPtr = MasterSubmixInstances[TypeIndex];
+			if (OldSubmixPtr.IsValid())
 			{
-				AmbisonicsSubmix->AmbisonicsPluginSettings = AmbisonicsMixer->GetDefaultSettings();
+				FMixerSubmixPtr ParentSubmixPtr = MasterSubmixInstances[TypeIndex]->GetParentSubmix().Pin();
+				if (ParentSubmixPtr.IsValid())
+				{
+					ParentSubmixPtr->RemoveChildSubmix(MasterSubmixInstances[TypeIndex]);
+				}
 			}
-			FMixerDevice::MasterSubmixes.Add(AmbisonicsSubmix);
 		}
 
-		// Register and setup the master submixes so that the rest of the submixes can hook into these core master submixes
+		// 1. Try loading from Developer Audio Settings
+		USoundSubmix* NewSubmix = Cast<USoundSubmix>(InObjectPath.TryLoad());
 
-		if (MasterSubmixInstances.Num() == 0)
+		// 2. If Unset or not found, fallback to engine asset
+		if (!NewSubmix)
 		{
-			for (int32 i = 0; i < EMasterSubmixType::Count; ++i)
-			{
-				if (FMixerDevice::MasterSubmixes[i])
-				{
-					FMixerSubmixPtr MixerSubmix = FMixerSubmixPtr(new FMixerSubmix(this));
-					MixerSubmix->Init(FMixerDevice::MasterSubmixes[i]);
+			static const FString EngineSubmixDir = TEXT("/Engine/EngineSounds/Submixes");
 
-					MasterSubmixInstances.Add(MixerSubmix);
-				}
-				else
-				{
-					MasterSubmixInstances.Add(nullptr);
-				}
-			}
+			InObjectPath = FString::Printf(TEXT("%s/%s.%s"), *EngineSubmixDir, *InDefaultName, *InDefaultName);
+			NewSubmix = Cast<USoundSubmix>(InObjectPath.TryLoad());
+			UE_LOG(LogAudioMixer, Display, TEXT("Submix unset or invalid in 'AudioSettings': Using engine asset '%s'"),
+				*InDefaultName,
+				*InObjectPath.GetAssetPathString());
+		}
 
-			FMixerSubmixPtr& MasterSubmixInstance = MasterSubmixInstances[EMasterSubmixType::Master];
-			
-			FSoundEffectSubmixInitData InitData;
-			InitData.SampleRate = GetSampleRate();
-	
-			bool bIsMasterReverbBypassed = false;
+		// 3. If engine version not found, dynamically spawn and post error
+		if (!NewSubmix)
+		{
+			UE_LOG(LogAudioMixer, Error, TEXT("Failed to load submix from engine asset path '%s'. Creating '%s' as a stub."),
+				*InObjectPath.GetAssetPathString(),
+				*InDefaultName);
 
-			// Setup the master reverb plugin
-			if (ReverbPluginInterface.IsValid() && MasterSubmixInstances[EMasterSubmixType::ReverbPlugin].IsValid())
-			{
-				FSoundEffectSubmix* ReverbPluginEffectSubmix = ReverbPluginInterface->GetEffectSubmix(FMixerDevice::MasterSubmixes[EMasterSubmixType::ReverbPlugin]);
+			NewSubmix = NewObject<USoundSubmix>(USoundSubmix::StaticClass(), *InDefaultName);
+			// Make the master reverb mute when backgrounded
+			NewSubmix->bMuteWhenBackgrounded = bInDefaultMuteWhenBackgrounded;
+		}
 
-				ReverbPluginEffectSubmix->Init(InitData);
-				ReverbPluginEffectSubmix->SetEnabled(true);
+		check(NewSubmix);
+		NewSubmix->AddToRoot();
+		MasterSubmixes[TypeIndex] = NewSubmix;
 
-				const uint32 ReverbPluginId = FMixerDevice::MasterSubmixes[EMasterSubmixType::ReverbPlugin]->GetUniqueID();
+		FMixerSubmixPtr NewMixerSubmix = MakeShared<FMixerSubmix, ESPMode::ThreadSafe>(this);
+		MasterSubmixInstances[TypeIndex] = NewMixerSubmix;
 
-				FMixerSubmixPtr MasterReverbPluginSubmix = MasterSubmixInstances[EMasterSubmixType::ReverbPlugin];
-				MasterReverbPluginSubmix->AddSoundEffectSubmix(ReverbPluginId, ReverbPluginEffectSubmix);
-				MasterReverbPluginSubmix->SetParentSubmix(MasterSubmixInstance);
-				MasterSubmixInstance->AddChildSubmix(MasterReverbPluginSubmix);
+		if (InType != EMasterSubmixType::Master)
+		{
+			FMixerSubmixPtr ParentSubmixInstance = NewSubmix->ParentSubmix
+				? GetSubmixInstance(NewSubmix->ParentSubmix).Pin()
+				: GetMasterSubmix().Pin();
 
-				bIsMasterReverbBypassed = ReverbPluginInterface->DoesReverbOverrideMasterReverb();
-			}
-					
-			if (MasterSubmixInstances[EMasterSubmixType::Reverb].IsValid() && !bIsMasterReverbBypassed)
+			NewMixerSubmix->SetParentSubmix(ParentSubmixInstance);
+			ParentSubmixInstance->AddChildSubmix(NewMixerSubmix);
+		}
+
+		NewMixerSubmix->Init(NewSubmix);
+	}
+
+	void FMixerDevice::InitSoundSubmixes()
+	{
+		if (!IsInAudioThread())
+		{
+			DECLARE_CYCLE_STAT(TEXT("FAudioThreadTask.InitSoundSubmixes"), STAT_InitSoundSubmixes, STATGROUP_AudioThreadCommands);
+
+			FAudioThread::RunCommandOnAudioThread([this]()
 			{
-				// Setup the master reverb only if we don't have a reverb plugin
-				USoundEffectSubmixPreset* ReverbPreset = nullptr;
-				USoundSubmix* MasterReverbSoundSubmix = FMixerDevice::MasterSubmixes[EMasterSubmixType::Reverb];
+				CSV_SCOPED_TIMING_STAT(Audio, InitSubmix);
+				InitSoundSubmixes();
+			}, GET_STATID(STAT_InitSoundSubmixes));
+			return;
+		}
 
-				if (GetDefault<UAudioSettings>()->bEnableLegacyReverb)
-				{
-					ReverbPreset = NewObject<USubmixEffectReverbPreset>(MasterReverbSoundSubmix, TEXT("Master Reverb Effect Preset"));
-				}
-				else
-				{
-					ReverbPreset = NewObject<USubmixEffectReverbFastPreset>(MasterReverbSoundSubmix, TEXT("Master Reverb Effect Fast Preset"));
-				}
-				
-				ReverbPreset->AddToRoot();
+		UAudioSettings* AudioSettings = GetMutableDefault<UAudioSettings>();
+		check(AudioSettings);
 
-				FSoundEffectSubmix* ReverbEffectSubmix = static_cast<FSoundEffectSubmix*>(ReverbPreset->CreateNewEffect());
+		if (MasterSubmixes.Num() > 0)
+		{
+			UE_LOG(LogAudioMixer, Display, TEXT("Re-initializing Sound Submixes..."));
+		}
+		else
+		{
+			UE_LOG(LogAudioMixer, Display, TEXT("Initializing Sound Submixes..."));
+		}
 
-				ReverbEffectSubmix->Init(InitData);
-				ReverbEffectSubmix->SetPreset(ReverbPreset);
-				ReverbEffectSubmix->SetEnabled(true);
+		FSoundEffectSubmixPtr ReverbPluginEffectSubmix;
+		USoundSubmix* ReverbPluginSubmix = nullptr;
 
-				const uint32 ReverbPresetId = ReverbPreset->GetUniqueID();
+		bSubmixRegistrationDisabled = true;
+		{
+			InitMasterSoundSubmix(EMasterSubmixType::Master, TEXT("MasterSubmixDefault"), false /* DefaultMuteWhenBackgrounded */, AudioSettings->MasterSubmix);
+			InitMasterSoundSubmix(EMasterSubmixType::Reverb, TEXT("MasterReverbSubmixDefault"), true /* DefaultMuteWhenBackgrounded */, AudioSettings->ReverbSubmix);
 
-				FMixerSubmixPtr MasterReverbSubmix = MasterSubmixInstances[EMasterSubmixType::Reverb];
-				MasterReverbSubmix->AddSoundEffectSubmix(ReverbPresetId, ReverbEffectSubmix);
-				MasterReverbSubmix->SetParentSubmix(MasterSubmixInstance);
-				MasterSubmixInstance->AddChildSubmix(MasterReverbSubmix);
+			if (!DisableSubmixEffectEQCvar)
+			{
+				InitMasterSoundSubmix(EMasterSubmixType::EQ, TEXT("MasterEQSubmixDefault"), false /* DefaultMuteWhenBackgrounded */, AudioSettings->EQSubmix);
 			}
 
-			// Setup the master EQ
-			if (FMixerDevice::MasterSubmixes[EMasterSubmixType::EQ])
+			InitMasterSoundSubmix(EMasterSubmixType::Ambisonics, TEXT("MasterAmbisonicSubmixDefault"), false /* DefaultMuteWhenBackgrounded */, AudioSettings->AmbisonicSubmix);
+
+			// Ensure submix UObject is loaded from reverb plugin interface asset if provided
+			if (IsReverbPluginEnabled() && ReverbPluginInterface)
 			{
-				USoundSubmix* MasterEQSoundSubmix = FMixerDevice::MasterSubmixes[EMasterSubmixType::EQ];
-				USubmixEffectSubmixEQPreset* EQPreset = NewObject<USubmixEffectSubmixEQPreset>(MasterEQSoundSubmix, TEXT("Master EQ Effect preset"));
-				EQPreset->AddToRoot();
-
-				FSoundEffectSubmix* EQEffectSubmix = static_cast<FSoundEffectSubmix*>(EQPreset->CreateNewEffect());
-				EQEffectSubmix->Init(InitData);
-				EQEffectSubmix->SetPreset(EQPreset);
-				EQEffectSubmix->SetEnabled(true);
-
-				const uint32 EQPresetId = EQPreset->GetUniqueID();
-
-				FMixerSubmixPtr MasterEQSubmix = MasterSubmixInstances[EMasterSubmixType::EQ];
-				MasterEQSubmix->AddSoundEffectSubmix(EQPresetId, EQEffectSubmix);
-				MasterEQSubmix->SetParentSubmix(MasterSubmixInstance);
-				MasterSubmixInstance->AddChildSubmix(MasterEQSubmix);
-
-				// Add the ambisonics master submix
-				FMixerSubmixPtr MasterAmbisonicsSubmix = MasterSubmixInstances[EMasterSubmixType::Ambisonics];
-				MasterAmbisonicsSubmix->SetParentSubmix(MasterEQSubmix);
-				MasterEQSubmix->AddChildSubmix(MasterAmbisonicsSubmix);
+				ReverbPluginSubmix = ReverbPluginInterface->GetSubmix();
+				check(ReverbPluginSubmix);
+				ReverbPluginSubmix->AddToRoot();
 			}
 		}
-
-		// Now register all the non-core submixes.
-
-#if WITH_EDITOR
-		Submixes.Reset();
-#endif
+		bSubmixRegistrationDisabled = false;
 
 		// Make sure all submixes are registered but not initialized
 		for (TObjectIterator<USoundSubmix> It; It; ++It)
@@ -743,28 +697,53 @@ namespace Audio
 			RegisterSoundSubmix(*It, false /* bInit */);
 		}
 
-		// Now setup the graph for all the submixes
-		for (auto& Entry : Submixes)
+		// Add the dynamically generated submix effect instance for reverb post registration but prior to other submix initialization
+		if (ReverbPluginSubmix)
 		{
-			USoundSubmix* SoundSubmix = Entry.Key;
-			FMixerSubmixPtr& SubmixInstance = Entry.Value;
-
-			// Setup up the submix instance's parent and add the submix instance as a child
-			FMixerSubmixWeakPtr ParentSubmixInstance;
-			if (SoundSubmix->ParentSubmix)
+			// Plugin must provide valid effect to enable reverb
+			ReverbPluginEffectSubmix = ReverbPluginInterface->GetEffectSubmix();
+			if (ReverbPluginEffectSubmix.IsValid())
 			{
-				ParentSubmixInstance = GetSubmixInstance(SoundSubmix->ParentSubmix);
+				// By now, submix should be registered and loaded
+				FMixerSubmixPtr ReverbPluginMixerSubmixPtr = GetSubmixInstance(ReverbPluginSubmix).Pin();
+				check(ReverbPluginMixerSubmixPtr.IsValid());
+				ReverbPluginMixerSubmixPtr->Init(ReverbPluginSubmix);
+
+				if (USoundEffectPreset* Preset = ReverbPluginEffectSubmix->GetPreset())
+				{
+					const uint32 ReverbPluginId = Preset->GetUniqueID();
+					FMixerSubmixWeakPtr ReverbPluginMixerSubmixWeakPtr = ReverbPluginMixerSubmixPtr;
+					AudioRenderThreadCommand([ReverbPluginMixerSubmixWeakPtr, ReverbPluginEffectSubmix, ReverbPluginId]()
+					{
+						FMixerSubmixPtr PluginSubmixPtr = ReverbPluginMixerSubmixWeakPtr.Pin();
+						if (PluginSubmixPtr.IsValid())
+						{
+							PluginSubmixPtr->AddSoundEffectSubmix(ReverbPluginId, ReverbPluginEffectSubmix);
+						}
+					});
+				}
 			}
 			else
 			{
-				ParentSubmixInstance = GetMasterSubmix();
-			}
-			FMixerSubmixPtr ParentSubmixInstancePtr = ParentSubmixInstance.Pin();
-			if (ParentSubmixInstancePtr.IsValid())
-			{
-				ParentSubmixInstancePtr->AddChildSubmix(SubmixInstance);
-				SubmixInstance->ParentSubmix = ParentSubmixInstance;
+				UE_LOG(LogAudioMixer, Error, TEXT("Reverb plugin failed to provide valid effect submix.  Plugin audio processing disabled."));
 			}
+		}
+
+		// Now setup the graph for all the submixes
+		FMixerSubmixPtr MasterSubmixPtr = GetMasterSubmix().Pin();
+		for (TPair<USoundSubmix*, FMixerSubmixPtr>& Entry : Submixes)
+		{
+			USoundSubmix* SoundSubmix = Entry.Key;
+			FMixerSubmixPtr& SubmixInstance = Entry.Value;
+
+			// Setup up the submix instance's parent and add the submix instance as a child
+			FMixerSubmixPtr ParentSubmixInstance = SoundSubmix->ParentSubmix
+				? GetSubmixInstance(SoundSubmix->ParentSubmix).Pin()
+				: MasterSubmixPtr;
+			check(ParentSubmixInstance.IsValid());
+
+			ParentSubmixInstance->AddChildSubmix(SubmixInstance);
+			SubmixInstance->SetParentSubmix(ParentSubmixInstance);
 
 			// Now add all the child submixes to this submix instance
 			for (USoundSubmix* ChildSubmix : SoundSubmix->ChildSubmixes)
@@ -774,10 +753,8 @@ namespace Audio
 				{
 					FChildSubmixInfo ChildSubmixInfo;
 					ChildSubmixInfo.SubmixPtr = GetSubmixInstance(ChildSubmix);
-					ChildSubmixInfo.bNeedsAmbisonicsEncoding = true;
-
-					FMixerSubmixPtr ChildSubmixInstancePtr = ChildSubmixInfo.SubmixPtr.Pin();
 
+					const FMixerSubmixPtr ChildSubmixInstancePtr = ChildSubmixInfo.SubmixPtr.Pin();
 					if (SubmixInstance.IsValid() && ChildSubmixInstancePtr.IsValid())
 					{
 						SubmixInstance->ChildSubmixes.Add(ChildSubmixInstancePtr->GetId(), ChildSubmixInfo);
@@ -786,8 +763,9 @@ namespace Audio
 				}
 			}
 
-			// Perform any other initialization on the submix instance
-			SubmixInstance->Init(SoundSubmix);
+			// Perform any other initialization on the submix instance.  Don't reinit for
+			// submixes already initialized directly (i.e. reverb plugins if already completed)
+			SubmixInstance->Init(SoundSubmix, false /* bAllowReInit */);
 		}
 	}
 	
@@ -811,11 +789,6 @@ namespace Audio
 		return MasterSubmixInstances[EMasterSubmixType::Master];
 	}
 
-	FMixerSubmixWeakPtr FMixerDevice::GetMasterReverbPluginSubmix()
-	{
-		return MasterSubmixInstances[EMasterSubmixType::ReverbPlugin];
-	}
-
 	FMixerSubmixWeakPtr FMixerDevice::GetMasterReverbSubmix()
 	{
 		return MasterSubmixInstances[EMasterSubmixType::Reverb];
@@ -831,7 +804,7 @@ namespace Audio
 		return MasterSubmixInstances[EMasterSubmixType::Ambisonics];
 	}
 
-	void FMixerDevice::AddMasterSubmixEffect(uint32 SubmixEffectId, FSoundEffectSubmix* SoundEffectSubmix)
+	void FMixerDevice::AddMasterSubmixEffect(uint32 SubmixEffectId, FSoundEffectSubmixPtr SoundEffectSubmix)
 	{
 		AudioRenderThreadCommand([this, SubmixEffectId, SoundEffectSubmix]()
 		{
@@ -882,7 +855,7 @@ namespace Audio
 #if WITH_EDITOR
 		check(IsInAudioThread());
 
-		const FMixerSubmixPtr MixerSubmix = Submixes.FindRef(InSoundSubmix);
+		FMixerSubmixPtr MixerSubmix = GetSubmixInstance(InSoundSubmix).Pin();
 		if (MixerSubmix.IsValid())
 		{
 			const float NewVolume = InSoundSubmix->OutputVolume;
@@ -909,7 +882,7 @@ namespace Audio
 			return;
 		}
 
-		const FMixerSubmixPtr MixerSubmix = Submixes.FindRef(InSoundSubmix);
+		FMixerSubmixPtr MixerSubmix = GetSubmixInstance(InSoundSubmix).Pin();
 		if (MixerSubmix.IsValid())
 		{
 			AudioRenderThreadCommand([MixerSubmix, NewVolume]()
@@ -972,7 +945,7 @@ namespace Audio
 	{
 		for (int32 i = 0; i < EMasterSubmixType::Count; ++i)
 		{
-			if (InSubmix == FMixerDevice::MasterSubmixes[i])
+			if (InSubmix == MasterSubmixes[i])
 			{
 				return true;
 			}
@@ -980,14 +953,14 @@ namespace Audio
 		return false;
 	}
 
-	FMixerSubmix* FMixerDevice::GetMasterSubmixInstance(USoundSubmix* InSubmix)
+	FMixerSubmixPtr FMixerDevice::GetMasterSubmixInstance(USoundSubmix* InSubmix)
 	{
 		check(MasterSubmixes.Num() == EMasterSubmixType::Count);
 		for (int32 i = 0; i < EMasterSubmixType::Count; ++i)
 		{
-			if (InSubmix == FMixerDevice::MasterSubmixes[i])
+			if (InSubmix == MasterSubmixes[i])
 			{
-				return MasterSubmixInstances[i].Get();
+				return MasterSubmixInstances[i];
 			}
 		}
 		return nullptr;
@@ -995,90 +968,89 @@ namespace Audio
 
 	void FMixerDevice::RegisterSoundSubmix(USoundSubmix* InSoundSubmix, bool bInit)
 	{
-		if (InSoundSubmix)
+		if (!InSoundSubmix || bSubmixRegistrationDisabled || IsMasterSubmixType(InSoundSubmix))
 		{
-			if (!IsInAudioThread())
-			{
-				DECLARE_CYCLE_STAT(TEXT("FAudioThreadTask.RegisterSoundSubmix"), STAT_AudioRegisterSoundSubmix, STATGROUP_AudioThreadCommands);
+			return;
+		}
 
-				FMixerDevice* MixerDevice = this;
-				FAudioThread::RunCommandOnAudioThread([MixerDevice, InSoundSubmix]()
-				{
-					CSV_SCOPED_TIMING_STAT(Audio, RegisterSubmix);
-					MixerDevice->RegisterSoundSubmix(InSoundSubmix);
-				}, GET_STATID(STAT_AudioRegisterSoundSubmix));
-				return;
-			}
+		if (!IsInAudioThread())
+		{
+			DECLARE_CYCLE_STAT(TEXT("FAudioThreadTask.RegisterSoundSubmix"), STAT_AudioRegisterSoundSubmix, STATGROUP_AudioThreadCommands);
 
-			if (!IsMasterSubmixType(InSoundSubmix))
+			FMixerDevice* MixerDevice = this;
+			FAudioThread::RunCommandOnAudioThread([MixerDevice, InSoundSubmix]()
 			{
-				FMixerSubmixPtr MixerSubmix; 
+				CSV_SCOPED_TIMING_STAT(Audio, RegisterSubmix);
+				MixerDevice->RegisterSoundSubmix(InSoundSubmix);
+			}, GET_STATID(STAT_AudioRegisterSoundSubmix));
+			return;
+		}
 
-				// If the sound submix wasn't already registered get it into the system.
-				if (!Submixes.Contains(InSoundSubmix))
-				{
-					MixerSubmix = FMixerSubmixPtr(new FMixerSubmix(this));
-					Submixes.Add(InSoundSubmix, MixerSubmix);
-				}
-				else
-				{
-					FMixerSubmixPtr* ExistingMixerSubmix = Submixes.Find(InSoundSubmix);
-					check(ExistingMixerSubmix);
-					MixerSubmix = *ExistingMixerSubmix;
-				}
+		// Ensure submix not already registered by first checking master submixes and then additional mixes.
+		FMixerSubmixPtr MixerSubmix = GetSubmixInstance(InSoundSubmix).Pin();
 
-				check(MixerSubmix.IsValid());
+		// If submix not already found, register it.
+		if (!MixerSubmix.IsValid())
+		{
+			MixerSubmix = MakeShared<FMixerSubmix, ESPMode::ThreadSafe>(this);
+			Submixes.Add(InSoundSubmix, MixerSubmix);
+		}
 
-				if (bInit)
-				{
-					// Setup the parent-child relationship
-					FMixerSubmixWeakPtr ParentSubmixInstance;
-					if (InSoundSubmix->ParentSubmix)
-					{
-						ParentSubmixInstance = GetSubmixInstance(InSoundSubmix->ParentSubmix);
-					}
-					else
-					{
-						ParentSubmixInstance = GetMasterSubmix();
-					}
+		if (bInit)
+		{
+			// Setup the parent-child relationship
+			FMixerSubmixPtr ParentSubmixInstance = InSoundSubmix->ParentSubmix
+				? GetSubmixInstance(InSoundSubmix->ParentSubmix).Pin()
+				: GetMasterSubmix().Pin();
 
-					FMixerSubmixPtr ParentSubmixInstancePtr = ParentSubmixInstance.Pin();
-					if (ParentSubmixInstancePtr.IsValid())
-					{
-						ParentSubmixInstancePtr->AddChildSubmix(MixerSubmix);
-						MixerSubmix->SetParentSubmix(ParentSubmixInstance);
-						MixerSubmix->Init(InSoundSubmix);
-					}
-				}
-			}
+			check(MixerSubmix.IsValid());
+			check(ParentSubmixInstance.IsValid());
+
+			ParentSubmixInstance->AddChildSubmix(MixerSubmix);
+			MixerSubmix->SetParentSubmix(ParentSubmixInstance);
+
+			MixerSubmix->Init(InSoundSubmix);
 		}
 	}
 
 	void FMixerDevice::UnregisterSoundSubmix(USoundSubmix* InSoundSubmix)
 	{
-		if (InSoundSubmix)
+		if (!InSoundSubmix || bSubmixRegistrationDisabled || IsMasterSubmixType(InSoundSubmix))
+		{
+			return;
+		}
+
+		if (!IsInAudioThread())
 		{
-			if (!IsInAudioThread())
+			DECLARE_CYCLE_STAT(TEXT("FAudioThreadTask.UnregisterSoundSubmix"), STAT_AudioUnregisterSoundSubmix, STATGROUP_AudioThreadCommands);
+
+			FAudioThread::RunCommandOnAudioThread([this, InSoundSubmix]()
 			{
-				DECLARE_CYCLE_STAT(TEXT("FAudioThreadTask.UnregisterSoundSubmix"), STAT_AudioUnregisterSoundSubmix, STATGROUP_AudioThreadCommands);
+				CSV_SCOPED_TIMING_STAT(Audio, UnregisterSubmix);
+				UnregisterSoundSubmix(InSoundSubmix);
+			}, GET_STATID(STAT_AudioUnregisterSoundSubmix));
+			return;
+		}
 
-				FMixerDevice* MixerDevice = this;
-				FAudioThread::RunCommandOnAudioThread([MixerDevice, InSoundSubmix]()
-				{
-					CSV_SCOPED_TIMING_STAT(Audio, UnregisterSubmix);
-					MixerDevice->UnregisterSoundSubmix(InSoundSubmix);
-				}, GET_STATID(STAT_AudioUnregisterSoundSubmix));
-				return;
-			}
+		FMixerSubmixPtr MasterSubmixPtr = GetMasterSubmix().Pin();
+		FMixerSubmixPtr ParentSubmixInstance = InSoundSubmix->ParentSubmix
+			? GetSubmixInstance(InSoundSubmix->ParentSubmix).Pin()
+			: MasterSubmixPtr;
+		if (ParentSubmixInstance.IsValid())
+		{
+			ParentSubmixInstance->RemoveChildSubmix(GetSubmixInstance(InSoundSubmix));
+		}
 
-			if (!IsMasterSubmixType(InSoundSubmix))
+		for (USoundSubmix* ChildSubmix : InSoundSubmix->ChildSubmixes)
+		{
+			FMixerSubmixPtr ChildSubmixPtr = GetSubmixInstance(ChildSubmix).Pin();
+			if (ChildSubmixPtr.IsValid())
 			{
-				if (InSoundSubmix)
-				{
-					Submixes.Remove(InSoundSubmix);
-				}
+				ChildSubmixPtr->SetParentSubmix(MasterSubmixPtr);
 			}
 		}
+
+		Submixes.Remove(InSoundSubmix);
 	}
 
 	void FMixerDevice::InitSoundEffectPresets()
@@ -1089,23 +1061,17 @@ namespace Audio
 #endif
 	}
 
-	FMixerSubmixWeakPtr FMixerDevice::GetSubmixInstance(USoundSubmix* SoundSubmix)
+	FMixerSubmixWeakPtr FMixerDevice::GetSubmixInstance(USoundSubmix* InSoundSubmix)
 	{
 		LLM_SCOPE(ELLMTag::AudioMixer);
 
-		check(SoundSubmix);
-		FMixerSubmixPtr* MixerSubmix = Submixes.Find(SoundSubmix);
-
-		// If the submix hasn't been registered yet, then register it now
-		if (!MixerSubmix)
+		FMixerSubmixPtr MixerSubmix = GetMasterSubmixInstance(InSoundSubmix);
+		if (MixerSubmix.IsValid())
 		{
-			RegisterSoundSubmix(SoundSubmix, true);
-			MixerSubmix = Submixes.Find(SoundSubmix);
+			return MixerSubmix;
 		}
 
-		// At this point, this should exist
-		check(MixerSubmix);
-		return *MixerSubmix;
+		return Submixes.FindRef(InSoundSubmix);
 	}
 
 	FMixerSourceVoice* FMixerDevice::GetMixerSourceVoice()
@@ -1310,10 +1276,10 @@ namespace Audio
 		}
 
 		// if we can find the submix here, record that submix. Otherwise, just record the master submix.
-		Audio::FMixerSubmixPtr* FoundSubmix = Submixes.Find(InSubmix);
-		if (FoundSubmix)
+		FMixerSubmixPtr FoundSubmix = GetSubmixInstance(InSubmix).Pin();
+		if (FoundSubmix.IsValid())
 		{
-			(*FoundSubmix)->OnStartRecordingOutput(ExpectedRecordingDuration);
+			FoundSubmix->OnStartRecordingOutput(ExpectedRecordingDuration);
 		}
 		else
 		{
@@ -1328,15 +1294,14 @@ namespace Audio
 	Audio::AlignedFloatBuffer& FMixerDevice::StopRecording(USoundSubmix* InSubmix, float& OutNumChannels, float& OutSampleRate)
 	{
 		// if we can find the submix here, record that submix. Otherwise, just record the master submix.
-		Audio::FMixerSubmixPtr* FoundSubmix = Submixes.Find(InSubmix);
-		if (FoundSubmix)
+		FMixerSubmixPtr FoundSubmix = GetSubmixInstance(InSubmix).Pin();
+		if (FoundSubmix.IsValid())
 		{
-			return (*FoundSubmix)->OnStopRecordingOutput(OutNumChannels, OutSampleRate);
+			return FoundSubmix->OnStopRecordingOutput(OutNumChannels, OutSampleRate);
 		}
 		else
 		{
-			FMixerSubmixWeakPtr MasterSubmix = GetMasterSubmix();
-			FMixerSubmixPtr MasterSubmixPtr = MasterSubmix.Pin();
+			FMixerSubmixPtr MasterSubmixPtr = GetMasterSubmix().Pin();
 			check(MasterSubmixPtr.IsValid());
 
 			return MasterSubmixPtr->OnStopRecordingOutput(OutNumChannels, OutSampleRate);
@@ -1358,10 +1323,10 @@ namespace Audio
 		}
 
 		// if we can find the submix here, pause that submix. Otherwise, just pause the master submix.
-		Audio::FMixerSubmixPtr* FoundSubmix = Submixes.Find(InSubmix);
-		if (FoundSubmix)
+		FMixerSubmixPtr FoundSubmix = GetSubmixInstance(InSubmix).Pin();
+		if (FoundSubmix.IsValid())
 		{
-			(*FoundSubmix)->PauseRecordingOutput();
+			FoundSubmix->PauseRecordingOutput();
 		}
 		else
 		{
@@ -1388,10 +1353,10 @@ namespace Audio
 		}
 
 		// if we can find the submix here, resume that submix. Otherwise, just resume the master submix.
-		Audio::FMixerSubmixPtr* FoundSubmix = Submixes.Find(InSubmix);
-		if (FoundSubmix)
+		FMixerSubmixPtr FoundSubmix = GetSubmixInstance(InSubmix).Pin();
+		if (FoundSubmix.IsValid())
 		{
-			(*FoundSubmix)->ResumeRecordingOutput();
+			FoundSubmix->ResumeRecordingOutput();
 		}
 		else
 		{
@@ -1418,10 +1383,10 @@ namespace Audio
 		}
 
 		// if we can find the submix here, record that submix. Otherwise, just record the master submix.
-		Audio::FMixerSubmixPtr* FoundSubmix = Submixes.Find(InSubmix);
-		if (FoundSubmix)
+		FMixerSubmixPtr FoundSubmix = GetSubmixInstance(InSubmix).Pin();
+		if (FoundSubmix.IsValid())
 		{
-			(*FoundSubmix)->StartEnvelopeFollowing(InSubmix->EnvelopeFollowerAttackTime, InSubmix->EnvelopeFollowerReleaseTime);
+			FoundSubmix->StartEnvelopeFollowing(InSubmix->EnvelopeFollowerAttackTime, InSubmix->EnvelopeFollowerReleaseTime);
 		}
 		else
 		{
@@ -1450,10 +1415,10 @@ namespace Audio
 		}
 
 		// if we can find the submix here, record that submix. Otherwise, just record the master submix.
-		Audio::FMixerSubmixPtr* FoundSubmix = Submixes.Find(InSubmix);
-		if (FoundSubmix)
+		FMixerSubmixPtr FoundSubmix = GetSubmixInstance(InSubmix).Pin();
+		if (FoundSubmix.IsValid())
 		{
-			(*FoundSubmix)->StopEnvelopeFollowing();
+			FoundSubmix->StopEnvelopeFollowing();
 		}
 		else
 		{
@@ -1482,10 +1447,10 @@ namespace Audio
 		}
 
 		// if we can find the submix here, record that submix. Otherwise, just record the master submix.
-		Audio::FMixerSubmixPtr* FoundSubmix = Submixes.Find(InSubmix);
-		if (FoundSubmix)
+		FMixerSubmixPtr FoundSubmix = GetSubmixInstance(InSubmix).Pin();
+		if (FoundSubmix.IsValid())
 		{
-			(*FoundSubmix)->AddEnvelopeFollowerDelegate(OnSubmixEnvelopeBP);
+			FoundSubmix->AddEnvelopeFollowerDelegate(OnSubmixEnvelopeBP);
 		}
 		else
 		{
@@ -1512,10 +1477,10 @@ namespace Audio
 			return;
 		}
 
-		Audio::FMixerSubmixPtr* FoundSubmix = Submixes.Find(InSubmix);
-		if (FoundSubmix)
+		FMixerSubmixPtr FoundSubmix = GetSubmixInstance(InSubmix).Pin();
+		if (FoundSubmix.IsValid())
 		{
-			(*FoundSubmix)->StartSpectrumAnalysis(InSettings);
+			FoundSubmix->StartSpectrumAnalysis(InSettings);
 		}
 		else
 		{
@@ -1541,10 +1506,10 @@ namespace Audio
 			return;
 		}
 
-		Audio::FMixerSubmixPtr* FoundSubmix = Submixes.Find(InSubmix);
-		if (FoundSubmix)
+		FMixerSubmixPtr FoundSubmix = GetSubmixInstance(InSubmix).Pin();
+		if (FoundSubmix.IsValid())
 		{
-			(*FoundSubmix)->StopSpectrumAnalysis();
+			FoundSubmix->StopSpectrumAnalysis();
 		}
 		else
 		{
@@ -1558,10 +1523,10 @@ namespace Audio
 
 	void FMixerDevice::GetMagnitudesForFrequencies(USoundSubmix* InSubmix, const TArray<float>& InFrequencies, TArray<float>& OutMagnitudes)
 	{
-		Audio::FMixerSubmixPtr* FoundSubmix = Submixes.Find(InSubmix);
-		if (FoundSubmix)
+		FMixerSubmixPtr FoundSubmix = GetSubmixInstance(InSubmix).Pin();
+		if (FoundSubmix.IsValid())
 		{
-			(*FoundSubmix)->GetMagnitudeForFrequencies(InFrequencies, OutMagnitudes);
+			FoundSubmix->GetMagnitudeForFrequencies(InFrequencies, OutMagnitudes);
 		}
 		else
 		{
@@ -1575,10 +1540,10 @@ namespace Audio
 
 	void FMixerDevice::GetPhasesForFrequencies(USoundSubmix* InSubmix, const TArray<float>& InFrequencies, TArray<float>& OutPhases)
 	{
-		Audio::FMixerSubmixPtr* FoundSubmix = Submixes.Find(InSubmix);
-		if (FoundSubmix)
+		FMixerSubmixPtr FoundSubmix = GetSubmixInstance(InSubmix).Pin();
+		if (FoundSubmix.IsValid())
 		{
-			(*FoundSubmix)->GetPhaseForFrequencies(InFrequencies, OutPhases);
+			FoundSubmix->GetPhaseForFrequencies(InFrequencies, OutPhases);
 		}
 		else
 		{
@@ -1614,10 +1579,10 @@ namespace Audio
 			return;
 		}
 
-		Audio::FMixerSubmixPtr* FoundSubmix = Submixes.Find(InSubmix);
-		if (FoundSubmix)
+		FMixerSubmixPtr FoundSubmix = GetSubmixInstance(InSubmix).Pin();
+		if (FoundSubmix.IsValid())
 		{
-			return (*FoundSubmix)->RegisterBufferListener(InSubmixBufferListener);
+			return FoundSubmix->RegisterBufferListener(InSubmixBufferListener);
 		}
 		else
 		{
@@ -1655,15 +1620,14 @@ namespace Audio
 			return;
 		}
 
-		Audio::FMixerSubmixPtr* FoundSubmix = Submixes.Find(InSubmix);
-		if (FoundSubmix)
+		FMixerSubmixPtr FoundSubmix = GetSubmixInstance(InSubmix).Pin();
+		if (FoundSubmix.IsValid())
 		{
-			return (*FoundSubmix)->UnregisterBufferListener(InSubmixBufferListener);
+			return FoundSubmix->UnregisterBufferListener(InSubmixBufferListener);
 		}
 		else
 		{
-			FMixerSubmixWeakPtr MasterSubmix = GetMasterSubmix();
-			FMixerSubmixPtr MasterSubmixPtr = MasterSubmix.Pin();
+			FMixerSubmixPtr MasterSubmixPtr = GetMasterSubmix().Pin();
 			check(MasterSubmixPtr.IsValid());
 
 			return MasterSubmixPtr->UnregisterBufferListener(InSubmixBufferListener);
diff --git a/Engine/Source/Runtime/AudioMixer/Private/AudioMixerEffectsManager.cpp b/Engine/Source/Runtime/AudioMixer/Private/AudioMixerEffectsManager.cpp
--- a/Engine/Source/Runtime/AudioMixer/Private/AudioMixerEffectsManager.cpp
+++ b/Engine/Source/Runtime/AudioMixer/Private/AudioMixerEffectsManager.cpp
@@ -93,12 +93,9 @@ namespace Audio
 #endif
 	}
 
-
-
 	FAudioMixerEffectsManager::FAudioMixerEffectsManager(FAudioDevice* InDevice)
 		: FAudioEffectsManager(InDevice)
 	{
-		bUseLegacyReverb = GetDefault<UAudioSettings>()->bEnableLegacyReverb;
 	}
 
 	FAudioMixerEffectsManager::~FAudioMixerEffectsManager()
@@ -113,20 +110,12 @@ namespace Audio
 		
 		if (MasterReverbSubmixPtr.IsValid())
 		{
-			FSoundEffectSubmix* SoundEffectSubmix = MasterReverbSubmixPtr->GetSubmixEffect(0);
-			if (SoundEffectSubmix)
+			FSoundEffectSubmixPtr SoundEffectSubmix = MasterReverbSubmixPtr->GetSubmixEffect(0);
+			if (SoundEffectSubmix.IsValid())
 			{
-				// Choose correct reverb based upon ini settings.
-				if (bUseLegacyReverb)
-				{
-					FSubmixEffectReverb* SoundEffectReverb = static_cast<FSubmixEffectReverb*>(SoundEffectSubmix);
-					SoundEffectReverb->SetEffectParameters(ReverbEffectParameters);
-				}
-				else
-				{
-					FSubmixEffectReverbFast* SoundEffectReverb = static_cast<FSubmixEffectReverbFast*>(SoundEffectSubmix);
-					SoundEffectReverb->SetEffectParameters(ReverbEffectParameters);
-				}
+				TSharedPtr<FSubmixEffectReverbFast, ESPMode::ThreadSafe> SoundEffectReverb = StaticCastSharedPtr<FSubmixEffectReverbFast, FSoundEffectSubmix, ESPMode::ThreadSafe>(SoundEffectSubmix);
+				SoundEffectReverb->SetEffectParameters(ReverbEffectParameters);
+
 				PrintReverbSettings(ReverbEffectParameters);
 			}
 		}
@@ -141,10 +130,10 @@ namespace Audio
 
 		if (MasterEQSubmixPtr.IsValid())
 		{
-			FSoundEffectSubmix* SoundEffectSubmix = MasterEQSubmixPtr->GetSubmixEffect(0);
-			if (SoundEffectSubmix)
+			FSoundEffectSubmixPtr SoundEffectSubmix = MasterEQSubmixPtr->GetSubmixEffect(0);
+			if (SoundEffectSubmix.IsValid())
 			{
-				FSubmixEffectSubmixEQ* SoundEffectEQ = static_cast<FSubmixEffectSubmixEQ*>(SoundEffectSubmix);
+				TSharedPtr<FSubmixEffectSubmixEQ, ESPMode::ThreadSafe> SoundEffectEQ = StaticCastSharedPtr<FSubmixEffectSubmixEQ, FSoundEffectSubmix, ESPMode::ThreadSafe>(SoundEffectSubmix);
 				SoundEffectEQ->SetEffectParameters(InEQEffectParameters);
 				PrintEQSettings(InEQEffectParameters);
 			}
diff --git a/Engine/Source/Runtime/AudioMixer/Private/AudioMixerSource.cpp b/Engine/Source/Runtime/AudioMixer/Private/AudioMixerSource.cpp
--- a/Engine/Source/Runtime/AudioMixer/Private/AudioMixerSource.cpp
+++ b/Engine/Source/Runtime/AudioMixer/Private/AudioMixerSource.cpp
@@ -90,7 +90,7 @@ namespace Audio
 		// Unfortunately, we need to know if this is a vorbis source since channel maps are different for 5.1 vorbis files
 		bIsVorbis = WaveInstance->WaveData->bDecompressedFromOgg;
 
-		bIsStoppingVoicesEnabled = ((FAudioDevice*)MixerDevice)->IsStoppingVoicesEnabled();
+		bIsStoppingVoicesEnabled = AudioDevice->IsStoppingVoicesEnabled();
 
 		bIsStopping = false;
 		bIsEffectTailsDone = true;
@@ -211,59 +211,18 @@ namespace Audio
 				// We'll only be using non-default submix sends (e.g. reverb).
 				if (!(InitParams.bUseHRTFSpatialization && MixerDevice->bSpatializationIsExternalSend))
 				{
-					// If this sound is an ambisonics file, we preempt the normal base submix routing and only send to master ambisonics submix
-					if (WaveInstance->bIsAmbisonics)
-					{
-						FMixerSourceSubmixSend SubmixSend;
-						SubmixSend.Submix = MixerDevice->GetMasterAmbisonicsSubmix();
-						SubmixSend.SendLevel = 1.0f;
-						SubmixSend.bIsMainSend = true;
-						InitParams.SubmixSends.Add(SubmixSend);
-					}
-					else
-					{
-						// If we've overridden which submix we're sending the sound, then add that as the first send
-						if (WaveInstance->SoundSubmix != nullptr)
-						{
-							FMixerSourceSubmixSend SubmixSend;
-							SubmixSend.Submix = MixerDevice->GetSubmixInstance(WaveInstance->SoundSubmix);
-							SubmixSend.SendLevel = 1.0f;
-							SubmixSend.bIsMainSend = true;
-							InitParams.SubmixSends.Add(SubmixSend);
-						}
-						else
-						{
-							// Send the voice to the EQ submix if it's enabled
-							const bool bIsEQDisabled = GetDefault<UAudioSettings>()->bDisableMasterEQ;
-							bool bUseMaster = true;
-							if (!bIsEQDisabled)
-							{
-								if (MixerDevice->GetMasterEQSubmix().IsValid())
-								{
-									// Default the submix to use to use the master submix if none are set
-									FMixerSourceSubmixSend SubmixSend;
-									SubmixSend.Submix = MixerDevice->GetMasterEQSubmix();
-									SubmixSend.SendLevel = 1.0f;
-									SubmixSend.bIsMainSend = true;
-									InitParams.SubmixSends.Add(SubmixSend);
-									bUseMaster = false;
-								}
-							}
-
-							if (bUseMaster)
-							{
-								// Default the submix to use to use the master submix if none are set
-								FMixerSourceSubmixSend SubmixSend;
-								SubmixSend.Submix = MixerDevice->GetMasterSubmix();
-								SubmixSend.SendLevel = 1.0f;
-								SubmixSend.bIsMainSend = true;
-								InitParams.SubmixSends.Add(SubmixSend);
-							}
-						}
-					}
+					FMixerSubmixWeakPtr SubmixPtr = WaveInstance->SoundSubmix
+						? MixerDevice->GetSubmixInstance(WaveInstance->SoundSubmix)
+						: MixerDevice->GetMasterSubmix();
+
+					FMixerSourceSubmixSend SubmixSend;
+					SubmixSend.Submix = SubmixPtr;
+					SubmixSend.SendLevel = 1.0f;
+					SubmixSend.bIsMainSend = true;
+					InitParams.SubmixSends.Add(SubmixSend);
 				}
 
-				// Now add any addition submix sends for this source
+				// Add submix sends for this source
 				for (FSoundSubmixSendInfo& SendInfo : WaveInstance->SoundSubmixSends)
 				{
 					if (SendInfo.SoundSubmix != nullptr)
@@ -898,7 +857,6 @@ namespace Audio
 
 			// 3. Apply editor gain stage(s)
 			CurrentVolume = FMath::Clamp<float>(GetDebugVolume(CurrentVolume), 0.0f, MAX_VOLUME);
-			MixerSourceVoice->SetVolume(CurrentVolume);
 		}
 		MixerSourceVoice->SetVolume(CurrentVolume);
 	}
@@ -957,22 +915,18 @@ namespace Audio
 			}
 
 			// Send the source audio to the reverb plugin if enabled
-			bool bPluginBypassedMasterReverb = false;
-			if (UseReverbPlugin())
+			if (UseReverbPlugin() && AudioDevice->ReverbPluginInterface)
 			{
-				if (MixerDevice->GetMasterReverbPluginSubmix().IsValid())
+				check(MixerDevice);
+				FMixerSubmixPtr ReverbPluginSubmixPtr = MixerDevice->GetSubmixInstance(AudioDevice->ReverbPluginInterface->GetSubmix()).Pin();
+				if (ReverbPluginSubmixPtr.IsValid())
 				{
-					MixerSourceVoice->SetSubmixSendInfo(MixerDevice->GetMasterReverbPluginSubmix(), ReverbSendLevel);
+					MixerSourceVoice->SetSubmixSendInfo(ReverbPluginSubmixPtr, ReverbSendLevel);
 				}
-
-				bPluginBypassedMasterReverb = AudioDevice->IsReverbPluginBypassingMasterReverb();
 			}
 
 			// Send the source audio to the master reverb
-			if (!bPluginBypassedMasterReverb && MixerDevice->GetMasterReverbSubmix().IsValid())
-			{
-				MixerSourceVoice->SetSubmixSendInfo(MixerDevice->GetMasterReverbSubmix(), ReverbSendLevel);
-			}
+			MixerSourceVoice->SetSubmixSendInfo(MixerDevice->GetMasterReverbSubmix(), ReverbSendLevel);
 		}
 
 		// Update submix send levels
diff --git a/Engine/Source/Runtime/AudioMixer/Private/AudioMixerSubmix.cpp b/Engine/Source/Runtime/AudioMixer/Private/AudioMixerSubmix.cpp
--- a/Engine/Source/Runtime/AudioMixer/Private/AudioMixerSubmix.cpp
+++ b/Engine/Source/Runtime/AudioMixer/Private/AudioMixerSubmix.cpp
@@ -66,7 +66,7 @@ namespace Audio
 		}
 	}
 
-	void FMixerSubmix::Init(USoundSubmix* InSoundSubmix)
+	void FMixerSubmix::Init(USoundSubmix* InSoundSubmix, bool bAllowReInit)
 	{
 		check(IsInAudioThread());
 		if (InSoundSubmix != nullptr)
@@ -77,9 +77,9 @@ namespace Audio
 				OwningSubmixObject = InSoundSubmix;
 				InitInternal();
 			}
-			else
+			// This is a re-init and needs to be thread safe
+			else if (bAllowReInit)
 			{
-				// This is a re-init and needs to be thread safe
 				check(OwningSubmixObject == InSoundSubmix);
 				SubmixCommand([this]()
 				{
@@ -120,10 +120,10 @@ namespace Audio
 
 				FSubmixEffectInfo EffectInfo;
 				EffectInfo.PresetId = EffectPreset->GetUniqueID();
-				EffectInfo.EffectInstance = SubmixEffect;
+				EffectInfo.EffectInstance = MakeShareable(SubmixEffect);
 
 				// Add the effect to this submix's chain
-				EffectSubmixChain.Add(EffectInfo);
+				EffectSubmixChain.Add(MoveTemp(EffectInfo));
 			}
 		}
 
@@ -137,14 +137,15 @@ namespace Audio
 			//If we do have a valid ambisonics decoder, lets use it. Otherwise, treat this submix like a device submix.
 			if (AmbisonicsMixer.IsValid())
 			{
-				if (!OwningSubmixObject->AmbisonicsPluginSettings)
+				UAmbisonicsSubmixSettingsBase* AmbisonicSubmixSettings = OwningSubmixObject->AmbisonicsPluginSettings;
+				if (!AmbisonicSubmixSettings)
 				{
-					OwningSubmixObject->AmbisonicsPluginSettings = AmbisonicsMixer->GetDefaultSettings();
+					AmbisonicSubmixSettings = AmbisonicsMixer->GetDefaultSettings();
 				}
 
-				if (OwningSubmixObject->AmbisonicsPluginSettings != nullptr)
+				if (AmbisonicSubmixSettings)
 				{
-					OnAmbisonicsSettingsChanged(OwningSubmixObject->AmbisonicsPluginSettings);
+					OnAmbisonicsSettingsChanged(AmbisonicSubmixSettings);
 				}
 				else
 				{
@@ -211,6 +212,28 @@ namespace Audio
 		});
 	}
 
+	void FMixerSubmix::RemoveChildSubmix(TWeakPtr<FMixerSubmix, ESPMode::ThreadSafe> SubmixWeakPtr)
+	{
+		TSharedPtr<FMixerSubmix, ESPMode::ThreadSafe> SubmixStrongPtr  = SubmixWeakPtr.Pin();
+		if (!SubmixStrongPtr.IsValid())
+		{
+			return;
+		}
+
+		const uint32 OldIdToRemove = SubmixStrongPtr->GetId();
+		SubmixCommand([this, OldIdToRemove]()
+		{
+			AUDIO_MIXER_CHECK_AUDIO_PLAT_THREAD(MixerDevice);
+
+			ChildSubmixes.Remove(OldIdToRemove);
+
+			if (ChannelFormat == ESubmixChannelFormat::Ambisonics && AmbisonicsMixer.IsValid())
+			{
+				UpdateAmbisonicsEncoderForChildren();
+			}
+		});
+	}
+
 	ESubmixChannelFormat FMixerSubmix::GetSubmixChannels() const
 	{
 		return ChannelFormat;
@@ -283,7 +306,7 @@ namespace Audio
 		}
 	}
 
-	void FMixerSubmix::AddSoundEffectSubmix(uint32 SubmixPresetId, FSoundEffectSubmix* InSoundEffectSubmix)
+	void FMixerSubmix::AddSoundEffectSubmix(uint32 SubmixPresetId, FSoundEffectSubmixPtr InSoundEffectSubmix)
 	{
 		AUDIO_MIXER_CHECK_AUDIO_PLAT_THREAD(MixerDevice);
 
@@ -292,7 +315,7 @@ namespace Audio
 		{
 			if (EffectSubmixChain[i].PresetId == SubmixPresetId)
 			{
-				// Already added. 
+				// Already added.
 				return;
 			}
 		}
@@ -309,18 +332,14 @@ namespace Audio
 	{
 		AUDIO_MIXER_CHECK_AUDIO_PLAT_THREAD(MixerDevice);
 
-		for (int32 i = 0; i < EffectSubmixChain.Num(); ++i)
+		for (FSubmixEffectInfo& Effect : EffectSubmixChain)
 		{
 			// If the ID's match, delete and remove the effect instance but don't modify the effect submix chain array itself
-			if (EffectSubmixChain[i].PresetId == SubmixPresetId)
+			if (Effect.PresetId == SubmixPresetId)
 			{
-				// Delete the reference to the effect instance
-				if (EffectSubmixChain[i].EffectInstance)
-				{
-					delete EffectSubmixChain[i].EffectInstance;
-					EffectSubmixChain[i].EffectInstance = nullptr;
-				}
-				EffectSubmixChain[i].PresetId = INDEX_NONE;
+				// Reset reference to the effect
+				Effect.EffectInstance.Reset();
+				Effect.PresetId = INDEX_NONE;
 				return;
 			}
 		}
@@ -331,12 +350,10 @@ namespace Audio
 	{
 		for (FSubmixEffectInfo& Info : EffectSubmixChain)
 		{
-			if (Info.EffectInstance)
+			if (Info.EffectInstance.IsValid())
 			{
 				Info.EffectInstance->ClearPreset();
-
-				delete Info.EffectInstance;
-				Info.EffectInstance = nullptr;
+				Info.EffectInstance.Reset();
 			}
 		}
 
@@ -500,12 +517,12 @@ namespace Audio
 
 			TSharedPtr<Audio::FMixerSubmix, ESPMode::ThreadSafe> SubmixPtr = ChildSubmix.SubmixPtr.Pin();
 
-			//Check to see if this child is an ambisonics submix.
+			// Check to see if this child is an ambisonics submix.
 			if (SubmixPtr.IsValid() && SubmixPtr->GetSubmixChannels() == ESubmixChannelFormat::Ambisonics)
 			{
 				UAmbisonicsSubmixSettingsBase* ChildAmbisonicsSettings = SubmixPtr->AmbisonicsSettings;
 
-				//Check if this child submix needs to be reencoded.
+				// Check if this child submix needs to be re-encoded.
 				if (!ChildAmbisonicsSettings || AmbisonicsMixer->ShouldReencodeBetween(ChildAmbisonicsSettings, AmbisonicsSettings))
 				{
 					ChildSubmix.bNeedsAmbisonicsEncoding = false;
@@ -697,10 +714,10 @@ namespace Audio
 
 			for (FSubmixEffectInfo& SubmixEffectInfo : EffectSubmixChain)
 			{
-				FSoundEffectSubmix* SubmixEffect = SubmixEffectInfo.EffectInstance;
+				FSoundEffectSubmixPtr SubmixEffect = SubmixEffectInfo.EffectInstance;
 
 				// SubmixEffectInfo.EffectInstance will be null if FMixerSubmix::RemoveSoundEffectSubmix was called earlier.
-				if (!SubmixEffect)
+				if (!SubmixEffect.IsValid())
 				{
 					continue;
 				}
@@ -858,7 +875,7 @@ namespace Audio
 		return EffectSubmixChain.Num();
 	}
 
-	FSoundEffectSubmix* FMixerSubmix::GetSubmixEffect(const int32 InIndex)
+	FSoundEffectSubmixPtr FMixerSubmix::GetSubmixEffect(const int32 InIndex)
 	{
 		if (InIndex < EffectSubmixChain.Num())
 		{
diff --git a/Engine/Source/Runtime/AudioMixer/Private/Effects/AudioMixerSubmixEffectReverbFast.cpp b/Engine/Source/Runtime/AudioMixer/Private/Effects/AudioMixerSubmixEffectReverbFast.cpp
--- a/Engine/Source/Runtime/AudioMixer/Private/Effects/AudioMixerSubmixEffectReverbFast.cpp
+++ b/Engine/Source/Runtime/AudioMixer/Private/Effects/AudioMixerSubmixEffectReverbFast.cpp
@@ -38,6 +38,7 @@ static FAutoConsoleVariableRef CVarDisableQuadReverbCVarFast(
 class UReverbEffect;
 
 FSubmixEffectReverbFast::FSubmixEffectReverbFast()
+	: bBypass(false)
 {
 }
 
@@ -109,6 +110,7 @@ void FSubmixEffectReverbFast::OnPresetChanged()
 	ReverbEffect.RoomRolloffFactor = 0.0f; // not used
 	ReverbEffect.Volume = Settings.WetLevel;
 
+	bBypass = Settings.bBypass;
 
 	SetEffectParameters(ReverbEffect);
 }
@@ -118,12 +120,17 @@ void FSubmixEffectReverbFast::OnProcessAudio(const FSoundEffectSubmixInputData&
 	LLM_SCOPE(ELLMTag::AudioMixer);
 
 	check(InData.NumChannels == 2);
- 	if (OutData.NumChannels < 2 || DisableSubmixReverbCVarFast == 1) 
+ 	if (OutData.NumChannels < 2 || DisableSubmixReverbCVarFast == 1)
 	{
 		// Not supported
 		return;
 	}
 
+	if (bBypass)
+	{
+		return;
+	}
+
 	UpdateParameters();
 
 	PlateReverb->ProcessAudio(*InData.AudioBuffer, InData.NumChannels, *OutData.AudioBuffer, OutData.NumChannels);
diff --git a/Engine/Source/Runtime/AudioMixer/Public/AudioMixerDevice.h b/Engine/Source/Runtime/AudioMixer/Public/AudioMixerDevice.h
--- a/Engine/Source/Runtime/AudioMixer/Public/AudioMixerDevice.h
+++ b/Engine/Source/Runtime/AudioMixer/Public/AudioMixerDevice.h
@@ -11,7 +11,7 @@
 // Forward Declarations
 class FOnSubmixEnvelopeBP;
 class IAudioMixerPlatformInterface;
-
+class USoundSubmix;
 
 namespace Audio
 {
@@ -63,7 +63,6 @@ namespace Audio
 		{
 			Master,
 			Reverb,
-			ReverbPlugin,
 			EQ,
 			Ambisonics,
 			Count,
@@ -102,7 +101,6 @@ namespace Audio
 		virtual void ResumeContext() override;
 		virtual void SuspendContext() override;
 		virtual void EnableDebugAudioOutput() override;
-		virtual void InitSoundSubmixes() override;
 		virtual FAudioPlatformSettings GetPlatformSettings() const override;
 		virtual void RegisterSoundSubmix(USoundSubmix* SoundSubmix, bool bInit = true) override;
 		virtual void UnregisterSoundSubmix(USoundSubmix* SoundSubmix) override;
@@ -153,7 +151,7 @@ namespace Audio
 		virtual void OnAudioStreamShutdown() override;
 		//~ End IAudioMixer
 
-		FMixerSubmixWeakPtr GetSubmixInstance(USoundSubmix* SoundSubmix);
+		FMixerSubmixWeakPtr GetSubmixInstance(USoundSubmix* InSoundSubmix);
 
 		// Functions which check the thread it's called on and helps make sure functions are called from correct threads
 		void CheckAudioThread() const;
@@ -185,12 +183,11 @@ namespace Audio
 
 		FMixerSubmixWeakPtr GetMasterSubmix(); 
 		FMixerSubmixWeakPtr GetMasterReverbSubmix();
-		FMixerSubmixWeakPtr GetMasterReverbPluginSubmix();
 		FMixerSubmixWeakPtr GetMasterEQSubmix();
 		FMixerSubmixWeakPtr GetMasterAmbisonicsSubmix();
 
 		// Add submix effect to master submix
-		void AddMasterSubmixEffect(uint32 SubmixEffectId, FSoundEffectSubmix* SoundEffect);
+		void AddMasterSubmixEffect(uint32 SubmixEffectId, FSoundEffectSubmixPtr SoundEffect);
 		
 		// Remove submix effect from master submix
 		void RemoveMasterSubmixEffect(uint32 SubmixEffectId);
@@ -217,6 +214,8 @@ namespace Audio
 
 	protected:
 
+		virtual void InitSoundSubmixes() override;
+
 		virtual void OnListenerUpdated(const TArray<FListener>& InListeners) override;
 
 		TArray<FTransform> ListenerTransforms;
@@ -227,6 +226,7 @@ namespace Audio
 
 		void Get2DChannelMapInternal(const int32 NumSourceChannels, const int32 NumOutputChannels, const bool bIsCenterChannelOnly, TArray<float>& OutChannelMap) const;
 		void InitializeChannelMaps();
+		void InitMasterSoundSubmix(EMasterSubmixType::Type InType, const FString& InDefaultName, bool bInDefaultMuteWhenBackgrounded, FSoftObjectPath& InOutObjectPath);
 		static int32 GetChannelMapCacheId(const int32 NumSourceChannels, const int32 NumOutputChannels, const bool bIsCenterChannelOnly);
 		void CacheChannelMap(const int32 NumSourceChannels, const int32 NumOutputChannels, const bool bIsCenterChannelOnly);
 		void InitializeChannelAzimuthMap(const int32 NumChannels);
@@ -239,7 +239,7 @@ namespace Audio
 	private:
 
 		bool IsMasterSubmixType(USoundSubmix* InSubmix) const;
-		FMixerSubmix* GetMasterSubmixInstance(USoundSubmix* InSubmix);
+		FMixerSubmixPtr GetMasterSubmixInstance(USoundSubmix* InSubmix);
 
 		// Pushes the command to a audio render thread command queue to be executed on render thread
 		void AudioRenderThreadCommand(TFunction<void()> Command);
@@ -247,7 +247,7 @@ namespace Audio
 		// Pumps the audio render thread command queue
 		void PumpCommandQueue();
 
-		static TArray<USoundSubmix*> MasterSubmixes;
+		TArray<USoundSubmix*> MasterSubmixes;
 		TArray<FMixerSubmixPtr> MasterSubmixInstances;
 
 		/** Ptr to the platform interface, which handles streaming audio to the hardware device. */
@@ -299,7 +299,7 @@ namespace Audio
 		/** The mixer source manager. */
 		FMixerSourceManager SourceManager;
 
-		/** ThreadId for the game thread (or if audio is running a seperate thread, that ID) */
+		/** ThreadId for the game thread (or if audio is running a separate thread, that ID) */
 		mutable int32 GameOrAudioThreadId;
 
 		/** ThreadId for the low-level platform audio mixer. */
@@ -310,6 +310,9 @@ namespace Audio
 
 		/** Whether or not we generate output audio to test multi-platform mixer. */
 		bool bDebugOutputEnabled;
+
+		/** Whether or not initialization of the submix system is underway and submixes can be registered */
+		bool bSubmixRegistrationDisabled;
 	};
 }
 
diff --git a/Engine/Source/Runtime/AudioMixer/Public/AudioMixerSubmix.h b/Engine/Source/Runtime/AudioMixer/Public/AudioMixerSubmix.h
--- a/Engine/Source/Runtime/AudioMixer/Public/AudioMixerSubmix.h
+++ b/Engine/Source/Runtime/AudioMixer/Public/AudioMixerSubmix.h
@@ -7,6 +7,7 @@
 #include "Sound/SoundSubmix.h"
 #include "DSP/EnvelopeFollower.h"
 #include "DSP/SpectrumAnalyzer.h"
+#include "Templates/SharedPointer.h"
 
 
 // Forward Declarations
@@ -54,7 +55,7 @@ namespace Audio
 		virtual ~FMixerSubmix();
 
 		// Initialize the submix object with the USoundSubmix ptr. Sets up child and parent connects.
-		void Init(USoundSubmix* InSoundSubmix);
+		void Init(USoundSubmix* InSoundSubmix, bool bAllowReInit = true);
 
 		// Returns the mixer submix Id
 		uint32 GetId() const { return Id; }
@@ -65,6 +66,9 @@ namespace Audio
 		// Adds the given submix to this submix's children
 		void AddChildSubmix(TWeakPtr<FMixerSubmix, ESPMode::ThreadSafe> Submix);
 
+		// Removes the given submix from this submix's children
+		void RemoveChildSubmix(TWeakPtr<FMixerSubmix, ESPMode::ThreadSafe> SubmixWeakPtr);
+
 		// Sets the static output volume of the submix
 		void SetOutputVolume(float InVolume);
 
@@ -90,7 +94,7 @@ namespace Audio
 		void RemoveSourceVoice(FMixerSourceVoice* InSourceVoice);
 
 		/** Appends the effect submix to the effect submix chain. */
-		void AddSoundEffectSubmix(uint32 SubmixPresetId, FSoundEffectSubmix* InSoundEffectSubmix);
+		void AddSoundEffectSubmix(uint32 SubmixPresetId, FSoundEffectSubmixPtr InSoundEffectSubmix);
 
 		/** Removes the submix effect from the effect submix chain. */
 		void RemoveSoundEffectSubmix(uint32 SubmixPresetId);
@@ -114,7 +118,7 @@ namespace Audio
 		int32 GetNumChainEffects() const;
 
 		// Returns the submix effect at the given effect chain index
-		FSoundEffectSubmix* GetSubmixEffect(const int32 InIndex);
+		FSoundEffectSubmixPtr GetSubmixEffect(const int32 InIndex);
 
 		// updates settings, potentially creating or removing ambisonics streams based on
 		void OnAmbisonicsSettingsChanged(UAmbisonicsSubmixSettingsBase* AmbisonicsSettings);
@@ -227,11 +231,10 @@ namespace Audio
 			uint32 PresetId;
 
 			// The effect instance ptr
-			FSoundEffectSubmix* EffectInstance;
+			FSoundEffectSubmixPtr EffectInstance;
 
 			FSubmixEffectInfo()
 				: PresetId(INDEX_NONE)
-				, EffectInstance(nullptr)
 			{
 			}
 		};
diff --git a/Engine/Source/Runtime/Engine/Classes/Sound/AudioSettings.h b/Engine/Source/Runtime/Engine/Classes/Sound/AudioSettings.h
--- a/Engine/Source/Runtime/Engine/Classes/Sound/AudioSettings.h
+++ b/Engine/Source/Runtime/Engine/Classes/Sound/AudioSettings.h
@@ -80,8 +80,6 @@ class ENGINE_API UAudioSettings : public UDeveloperSettings
 	DECLARE_EVENT(UAudioSettings, FAudioSettingsChanged)
 #endif // WITH_EDITOR
 
-	virtual void Serialize(FArchive& Ar) override;
-
 	/** The SoundClass assigned to newly created sounds */
 	UPROPERTY(config, EditAnywhere, Category="Audio", meta=(AllowedClasses="SoundClass", DisplayName="Default Sound Class"))
 	FSoftObjectPath DefaultSoundClassName;
@@ -102,6 +100,22 @@ class ENGINE_API UAudioSettings : public UDeveloperSettings
 	UPROPERTY(config, EditAnywhere, Category="Audio", meta=(AllowedClasses="SoundClass", DisplayName = "VOIP Sound Class"))
 	FSoftObjectPath VoiPSoundClass;
 
+	/** The Master submix through which all sounds are routed */
+	UPROPERTY(config, EditAnywhere, Category="Mix", meta=(AllowedClasses="SoundSubmix"))
+	FSoftObjectPath MasterSubmix;
+
+	/** The submix through which all sounds set to use reverb are routed */
+	UPROPERTY(config, EditAnywhere, Category="Mix", meta=(AllowedClasses="SoundSubmix"))
+	FSoftObjectPath ReverbSubmix;
+
+	/** The submix through which all sounds set to use legacy EQ system are routed */
+	UPROPERTY(config, EditAnywhere, Category="Mix", meta=(AllowedClasses="SoundSubmix", DisplayName = "EQ Submix (Legacy)"), AdvancedDisplay)
+	FSoftObjectPath EQSubmix;
+
+	/** The submix through which sounds spatialized using ambisonics are routed */
+	UPROPERTY(config, EditAnywhere, Category="Mix", meta=(AllowedClasses="SoundSubmix"))
+	FSoftObjectPath AmbisonicSubmix;
+
 	/** Sample rate used for voice over IP. VOIP audio is resampled to the application's sample rate on the receiver side. */
 	UPROPERTY(config, EditAnywhere, Category = "Audio", meta = (DisplayName = "VOIP Sample Rate"))
 	EVoiceSampleRate VoiPSampleRate;
@@ -110,10 +124,6 @@ class ENGINE_API UAudioSettings : public UDeveloperSettings
 	UPROPERTY(config)
 	float DefaultReverbSendLevel_DEPRECATED;
 
-	/** Enables legacy version of reverb. The legacy reverb runs more slowly, but by most other measures is functionally equivalent. It has a slight perceptual difference. */
-	UPROPERTY(config, EditAnywhere, Category = "Audio", AdvancedDisplay)
-	uint32 bEnableLegacyReverb: 1;
-
 	/** How many streaming sounds can be played at the same time (if more are played they will be sorted by priority) */
 	UPROPERTY(config, EditAnywhere, Category="Audio", meta=(ClampMin=0))
 	int32 MaximumConcurrentStreams;
@@ -203,6 +213,7 @@ public:
 private:
 #if WITH_EDITOR
 	TArray<FAudioQualitySettings> CachedQualityLevels;
+	FSoftObjectPath CachedAmbisonicSubmix;
 #endif // WITH_EDITOR
 
 	void AddDefaultSettings();
diff --git a/Engine/Source/Runtime/Engine/Classes/Sound/SoundEffectSubmix.h b/Engine/Source/Runtime/Engine/Classes/Sound/SoundEffectSubmix.h
--- a/Engine/Source/Runtime/Engine/Classes/Sound/SoundEffectSubmix.h
+++ b/Engine/Source/Runtime/Engine/Classes/Sound/SoundEffectSubmix.h
@@ -97,3 +97,4 @@ public:
 	/** Processes audio in the submix effect. */
 	void ProcessAudio(FSoundEffectSubmixInputData& InData, FSoundEffectSubmixOutputData& OutData);
 };
+typedef TSharedPtr<FSoundEffectSubmix, ESPMode::ThreadSafe> FSoundEffectSubmixPtr;
diff --git a/Engine/Source/Runtime/Engine/Private/AudioDevice.cpp b/Engine/Source/Runtime/Engine/Private/AudioDevice.cpp
--- a/Engine/Source/Runtime/Engine/Private/AudioDevice.cpp
+++ b/Engine/Source/Runtime/Engine/Private/AudioDevice.cpp
@@ -32,6 +32,7 @@
 #include "Sound/SoundCue.h"
 #include "Sound/SoundNode.h"
 #include "Sound/SoundNodeWavePlayer.h"
+#include "Sound/SoundSubmix.h"
 #include "UnrealEngine.h"
 #include "UObject/UObjectHash.h"
 #include "UObject/UObjectIterator.h"
@@ -281,7 +282,6 @@ FAudioDevice::FAudioDevice()
 	, bSpatializationInterfaceEnabled(false)
 	, bOcclusionInterfaceEnabled(false)
 	, bReverbInterfaceEnabled(false)
-	, bReverbPluginBypassesMasterReverb(false)
 	, bModulationInterfaceEnabled(false)
 	, bPluginListenersInitialized(false)
 	, bHRTFEnabledForAll(false)
@@ -441,7 +441,6 @@ bool FAudioDevice::Init(int32 InMaxSources)
 	{
 		ReverbPluginInterface = ReverbPluginFactory->CreateNewReverbPlugin(this);
 		bReverbInterfaceEnabled = true;
-		bReverbPluginBypassesMasterReverb = ReverbPluginInterface->DoesReverbOverrideMasterReverb();
 		bReverbIsExternalSend = ReverbPluginFactory->IsExternalSend();
 		UE_LOG(LogAudio, Log, TEXT("Audio Reverb Plugin: %s"), *(ReverbPluginFactory->GetDisplayName()));
 	}
@@ -4016,6 +4015,7 @@ void FAudioDevice::StartSources(TArray<FWaveInstance*>& WaveInstances, int32 Fir
 					WaveInstance->StopWithoutNotification();
 					Source->WaveInstance = nullptr;
 					FreeSources.Add(Source);
+					WaveInstanceSourceMap.Remove(WaveInstance);
 				}
 			}
 			else if (Source)
diff --git a/Engine/Source/Runtime/Engine/Private/AudioSettings.cpp b/Engine/Source/Runtime/Engine/Private/AudioSettings.cpp
--- a/Engine/Source/Runtime/Engine/Private/AudioSettings.cpp
+++ b/Engine/Source/Runtime/Engine/Private/AudioSettings.cpp
@@ -6,16 +6,22 @@
 
 #include "Sound/AudioSettings.h"
 #include "Misc/ConfigCacheIni.h"
+#include "Misc/Paths.h"
 #include "Sound/SoundSubmix.h"
 #include "Sound/SoundBase.h"
 #include "Sound/SoundNodeQualityLevel.h"
 #include "UObject/UObjectHash.h"
 #include "UObject/UObjectIterator.h"
 
+#if WITH_EDITOR
+#include "Framework/Notifications/NotificationManager.h"
+#include "Widgets/Notifications/SNotificationList.h"
+#endif // WITH_EDITOR
+
 #define LOCTEXT_NAMESPACE "AudioSettings"
 
 UAudioSettings::UAudioSettings(const FObjectInitializer& ObjectInitializer)
-: Super(ObjectInitializer)
+	: Super(ObjectInitializer)
 {
 	SectionName = TEXT("Audio");
 	AddDefaultSettings();
@@ -34,7 +40,6 @@ void UAudioSettings::AddDefaultSettings()
 	QualityLevels.Add(DefaultSettings);
 	bAllowPlayWhenSilent = true;
 	DefaultReverbSendLevel_DEPRECATED = 0.0f;
-	bEnableLegacyReverb = false;
 	VoiPSampleRate = EVoiceSampleRate::Low16000Hz;
 	NumStoppingSources = 8;
 }
@@ -42,6 +47,9 @@ void UAudioSettings::AddDefaultSettings()
 #if WITH_EDITOR
 void UAudioSettings::PreEditChange(FProperty* PropertyAboutToChange)
 {
+	// Cache ambisonic submix in case user tries to set to submix that isn't set to ambisonics
+	CachedAmbisonicSubmix = AmbisonicSubmix;
+
 	// Cache at least the first entry in case someone tries to clear the array
 	CachedQualityLevels = QualityLevels;
 }
@@ -50,9 +58,37 @@ void UAudioSettings::PostEditChangeChainProperty(FPropertyChangedChainEvent& Pro
 {
 	if (PropertyChangedEvent.Property)
 	{
-		bool bReconcileNodes = false;
-
-		if (PropertyChangedEvent.Property->GetFName() == GET_MEMBER_NAME_CHECKED(UAudioSettings, QualityLevels))
+		bool bReconcileQualityNodes = false;
+		bool bPromptRestartRequired = false;
+		FName PropertyName = PropertyChangedEvent.Property->GetFName();
+		if (PropertyName == GET_MEMBER_NAME_CHECKED(UAudioSettings, MasterSubmix)
+			|| PropertyName == GET_MEMBER_NAME_CHECKED(UAudioSettings, EQSubmix)
+			|| PropertyName == GET_MEMBER_NAME_CHECKED(UAudioSettings, ReverbSubmix))
+		{
+			bPromptRestartRequired = true;
+		}
+		else if (PropertyName == GET_MEMBER_NAME_CHECKED(UAudioSettings, AmbisonicSubmix))
+		{
+			if (USoundSubmix* NewSubmix = Cast<USoundSubmix>(AmbisonicSubmix.TryLoad()))
+			{
+				if (NewSubmix->ChannelFormat != ESubmixChannelFormat::Ambisonics)
+				{
+					FNotificationInfo Info(LOCTEXT("AudioSettings_InvalidAmbisonicSubmixFormat",
+						"Ambisonic Submix format must be set to 'Ambisonics' in order to be set as 'Master Ambisonics Submix'."));
+					Info.bFireAndForget = true;
+					Info.ExpireDuration = 2.0f;
+					Info.bUseThrobber = true;
+					FSlateNotificationManager::Get().AddNotification(Info);
+
+					AmbisonicSubmix = CachedAmbisonicSubmix;
+				}
+			}
+			else
+			{
+				bPromptRestartRequired = true;
+			}
+		}
+		else if (PropertyName == GET_MEMBER_NAME_CHECKED(UAudioSettings, QualityLevels))
 		{
 			if (QualityLevels.Num() == 0)
 			{
@@ -86,14 +122,14 @@ void UAudioSettings::PostEditChangeChainProperty(FPropertyChangedChainEvent& Pro
 				}
 			}
 
-			bReconcileNodes = true;
+			bReconcileQualityNodes = true;
 		}
-		else if (PropertyChangedEvent.Property->GetFName() == GET_MEMBER_NAME_CHECKED(FAudioQualitySettings, DisplayName))
+		else if (PropertyName == GET_MEMBER_NAME_CHECKED(FAudioQualitySettings, DisplayName))
 		{
-			bReconcileNodes = true;
+			bReconcileQualityNodes = true;
 		}
 
-		if (bReconcileNodes)
+		if (bReconcileQualityNodes)
 		{
 			for (TObjectIterator<USoundNodeQualityLevel> It; It; ++It)
 			{
@@ -101,15 +137,20 @@ void UAudioSettings::PostEditChangeChainProperty(FPropertyChangedChainEvent& Pro
 			}
 		}
 
+		if (bPromptRestartRequired)
+		{
+			FNotificationInfo Info(LOCTEXT("AudioSettings_ChangeRequiresEditorRestart",
+				"Change to Audio Settings requires editor restart in order for changes to take effect."));
+			Info.bFireAndForget = true;
+			Info.ExpireDuration = 2.0f;
+			Info.bUseThrobber = true;
+			FSlateNotificationManager::Get().AddNotification(Info);
+		}
+
 		AudioSettingsChanged.Broadcast();
 	}
 }
-#endif
-
-void UAudioSettings::Serialize(FArchive& Ar)
-{
-	Super::Serialize(Ar);
-}
+#endif // WITH_EDITOR
 
 const FAudioQualitySettings& UAudioSettings::GetQualityLevelSettings(int32 QualityLevel) const
 {
diff --git a/Engine/Source/Runtime/Engine/Private/SoundWave.cpp b/Engine/Source/Runtime/Engine/Private/SoundWave.cpp
--- a/Engine/Source/Runtime/Engine/Private/SoundWave.cpp
+++ b/Engine/Source/Runtime/Engine/Private/SoundWave.cpp
@@ -1663,6 +1663,8 @@ void USoundWave::FinishDestroy()
 
 void USoundWave::Parse(FAudioDevice* AudioDevice, const UPTRINT NodeWaveInstanceHash, FActiveSound& ActiveSound, const FSoundParseParameters& ParseParams, TArray<FWaveInstance*>& WaveInstances)
 {
+	check(AudioDevice);
+
 	FWaveInstance* WaveInstance = ActiveSound.FindWaveInstance(NodeWaveInstanceHash);
 
 	const bool bIsNewWave = WaveInstance == nullptr;
@@ -1753,7 +1755,12 @@ void USoundWave::Parse(FAudioDevice* AudioDevice, const UPTRINT NodeWaveInstance
 		WaveInstance->bReverb = ActiveSound.bReverb || SoundClassProperties->bReverb;
 		WaveInstance->OutputTarget = SoundClassProperties->OutputTarget;
 
-		if (SoundClassProperties->DefaultSubmix)
+		if (SoundClassProperties->bApplyEffects)
+		{
+			const UAudioSettings* Settings = GetDefault<UAudioSettings>();
+			WaveInstance->SoundSubmix = Cast<USoundSubmix>(Settings->EQSubmix.TryLoad());
+		}
+		else if (SoundClassProperties->DefaultSubmix)
 		{
 			WaveInstance->SoundSubmix = SoundClassProperties->DefaultSubmix;
 		}
@@ -1775,7 +1782,6 @@ void USoundWave::Parse(FAudioDevice* AudioDevice, const UPTRINT NodeWaveInstance
 		WaveInstance->RadioFilterVolumeThreshold = 0.f;
 		WaveInstance->StereoBleed = 0.f;
 		WaveInstance->LFEBleed = 0.f;
-		WaveInstance->SoundSubmix = ActiveSound.GetSoundSubmix();
 		WaveInstance->bIsUISound = ActiveSound.bIsUISound;
 		WaveInstance->bIsMusic = ActiveSound.bIsMusic;
 		WaveInstance->bReverb = ActiveSound.bReverb;
@@ -1785,10 +1791,19 @@ void USoundWave::Parse(FAudioDevice* AudioDevice, const UPTRINT NodeWaveInstance
 	}
 
 	WaveInstance->bIsAmbisonics = bIsAmbisonics;
-	if (ParseParams.SoundSubmix)
+	if (bIsAmbisonics)
+	{
+		const UAudioSettings* Settings = GetDefault<UAudioSettings>();
+		WaveInstance->SoundSubmix = Cast<USoundSubmix>(Settings->AmbisonicSubmix.TryLoad());
+	}
+	else if (ParseParams.SoundSubmix)
 	{
 		WaveInstance->SoundSubmix = ParseParams.SoundSubmix;
 	}
+	else if (USoundSubmix* WaveSubmix = GetSoundSubmix())
+	{
+		WaveInstance->SoundSubmix = WaveSubmix;
+	}
 
 	// If set to bAlwaysPlay, increase the current sound's priority scale by 10x. This will still result in a possible 0-priority output if the sound has 0 actual volume
 	if (bAlwaysPlay)
diff --git a/Engine/Source/Runtime/Engine/Public/AudioDevice.h b/Engine/Source/Runtime/Engine/Public/AudioDevice.h
--- a/Engine/Source/Runtime/Engine/Public/AudioDevice.h
+++ b/Engine/Source/Runtime/Engine/Public/AudioDevice.h
@@ -1210,12 +1210,6 @@ public:
 		return bReverbInterfaceEnabled;
 	}
 
-	/** Whether or not the reverb plugin is bypassing the master reverb. */
-	bool IsReverbPluginBypassingMasterReverb() const
-	{
-		return IsReverbPluginEnabled() && bReverbPluginBypassesMasterReverb;
-	}
-
 	static bool IsReverbPluginLoaded()
 	{
 		if (FAudioDevice* MainAudioDevice = GEngine->GetMainAudioDevice())
@@ -1868,7 +1862,6 @@ private:
 	uint8 bSpatializationInterfaceEnabled:1;
 	uint8 bOcclusionInterfaceEnabled:1;
 	uint8 bReverbInterfaceEnabled:1;
-	uint8 bReverbPluginBypassesMasterReverb:1;
 	uint8 bModulationInterfaceEnabled:1;
 
 	/** Whether or not we've initialized plugin listeners array. */
diff --git a/Engine/Source/Runtime/Engine/Public/IAudioExtensionPlugin.h b/Engine/Source/Runtime/Engine/Public/IAudioExtensionPlugin.h
--- a/Engine/Source/Runtime/Engine/Public/IAudioExtensionPlugin.h
+++ b/Engine/Source/Runtime/Engine/Public/IAudioExtensionPlugin.h
@@ -9,6 +9,7 @@
 #include "Math/Interval.h"
 #include "Modules/ModuleInterface.h"
 #include "Modules/ModuleManager.h"
+#include "Sound/SoundEffectSubmix.h"
 #include "UObject/ObjectMacros.h"
 
 #if !UE_BUILD_SHIPPING
@@ -22,7 +23,6 @@
 class FAudioDevice;
 class USoundSubmix;
 
-
 /**
 * Enumeration of audio plugin types
 *
@@ -774,15 +774,16 @@ public:
 	/** Called when a source is done playing and is released. */
 	virtual void OnReleaseSource(const uint32 SourceId) = 0;
 
-	virtual class FSoundEffectSubmix* GetEffectSubmix(USoundSubmix* Submix) = 0;
+	/** Returns the plugin-managed effect submix instance */
+	virtual FSoundEffectSubmixPtr GetEffectSubmix() = 0;
+
+	/** Returns the plugin-managed effect submix */
+	virtual USoundSubmix* GetSubmix() = 0;
 
 	/** Processes audio with the given input and output data structs.*/
 	virtual void ProcessSourceAudio(const FAudioPluginSourceInputData& InputData, FAudioPluginSourceOutputData& OutputData)
 	{
 	}
-
-	/** Returns whether or not the plugin reverb overrides the master reverb. If true, then the built in reverb will be uninitialized and bypassed. */
-	virtual bool DoesReverbOverrideMasterReverb() const { return true; }
 };
 
 /************************************************************************/
