- Fix infinite recursion crash in editor when adding submix and disallowing linking submix as master's parent
- Fix crash when removing a child submix from parent in wrong submix command cue (causing a modification of children while looping through parent ProcessAudio call)
- Clean-up weakptrs in submix instances that become invalid
- Ensure submix soft object paths are loaded on game thread

diff --git a/Engine/Plugins/Runtime/ResonanceAudio/Source/ResonanceAudio/Private/ResonanceAudioReverb.cpp b/Engine/Plugins/Runtime/ResonanceAudio/Source/ResonanceAudio/Private/ResonanceAudioReverb.cpp
--- a/Engine/Plugins/Runtime/ResonanceAudio/Source/ResonanceAudio/Private/ResonanceAudioReverb.cpp
+++ b/Engine/Plugins/Runtime/ResonanceAudio/Source/ResonanceAudio/Private/ResonanceAudioReverb.cpp
@@ -108,12 +108,9 @@ namespace ResonanceAudio
 		}
 	}
 
-	FSoundEffectSubmixPtr FResonanceAudioReverb::GetEffectSubmix()
+	void FResonanceAudioReverb::InitEffectSubmix()
 	{
-		if (SubmixEffect.IsValid())
-		{
-			return SubmixEffect;
-		}
+		check(!SubmixEffect.IsValid());
 
 		// Load the global reverb preset settings:
 		const FSoftObjectPath ReverbPluginPresetName = GetDefault<UResonanceAudioSettings>()->GlobalReverbPreset;
@@ -144,6 +141,14 @@ namespace ResonanceAudio
 		{
 			ReverbPluginPreset = nullptr;
 		}
+	}
+
+	FSoundEffectSubmixPtr FResonanceAudioReverb::GetEffectSubmix()
+	{
+		if (!SubmixEffect.IsValid())
+		{
+			InitEffectSubmix();
+		}
 
 		return SubmixEffect;
 	}
@@ -165,6 +170,12 @@ namespace ResonanceAudio
 			ReverbSubmix->bMuteWhenBackgrounded = true;
 		}
 
+		// SubmixEffect is required to be initialized for ReverbPluginPreset to be set
+		if (!SubmixEffect.IsValid())
+		{
+			InitEffectSubmix();
+		}
+
 		if (ReverbPluginPreset)
 		{
 			bool bFoundPreset = false;
diff --git a/Engine/Plugins/Runtime/ResonanceAudio/Source/ResonanceAudio/Private/ResonanceAudioReverb.h b/Engine/Plugins/Runtime/ResonanceAudio/Source/ResonanceAudio/Private/ResonanceAudioReverb.h
--- a/Engine/Plugins/Runtime/ResonanceAudio/Source/ResonanceAudio/Private/ResonanceAudioReverb.h
+++ b/Engine/Plugins/Runtime/ResonanceAudio/Source/ResonanceAudio/Private/ResonanceAudioReverb.h
@@ -134,6 +134,9 @@ namespace ResonanceAudio
 		void SetReverbBrightness();
 
 	private:
+
+		void InitEffectSubmix();
+
 		FResonanceAudioReverbPluginSettings ReverbSettings;
 		RaRoomProperties RoomProperties;
 		RaReverbProperties ReverbProperties;
diff --git a/Engine/Plugins/Runtime/Steam/SteamAudio/Source/SteamAudio/Private/PhononReverb.cpp b/Engine/Plugins/Runtime/Steam/SteamAudio/Source/SteamAudio/Private/PhononReverb.cpp
--- a/Engine/Plugins/Runtime/Steam/SteamAudio/Source/SteamAudio/Private/PhononReverb.cpp
+++ b/Engine/Plugins/Runtime/Steam/SteamAudio/Source/SteamAudio/Private/PhononReverb.cpp
@@ -395,7 +395,7 @@ namespace SteamAudio
 	{
 		if (!SubmixEffect.IsValid())
 		{
-			USubmixEffectReverbPluginPreset* ReverbPluginPreset = NewObject<USubmixEffectReverbPluginPreset>(GetSubmix(), TEXT("Master Reverb Plugin Effect Preset"));
+			USubmixEffectReverbPluginPreset* ReverbPluginPreset = NewObject<USubmixEffectReverbPluginPreset>(GetSubmix(), TEXT("Reverb Plugin Effect Preset"));
 			FSubmixEffectReverbPlugin* ReverbSubmixEffect = static_cast<FSubmixEffectReverbPlugin*>(ReverbPluginPreset->CreateNewEffect());
 			ReverbSubmixEffect->SetPhononReverbPlugin(this);
 
diff --git a/Engine/Source/Editor/AudioEditor/Private/SoundSubmixGraphSchema.cpp b/Engine/Source/Editor/AudioEditor/Private/SoundSubmixGraphSchema.cpp
--- a/Engine/Source/Editor/AudioEditor/Private/SoundSubmixGraphSchema.cpp
+++ b/Engine/Source/Editor/AudioEditor/Private/SoundSubmixGraphSchema.cpp
@@ -11,6 +11,7 @@
 #include "ToolMenus.h"
 
 #include "ScopedTransaction.h"
+#include "Sound/AudioSettings.h"
 #include "Sound/SoundSubmix.h"
 #include "SoundSubmixGraph/SoundSubmixGraphNode.h"
 #include "SoundSubmixGraph/SoundSubmixGraph.h"
@@ -105,6 +106,23 @@ bool USoundSubmixGraphSchema::ConnectionCausesLoop(const UEdGraphPin* InputPin,
 	USoundSubmixGraphNode* InputNode = CastChecked<USoundSubmixGraphNode>(InputPin->GetOwningNode());
 	USoundSubmixGraphNode* OutputNode = CastChecked<USoundSubmixGraphNode>(OutputPin->GetOwningNode());
 
+	// Master Submix cannot be an input as it would create an inferred loop for submixes without an explicit parent
+	if (const UAudioSettings* Settings = GetDefault<UAudioSettings>())
+	{
+		if (USoundSubmix* MasterSubmix = Cast<USoundSubmix>(Settings->MasterSubmix.TryLoad()))
+		{
+			if (OutputNode->SoundSubmix == MasterSubmix)
+			{
+				return true;
+			}
+
+			if (OutputNode->SoundSubmix->RecurseCheckChild(MasterSubmix))
+			{
+				return true;
+			}
+		}
+	}
+
 	return OutputNode->SoundSubmix->RecurseCheckChild(InputNode->SoundSubmix);
 }
 
diff --git a/Engine/Source/Runtime/AudioMixer/Private/AudioMixerDevice.cpp b/Engine/Source/Runtime/AudioMixer/Private/AudioMixerDevice.cpp
--- a/Engine/Source/Runtime/AudioMixer/Private/AudioMixerDevice.cpp
+++ b/Engine/Source/Runtime/AudioMixer/Private/AudioMixerDevice.cpp
@@ -551,9 +551,9 @@ namespace Audio
 		SourceManager.UpdatePendingReleaseData(true);
 	}
 
-	void FMixerDevice::InitMasterSoundSubmix(EMasterSubmixType::Type InType, const FString& InDefaultName, bool bInDefaultMuteWhenBackgrounded, FSoftObjectPath& InObjectPath)
+	void FMixerDevice::LoadMasterSoundSubmix(EMasterSubmixType::Type InType, const FString& InDefaultName, bool bInDefaultMuteWhenBackgrounded, FSoftObjectPath& InObjectPath)
 	{
-		check(IsInAudioThread());
+		check(IsInGameThread());
 
 		const int32 MasterSubmixCount = static_cast<int32>(EMasterSubmixType::Type::Count);
 		if(MasterSubmixes.Num() < MasterSubmixCount)
@@ -627,90 +627,30 @@ namespace Audio
 		FMixerSubmixPtr NewMixerSubmix = MakeShared<FMixerSubmix, ESPMode::ThreadSafe>(this);
 		MasterSubmixInstances[TypeIndex] = NewMixerSubmix;
 
-		if (InType != EMasterSubmixType::Master)
-		{
-			FMixerSubmixPtr ParentSubmixInstance = NewSubmix->ParentSubmix
-				? GetSubmixInstance(NewSubmix->ParentSubmix).Pin()
-				: GetMasterSubmix().Pin();
-
-			NewMixerSubmix->SetParentSubmix(ParentSubmixInstance);
-			ParentSubmixInstance->AddChildSubmix(NewMixerSubmix);
-		}
-
-		NewMixerSubmix->Init(NewSubmix);
+		NewMixerSubmix->Init(NewSubmix, false /* bAllowReInit */);
 	}
 
-	void FMixerDevice::InitSoundSubmixes()
+	void FMixerDevice::LoadPluginSoundSubmixes()
 	{
-		if (!IsInAudioThread())
-		{
-			DECLARE_CYCLE_STAT(TEXT("FAudioThreadTask.InitSoundSubmixes"), STAT_InitSoundSubmixes, STATGROUP_AudioThreadCommands);
-
-			FAudioThread::RunCommandOnAudioThread([this]()
-			{
-				CSV_SCOPED_TIMING_STAT(Audio, InitSubmix);
-				InitSoundSubmixes();
-			}, GET_STATID(STAT_InitSoundSubmixes));
-			return;
-		}
-
-		UAudioSettings* AudioSettings = GetMutableDefault<UAudioSettings>();
-		check(AudioSettings);
-
-		if (MasterSubmixes.Num() > 0)
-		{
-			UE_LOG(LogAudioMixer, Display, TEXT("Re-initializing Sound Submixes..."));
-		}
-		else
-		{
-			UE_LOG(LogAudioMixer, Display, TEXT("Initializing Sound Submixes..."));
-		}
-
-		FSoundEffectSubmixPtr ReverbPluginEffectSubmix;
-		USoundSubmix* ReverbPluginSubmix = nullptr;
+		check(IsInGameThread());
 
-		bSubmixRegistrationDisabled = true;
+		if (IsReverbPluginEnabled() && ReverbPluginInterface)
 		{
-			InitMasterSoundSubmix(EMasterSubmixType::Master, TEXT("MasterSubmixDefault"), false /* DefaultMuteWhenBackgrounded */, AudioSettings->MasterSubmix);
-			InitMasterSoundSubmix(EMasterSubmixType::Reverb, TEXT("MasterReverbSubmixDefault"), true /* DefaultMuteWhenBackgrounded */, AudioSettings->ReverbSubmix);
-
-			if (!DisableSubmixEffectEQCvar)
-			{
-				InitMasterSoundSubmix(EMasterSubmixType::EQ, TEXT("MasterEQSubmixDefault"), false /* DefaultMuteWhenBackgrounded */, AudioSettings->EQSubmix);
-			}
-
-			InitMasterSoundSubmix(EMasterSubmixType::Ambisonics, TEXT("MasterAmbisonicSubmixDefault"), false /* DefaultMuteWhenBackgrounded */, AudioSettings->AmbisonicSubmix);
+			USoundSubmix* ReverbPluginSubmix = ReverbPluginInterface->GetSubmix();
+			check(ReverbPluginSubmix);
+			ReverbPluginSubmix->AddToRoot();
 
-			// Ensure submix UObject is loaded from reverb plugin interface asset if provided
-			if (IsReverbPluginEnabled() && ReverbPluginInterface)
-			{
-				ReverbPluginSubmix = ReverbPluginInterface->GetSubmix();
-				check(ReverbPluginSubmix);
-				ReverbPluginSubmix->AddToRoot();
-			}
-		}
-		bSubmixRegistrationDisabled = false;
-
-		// Make sure all submixes are registered but not initialized
-		for (TObjectIterator<USoundSubmix> It; It; ++It)
-		{
-			RegisterSoundSubmix(*It, false /* bInit */);
-		}
+			LoadSoundSubmix(*ReverbPluginSubmix);
 
-		// Add the dynamically generated submix effect instance for reverb post registration but prior to other submix initialization
-		if (ReverbPluginSubmix)
-		{
 			// Plugin must provide valid effect to enable reverb
-			ReverbPluginEffectSubmix = ReverbPluginInterface->GetEffectSubmix();
+			FSoundEffectSubmixPtr ReverbPluginEffectSubmix = ReverbPluginInterface->GetEffectSubmix();
 			if (ReverbPluginEffectSubmix.IsValid())
 			{
-				// By now, submix should be registered and loaded
-				FMixerSubmixPtr ReverbPluginMixerSubmixPtr = GetSubmixInstance(ReverbPluginSubmix).Pin();
-				check(ReverbPluginMixerSubmixPtr.IsValid());
-				ReverbPluginMixerSubmixPtr->Init(ReverbPluginSubmix);
-
 				if (USoundEffectPreset* Preset = ReverbPluginEffectSubmix->GetPreset())
 				{
+					FMixerSubmixPtr ReverbPluginMixerSubmixPtr = GetSubmixInstance(ReverbPluginSubmix).Pin();
+					check(ReverbPluginMixerSubmixPtr.IsValid());
+
 					const uint32 ReverbPluginId = Preset->GetUniqueID();
 					FMixerSubmixWeakPtr ReverbPluginMixerSubmixWeakPtr = ReverbPluginMixerSubmixPtr;
 					AudioRenderThreadCommand([ReverbPluginMixerSubmixWeakPtr, ReverbPluginEffectSubmix, ReverbPluginId]()
@@ -728,47 +668,97 @@ namespace Audio
 				UE_LOG(LogAudioMixer, Error, TEXT("Reverb plugin failed to provide valid effect submix.  Plugin audio processing disabled."));
 			}
 		}
+	}
 
-		// Now setup the graph for all the submixes
-		FMixerSubmixPtr MasterSubmixPtr = GetMasterSubmix().Pin();
-		for (TPair<USoundSubmix*, FMixerSubmixPtr>& Entry : Submixes)
+	void FMixerDevice::InitSoundSubmixes()
+	{
+		if (IsInGameThread())
 		{
-			USoundSubmix* SoundSubmix = Entry.Key;
-			FMixerSubmixPtr& SubmixInstance = Entry.Value;
+			bSubmixRegistrationDisabled = true;
 
-			// Setup up the submix instance's parent and add the submix instance as a child
-			FMixerSubmixPtr ParentSubmixInstance = SoundSubmix->ParentSubmix
-				? GetSubmixInstance(SoundSubmix->ParentSubmix).Pin()
-				: MasterSubmixPtr;
-			check(ParentSubmixInstance.IsValid());
+			UAudioSettings* AudioSettings = GetMutableDefault<UAudioSettings>();
+			check(AudioSettings);
+
+			if (MasterSubmixes.Num() > 0)
+			{
+				UE_LOG(LogAudioMixer, Display, TEXT("Re-initializing Sound Submixes..."));
+			}
+			else
+			{
+				UE_LOG(LogAudioMixer, Display, TEXT("Initializing Sound Submixes..."));
+			}
 
-			ParentSubmixInstance->AddChildSubmix(SubmixInstance);
-			SubmixInstance->SetParentSubmix(ParentSubmixInstance);
+			// 1. Load or reload all sound submixes/instances
+			LoadMasterSoundSubmix(EMasterSubmixType::Master, TEXT("MasterSubmixDefault"), false /* DefaultMuteWhenBackgrounded */, AudioSettings->MasterSubmix);
+			LoadMasterSoundSubmix(EMasterSubmixType::Reverb, TEXT("MasterReverbSubmixDefault"), true /* DefaultMuteWhenBackgrounded */, AudioSettings->ReverbSubmix);
 
-			// Now add all the child submixes to this submix instance
-			for (USoundSubmix* ChildSubmix : SoundSubmix->ChildSubmixes)
+			if (!DisableSubmixEffectEQCvar)
 			{
-				// ChildSubmix lists can contain null entries.
-				if (ChildSubmix)
-				{
-					FChildSubmixInfo ChildSubmixInfo;
-					ChildSubmixInfo.SubmixPtr = GetSubmixInstance(ChildSubmix);
+				LoadMasterSoundSubmix(EMasterSubmixType::EQ, TEXT("MasterEQSubmixDefault"), false /* DefaultMuteWhenBackgrounded */, AudioSettings->EQSubmix);
+			}
 
-					const FMixerSubmixPtr ChildSubmixInstancePtr = ChildSubmixInfo.SubmixPtr.Pin();
-					if (SubmixInstance.IsValid() && ChildSubmixInstancePtr.IsValid())
-					{
-						SubmixInstance->ChildSubmixes.Add(ChildSubmixInstancePtr->GetId(), ChildSubmixInfo);
-					}
+			LoadMasterSoundSubmix(EMasterSubmixType::Ambisonics, TEXT("MasterAmbisonicSubmixDefault"), false /* DefaultMuteWhenBackgrounded */, AudioSettings->AmbisonicSubmix);
+			LoadPluginSoundSubmixes();
+
+			for (TObjectIterator<USoundSubmix> It; It; ++It)
+			{
+				USoundSubmix* SubmixToLoad = *It;
+				check(SubmixToLoad);
 
+				if (!IsMasterSubmixType(SubmixToLoad))
+				{
+					LoadSoundSubmix(*SubmixToLoad);
 				}
 			}
+			bSubmixRegistrationDisabled = false;
+		}
+
+		if (!IsInAudioThread())
+		{
+			DECLARE_CYCLE_STAT(TEXT("FAudioThreadTask.InitSoundSubmixes"), STAT_InitSoundSubmixes, STATGROUP_AudioThreadCommands);
+
+			FAudioThread::RunCommandOnAudioThread([this]()
+			{
+				CSV_SCOPED_TIMING_STAT(Audio, InitSubmix);
+				InitSoundSubmixes();
+			}, GET_STATID(STAT_InitSoundSubmixes));
+			return;
+		}
+
+		for (int32 i = 0; i < static_cast<int32>(EMasterSubmixType::Count); ++i)
+		{
+			USoundSubmix* SoundSubmix = MasterSubmixes[i];
+			check(SoundSubmix);
+			FMixerSubmixPtr& MasterSubmixInstance = MasterSubmixInstances[i];
+
+			if (SoundSubmix != MasterSubmixes[static_cast<int32>(EMasterSubmixType::Master)])
+			{
+				RebuildSubmixLinks(*SoundSubmix, MasterSubmixInstance);
+			}
+		};
 
-			// Perform any other initialization on the submix instance.  Don't reinit for
-			// submixes already initialized directly (i.e. reverb plugins if already completed)
-			SubmixInstance->Init(SoundSubmix, false /* bAllowReInit */);
+		for (TPair<USoundSubmix*, FMixerSubmixPtr>& Entry : Submixes)
+		{
+			USoundSubmix* SoundSubmix = Entry.Key;
+			check(SoundSubmix);
+			FMixerSubmixPtr& SubmixInstance = Entry.Value;
+			RebuildSubmixLinks(*SoundSubmix, SubmixInstance);
 		}
 	}
-	
+
+	void FMixerDevice::RebuildSubmixLinks(USoundSubmix& SoundSubmix, FMixerSubmixPtr& SubmixInstance)
+	{
+		// Setup up the submix instance's parent and add the submix instance as a child
+		FMixerSubmixPtr ParentSubmixInstance = SoundSubmix.ParentSubmix
+			? GetSubmixInstance(SoundSubmix.ParentSubmix).Pin()
+			: GetMasterSubmix().Pin();
+
+		check(ParentSubmixInstance.IsValid());
+
+		SubmixInstance->SetParentSubmix(ParentSubmixInstance);
+		ParentSubmixInstance->AddChildSubmix(SubmixInstance);
+	}
+
  	FAudioPlatformSettings FMixerDevice::GetPlatformSettings() const
  	{
 		FAudioPlatformSettings Settings = AudioMixerPlatform->GetPlatformSettings();
@@ -986,31 +976,29 @@ namespace Audio
 			return;
 		}
 
+		LoadSoundSubmix(*InSoundSubmix);
+		FMixerSubmixPtr SubmixPtr = GetSubmixInstance(InSoundSubmix).Pin();
+		if (bInit)
+		{
+			SubmixPtr->Init(InSoundSubmix, true /* BAllowReInit */);
+		}
+
+		RebuildSubmixLinks(*InSoundSubmix, SubmixPtr);
+	}
+
+	void FMixerDevice::LoadSoundSubmix(USoundSubmix& InSoundSubmix)
+	{
 		// Ensure submix not already registered by first checking master submixes and then additional mixes.
-		FMixerSubmixPtr MixerSubmix = GetSubmixInstance(InSoundSubmix).Pin();
+		FMixerSubmixPtr MixerSubmix = GetSubmixInstance(&InSoundSubmix).Pin();
 
 		// If submix not already found, register it.
 		if (!MixerSubmix.IsValid())
 		{
 			MixerSubmix = MakeShared<FMixerSubmix, ESPMode::ThreadSafe>(this);
-			Submixes.Add(InSoundSubmix, MixerSubmix);
+			Submixes.Add(&InSoundSubmix, MixerSubmix);
 		}
 
-		if (bInit)
-		{
-			// Setup the parent-child relationship
-			FMixerSubmixPtr ParentSubmixInstance = InSoundSubmix->ParentSubmix
-				? GetSubmixInstance(InSoundSubmix->ParentSubmix).Pin()
-				: GetMasterSubmix().Pin();
-
-			check(MixerSubmix.IsValid());
-			check(ParentSubmixInstance.IsValid());
-
-			ParentSubmixInstance->AddChildSubmix(MixerSubmix);
-			MixerSubmix->SetParentSubmix(ParentSubmixInstance);
-
-			MixerSubmix->Init(InSoundSubmix);
-		}
+		MixerSubmix->Init(&InSoundSubmix, false /* bAllowReInit */);
 	}
 
 	void FMixerDevice::UnregisterSoundSubmix(USoundSubmix* InSoundSubmix)
@@ -1032,25 +1020,35 @@ namespace Audio
 			return;
 		}
 
-		FMixerSubmixPtr MasterSubmixPtr = GetMasterSubmix().Pin();
-		FMixerSubmixPtr ParentSubmixInstance = InSoundSubmix->ParentSubmix
-			? GetSubmixInstance(InSoundSubmix->ParentSubmix).Pin()
-			: MasterSubmixPtr;
+		UnloadSoundSubmix(*InSoundSubmix);
+	}
+
+	void FMixerDevice::UnloadSoundSubmix(USoundSubmix& InSoundSubmix)
+	{
+		check(IsInAudioThread());
+
+		FMixerSubmixWeakPtr MasterSubmix = GetMasterSubmix();
+		FMixerSubmixPtr ParentSubmixInstance = InSoundSubmix.ParentSubmix
+			? GetSubmixInstance(InSoundSubmix.ParentSubmix).Pin()
+			: MasterSubmix.Pin();
+
 		if (ParentSubmixInstance.IsValid())
 		{
-			ParentSubmixInstance->RemoveChildSubmix(GetSubmixInstance(InSoundSubmix));
+			ParentSubmixInstance->RemoveChildSubmix(GetSubmixInstance(&InSoundSubmix));
 		}
 
-		for (USoundSubmix* ChildSubmix : InSoundSubmix->ChildSubmixes)
+		for (USoundSubmix* ChildSubmix : InSoundSubmix.ChildSubmixes)
 		{
 			FMixerSubmixPtr ChildSubmixPtr = GetSubmixInstance(ChildSubmix).Pin();
 			if (ChildSubmixPtr.IsValid())
 			{
-				ChildSubmixPtr->SetParentSubmix(MasterSubmixPtr);
+				ChildSubmixPtr->SetParentSubmix(ParentSubmixInstance.IsValid()
+					? ParentSubmixInstance
+					: MasterSubmix);
 			}
 		}
 
-		Submixes.Remove(InSoundSubmix);
+		Submixes.Remove(&InSoundSubmix);
 	}
 
 	void FMixerDevice::InitSoundEffectPresets()
diff --git a/Engine/Source/Runtime/AudioMixer/Private/AudioMixerSubmix.cpp b/Engine/Source/Runtime/AudioMixer/Private/AudioMixerSubmix.cpp
--- a/Engine/Source/Runtime/AudioMixer/Private/AudioMixerSubmix.cpp
+++ b/Engine/Source/Runtime/AudioMixer/Private/AudioMixerSubmix.cpp
@@ -174,12 +174,28 @@ namespace Audio
 
 	void FMixerSubmix::SetParentSubmix(TWeakPtr<FMixerSubmix, ESPMode::ThreadSafe> SubmixWeakPtr)
 	{
+		if (ParentSubmix == SubmixWeakPtr)
+		{
+			return;
+		}
+
+		TSharedPtr<Audio::FMixerSubmix, ESPMode::ThreadSafe> ParentPtr = ParentSubmix.Pin();
+		if (ParentPtr.IsValid())
+		{
+			const uint32 InChildId = GetId();
+			ParentPtr->SubmixCommand([this, InChildId, SubmixWeakPtr]()
+			{
+				AUDIO_MIXER_CHECK_AUDIO_PLAT_THREAD(MixerDevice);
+
+				ChildSubmixes.Remove(InChildId);
+			});
+		}
+
 		SubmixCommand([this, SubmixWeakPtr]()
 		{
 			AUDIO_MIXER_CHECK_AUDIO_PLAT_THREAD(MixerDevice);
 
 			ParentSubmix = SubmixWeakPtr;
-
 			if (ChannelFormat == ESubmixChannelFormat::Ambisonics && AmbisonicsMixer.IsValid())
 			{
 				UpdateAmbisonicsDecoderForParent();
@@ -654,6 +670,7 @@ namespace Audio
 			CSV_SCOPED_TIMING_STAT(Audio, SubmixChildren);
 
 			// First loop this submix's child submixes mixing in their output into this submix's dry/wet buffers.
+			TArray<uint32> ToRemove;
 			for (auto& ChildSubmixEntry : ChildSubmixes)
 			{
 				TSharedPtr<Audio::FMixerSubmix, ESPMode::ThreadSafe> ChildSubmix = ChildSubmixEntry.Value.SubmixPtr.Pin();
@@ -676,6 +693,15 @@ namespace Audio
 						Audio::MixInBufferFast(ScratchBufferPtr, BufferPtr, NumSamples);
 					}
 				}
+				else
+				{
+					ToRemove.Add(ChildSubmixEntry.Key);
+				}
+			}
+
+			for (uint32 Key : ToRemove)
+			{
+				ChildSubmixes.Remove(Key);
 			}
 		}
 
diff --git a/Engine/Source/Runtime/AudioMixer/Public/AudioMixerDevice.h b/Engine/Source/Runtime/AudioMixer/Public/AudioMixerDevice.h
--- a/Engine/Source/Runtime/AudioMixer/Public/AudioMixerDevice.h
+++ b/Engine/Source/Runtime/AudioMixer/Public/AudioMixerDevice.h
@@ -224,9 +224,10 @@ namespace Audio
 		// Resets the thread ID used for audio rendering
 		void ResetAudioRenderingThreadId();
 
+		void RebuildSubmixLinks(USoundSubmix& SoundSubmix, FMixerSubmixPtr& SubmixInstance);
+
 		void Get2DChannelMapInternal(const int32 NumSourceChannels, const int32 NumOutputChannels, const bool bIsCenterChannelOnly, TArray<float>& OutChannelMap) const;
 		void InitializeChannelMaps();
-		void InitMasterSoundSubmix(EMasterSubmixType::Type InType, const FString& InDefaultName, bool bInDefaultMuteWhenBackgrounded, FSoftObjectPath& InOutObjectPath);
 		static int32 GetChannelMapCacheId(const int32 NumSourceChannels, const int32 NumOutputChannels, const bool bIsCenterChannelOnly);
 		void CacheChannelMap(const int32 NumSourceChannels, const int32 NumOutputChannels, const bool bIsCenterChannelOnly);
 		void InitializeChannelAzimuthMap(const int32 NumChannels);
@@ -236,6 +237,11 @@ namespace Audio
 
 		bool IsMainAudioDevice() const;
 
+		void LoadMasterSoundSubmix(EMasterSubmixType::Type InType, const FString& InDefaultName, bool bInDefaultMuteWhenBackgrounded, FSoftObjectPath& InOutObjectPath);
+		void LoadPluginSoundSubmixes();
+		void LoadSoundSubmix(USoundSubmix& SoundSubmix);
+		void UnloadSoundSubmix(USoundSubmix& SoundSubmix);
+
 	private:
 
 		bool IsMasterSubmixType(USoundSubmix* InSubmix) const;
diff --git a/Engine/Source/Runtime/Engine/Classes/Sound/AudioSettings.h b/Engine/Source/Runtime/Engine/Classes/Sound/AudioSettings.h
--- a/Engine/Source/Runtime/Engine/Classes/Sound/AudioSettings.h
+++ b/Engine/Source/Runtime/Engine/Classes/Sound/AudioSettings.h
@@ -214,6 +214,7 @@ private:
 #if WITH_EDITOR
 	TArray<FAudioQualitySettings> CachedQualityLevels;
 	FSoftObjectPath CachedAmbisonicSubmix;
+	FSoftObjectPath CachedMasterSubmix;
 #endif // WITH_EDITOR
 
 	void AddDefaultSettings();
diff --git a/Engine/Source/Runtime/Engine/Classes/Sound/SoundSubmix.h b/Engine/Source/Runtime/Engine/Classes/Sound/SoundSubmix.h
--- a/Engine/Source/Runtime/Engine/Classes/Sound/SoundSubmix.h
+++ b/Engine/Source/Runtime/Engine/Classes/Sound/SoundSubmix.h
@@ -183,7 +183,7 @@ public:
 	/**
 	* @return true if the child sound class exists in the tree
 	*/
-	bool RecurseCheckChild(USoundSubmix* ChildSoundSubmix);
+	bool RecurseCheckChild(const USoundSubmix* ChildSoundSubmix) const;
 
 	/**
 	* Set the parent submix of this SoundSubmix, removing it as a child from its previous owner
diff --git a/Engine/Source/Runtime/Engine/Private/AudioSettings.cpp b/Engine/Source/Runtime/Engine/Private/AudioSettings.cpp
--- a/Engine/Source/Runtime/Engine/Private/AudioSettings.cpp
+++ b/Engine/Source/Runtime/Engine/Private/AudioSettings.cpp
@@ -50,6 +50,9 @@ void UAudioSettings::PreEditChange(FProperty* PropertyAboutToChange)
 	// Cache ambisonic submix in case user tries to set to submix that isn't set to ambisonics
 	CachedAmbisonicSubmix = AmbisonicSubmix;
 
+	// Cache master submix in case user tries to set to submix that isn't a top-level submix
+	CachedMasterSubmix = MasterSubmix;
+
 	// Cache at least the first entry in case someone tries to clear the array
 	CachedQualityLevels = QualityLevels;
 }
@@ -61,8 +64,28 @@ void UAudioSettings::PostEditChangeChainProperty(FPropertyChangedChainEvent& Pro
 		bool bReconcileQualityNodes = false;
 		bool bPromptRestartRequired = false;
 		FName PropertyName = PropertyChangedEvent.Property->GetFName();
-		if (PropertyName == GET_MEMBER_NAME_CHECKED(UAudioSettings, MasterSubmix)
-			|| PropertyName == GET_MEMBER_NAME_CHECKED(UAudioSettings, EQSubmix)
+		if (PropertyName == GET_MEMBER_NAME_CHECKED(UAudioSettings, MasterSubmix))
+		{
+			if (USoundSubmix* NewSubmix = Cast<USoundSubmix>(AmbisonicSubmix.TryLoad()))
+			{
+				if (NewSubmix->ParentSubmix != nullptr)
+				{
+					FNotificationInfo Info(LOCTEXT("AudioSettings_InvalidMasterSubmix",
+						"Master Submix' cannot be set to submix with parent."));
+					Info.bFireAndForget = true;
+					Info.ExpireDuration = 2.0f;
+					Info.bUseThrobber = true;
+					FSlateNotificationManager::Get().AddNotification(Info);
+
+					MasterSubmix = CachedMasterSubmix;
+				}
+			}
+			else
+			{
+				bPromptRestartRequired = true;
+			}
+		}
+		else if(PropertyName == GET_MEMBER_NAME_CHECKED(UAudioSettings, EQSubmix)
 			|| PropertyName == GET_MEMBER_NAME_CHECKED(UAudioSettings, ReverbSubmix))
 		{
 			bPromptRestartRequired = true;
diff --git a/Engine/Source/Runtime/Engine/Private/SoundSubmix.cpp b/Engine/Source/Runtime/Engine/Private/SoundSubmix.cpp
--- a/Engine/Source/Runtime/Engine/Private/SoundSubmix.cpp
+++ b/Engine/Source/Runtime/Engine/Private/SoundSubmix.cpp
@@ -359,7 +359,7 @@ void USoundSubmix::PostDuplicate(EDuplicateMode::Type DuplicateMode)
 	Super::PostDuplicate(DuplicateMode);
 }
 
-bool USoundSubmix::RecurseCheckChild(USoundSubmix* ChildSoundSubmix)
+bool USoundSubmix::RecurseCheckChild(const USoundSubmix* ChildSoundSubmix) const
 {
 	for (int32 Index = 0; Index < ChildSubmixes.Num(); Index++)
 	{
