#UE4 Fix Crash at startup when au.DisableSubmixEffectEQ is 1

diff --git a/Engine/Source/Runtime/AudioMixer/Private/AudioMixerDevice.cpp b/Engine/Source/Runtime/AudioMixer/Private/AudioMixerDevice.cpp
--- a/Engine/Source/Runtime/AudioMixer/Private/AudioMixerDevice.cpp
+++ b/Engine/Source/Runtime/AudioMixer/Private/AudioMixerDevice.cpp
@@ -727,6 +727,11 @@ namespace Audio
 
 		for (int32 i = 0; i < static_cast<int32>(EMasterSubmixType::Count); ++i)
 		{
+			if (DisableSubmixEffectEQCvar && i == static_cast<int32>(EMasterSubmixType::EQ))
+			{
+				continue;
+			}
+
 			USoundSubmix* SoundSubmix = MasterSubmixes[i];
 			check(SoundSubmix);
 			FMixerSubmixPtr& MasterSubmixInstance = MasterSubmixInstances[i];
