Fix decoders not being cleared when timesynth is deactivated while playing and reactivated
Fix scratch buffer leak

diff --git a/Engine/Plugins/Runtime/TimeSynth/Source/TimeSynth/Private/TimeSynthComponent.cpp b/Engine/Plugins/Runtime/TimeSynth/Source/TimeSynth/Private/TimeSynthComponent.cpp
--- a/Engine/Plugins/Runtime/TimeSynth/Source/TimeSynth/Private/TimeSynthComponent.cpp
+++ b/Engine/Plugins/Runtime/TimeSynth/Source/TimeSynth/Private/TimeSynthComponent.cpp
@@ -358,6 +358,8 @@ void UTimeSynthComponent::ShutdownPlayingClips()
 		ActivePlayingClipIndices_AudioRenderThread.RemoveAtSwap(i, 1, false);
 		FreePlayingClipIndices_AudioRenderThread.Add(ClipIndex);
 	}
+
+	SoundWaveDecoder.Reset();
 }
 
 void UTimeSynthComponent::OnEndGenerate() 
diff --git a/Engine/Source/Runtime/AudioMixer/Private/SoundWaveDecoder.cpp b/Engine/Source/Runtime/AudioMixer/Private/SoundWaveDecoder.cpp
--- a/Engine/Source/Runtime/AudioMixer/Private/SoundWaveDecoder.cpp
+++ b/Engine/Source/Runtime/AudioMixer/Private/SoundWaveDecoder.cpp
@@ -314,6 +314,7 @@ namespace Audio
 		else
 		{
 			static AlignedFloatBuffer ScratchBuffer;
+			ScratchBuffer.Reset();
 			ScratchBuffer.AddZeroed(InNumFrames * SourceInfo.NumSourceChannels);
 
 			GetAudioBufferInternal(InNumFrames, InNumChannels, ScratchBuffer);
@@ -411,7 +412,8 @@ namespace Audio
 			{
 				DecodingSources.Add(InitData.Handle.Id, DecodingSoundWaveDataPtr);
 
-				UE_LOG(LogTemp, Log, TEXT("Decoding sources size %d."), DecodingSources.Num());
+				UE_LOG(LogAudioMixer, Verbose, TEXT("Decoding SoundWave '%s' (Num Decoding: %d)"),
+					*InitData.Handle.SoundWaveName.ToString(), DecodingSources.Num());
 			});
 
 			return true;
@@ -479,6 +481,15 @@ namespace Audio
 		DecodingSources.Remove(Handle.Id);
 	}
 
+	void FSoundSourceDecoder::Reset()
+	{
+		PumpDecoderCommandQueue();
+
+		DecodingSources.Reset();
+		InitializingDecodingSources.Reset();
+		PrecachingSources.Reset();
+	}
+
 	void FSoundSourceDecoder::SetSourcePitchScale(const FDecodingSoundSourceHandle& Handle, float InPitchScale)
 	{
 
diff --git a/Engine/Source/Runtime/AudioMixer/Public/SoundWaveDecoder.h b/Engine/Source/Runtime/AudioMixer/Public/SoundWaveDecoder.h
--- a/Engine/Source/Runtime/AudioMixer/Public/SoundWaveDecoder.h
+++ b/Engine/Source/Runtime/AudioMixer/Public/SoundWaveDecoder.h
@@ -210,6 +210,9 @@ namespace Audio
 		// Removes the decoding source from the decoder
 		void RemoveDecodingSource(const FDecodingSoundSourceHandle& Handle);
 
+		// Resets internal state of decoder
+		void Reset();
+
 		// Sets the source pitch scale
 		void SetSourcePitchScale(const FDecodingSoundSourceHandle& Handle, float InPitchScale);
 
