- Allow submix soft object paths in Audio Settings to reference same submix (ex. Setting Master EQ Submix field in AudioSettings to the same Submix asset as Master Submix now works)
- Support master submix effect chain updates at edit time (don't require restart)
- Fix null reference exception when nulling parent submix
- Clean up default submix assets
    - Move 'fast' reverb to default and remove 'fast' asset

diff --git a/Engine/Source/Runtime/AudioMixer/Private/AudioMixerDevice.cpp b/Engine/Source/Runtime/AudioMixer/Private/AudioMixerDevice.cpp
--- a/Engine/Source/Runtime/AudioMixer/Private/AudioMixerDevice.cpp
+++ b/Engine/Source/Runtime/AudioMixer/Private/AudioMixerDevice.cpp
@@ -622,11 +622,27 @@ namespace Audio
 
 		check(NewSubmix);
 		NewSubmix->AddToRoot();
-		MasterSubmixes[TypeIndex] = NewSubmix;
 
-		FMixerSubmixPtr NewMixerSubmix = MakeShared<FMixerSubmix, ESPMode::ThreadSafe>(this);
-		MasterSubmixInstances[TypeIndex] = NewMixerSubmix;
+		// If sharing submix with other explicitly defined MasterSubmix, create
+		// shared pointer directed to already existing submix instance. Otherwise,
+		// create a new version.
+		FMixerSubmixPtr NewMixerSubmix = GetMasterSubmixInstance(NewSubmix);
+		if (!NewMixerSubmix.IsValid())
+		{
+			UE_LOG(LogAudioMixer, Display, TEXT("Creating Master Submix '%s'"), *NewSubmix->GetName());
+			NewMixerSubmix = MakeShared<FMixerSubmix, ESPMode::ThreadSafe>(this);
+		}
 
+		// Ensure that master submixes are ONLY tracked in master submix array.
+		// MasterSubmixes array can share instances, but should not be duplicated in Submixes Map.
+		if (Submixes.Remove(NewSubmix) > 0)
+		{
+			UE_LOG(LogAudioMixer, Display, TEXT("Submix '%s' has been promoted to master array."), *NewSubmix->GetName());
+		}
+
+		// Update/add new submix and instance to respective master arrays
+		MasterSubmixes[TypeIndex] = NewSubmix;
+		MasterSubmixInstances[TypeIndex] = NewMixerSubmix;
 		NewMixerSubmix->Init(NewSubmix, false /* bAllowReInit */);
 	}
 
@@ -963,7 +979,7 @@ namespace Audio
 
 	void FMixerDevice::RegisterSoundSubmix(USoundSubmix* InSoundSubmix, bool bInit)
 	{
-		if (!InSoundSubmix || bSubmixRegistrationDisabled || IsMasterSubmixType(InSoundSubmix))
+		if (!InSoundSubmix || bSubmixRegistrationDisabled)
 		{
 			return;
 		}
@@ -981,14 +997,23 @@ namespace Audio
 			return;
 		}
 
-		LoadSoundSubmix(*InSoundSubmix);
+		const bool bIsMasterSubmix = IsMasterSubmixType(InSoundSubmix);
+
+		if (!bIsMasterSubmix)
+		{
+			LoadSoundSubmix(*InSoundSubmix);
+		}
+
 		FMixerSubmixPtr SubmixPtr = GetSubmixInstance(InSoundSubmix).Pin();
 		if (bInit)
 		{
-			SubmixPtr->Init(InSoundSubmix, true /* BAllowReInit */);
+			SubmixPtr->Init(InSoundSubmix, true /* bAllowReInit */);
 		}
 
-		RebuildSubmixLinks(*InSoundSubmix, SubmixPtr);
+		if (!bIsMasterSubmix)
+		{
+			RebuildSubmixLinks(*InSoundSubmix, SubmixPtr);
+		}
 	}
 
 	void FMixerDevice::LoadSoundSubmix(USoundSubmix& InSoundSubmix)
diff --git a/Engine/Source/Runtime/Engine/Private/AudioSettings.cpp b/Engine/Source/Runtime/Engine/Private/AudioSettings.cpp
--- a/Engine/Source/Runtime/Engine/Private/AudioSettings.cpp
+++ b/Engine/Source/Runtime/Engine/Private/AudioSettings.cpp
@@ -66,12 +66,12 @@ void UAudioSettings::PostEditChangeChainProperty(FPropertyChangedChainEvent& Pro
 		FName PropertyName = PropertyChangedEvent.Property->GetFName();
 		if (PropertyName == GET_MEMBER_NAME_CHECKED(UAudioSettings, MasterSubmix))
 		{
-			if (USoundSubmix* NewSubmix = Cast<USoundSubmix>(AmbisonicSubmix.TryLoad()))
+			if (USoundSubmix* NewSubmix = Cast<USoundSubmix>(MasterSubmix.TryLoad()))
 			{
-				if (NewSubmix->ParentSubmix != nullptr)
+				if (NewSubmix->ParentSubmix)
 				{
 					FNotificationInfo Info(LOCTEXT("AudioSettings_InvalidMasterSubmix",
-						"Master Submix' cannot be set to submix with parent."));
+						"'Master Submix' cannot be set to submix with parent."));
 					Info.bFireAndForget = true;
 					Info.ExpireDuration = 2.0f;
 					Info.bUseThrobber = true;
diff --git a/Engine/Source/Runtime/Engine/Private/SoundSubmix.cpp b/Engine/Source/Runtime/Engine/Private/SoundSubmix.cpp
--- a/Engine/Source/Runtime/Engine/Private/SoundSubmix.cpp
+++ b/Engine/Source/Runtime/Engine/Private/SoundSubmix.cpp
@@ -384,15 +384,19 @@ void USoundSubmix::SetParentSubmix(USoundSubmix* InParentSubmix)
 {
 	if (ParentSubmix != InParentSubmix)
 	{
-		if (ParentSubmix != nullptr)
+		if (ParentSubmix)
 		{
 			ParentSubmix->Modify();
 			ParentSubmix->ChildSubmixes.Remove(this);
 		}
 
 		Modify();
+
 		ParentSubmix = InParentSubmix;
-		ParentSubmix->ChildSubmixes.AddUnique(this);
+		if (ParentSubmix)
+		{
+			ParentSubmix->ChildSubmixes.AddUnique(this);
+		}
 	}
 }
 
