Dynamics Processor Updates
- Add Low/High shelf key detector filtering
- Add auditioning key signal
- Add gain stage dedicated to input key signal
- Stage disabled properties for external sidechaining
- Clean-up processor settings in editor
- Minor clean-up

Biquad Filter Updates
- Add Low/High shelf options

TODO:
- Add support for external sidechain submix
- Update DynProc SourceEffect/Granulator (Time synth?) to use same settings and expose/clean-up new/modified settings

diff --git a/Engine/Plugins/Runtime/Synthesis/Source/Synthesis/Private/SourceEffectDynamicsProcessor.cpp b/Engine/Plugins/Runtime/Synthesis/Source/Synthesis/Private/SourceEffectDynamicsProcessor.cpp
--- a/Engine/Plugins/Runtime/Synthesis/Source/Synthesis/Private/SourceEffectDynamicsProcessor.cpp
+++ b/Engine/Plugins/Runtime/Synthesis/Source/Synthesis/Private/SourceEffectDynamicsProcessor.cpp
@@ -61,8 +61,9 @@ void FSourceEffectDynamicsProcessor::OnPresetChanged()
 	DynamicsProcessor.SetKneeBandwidth(Settings.KneeBandwidthDb);
 	DynamicsProcessor.SetInputGain(Settings.InputGainDb);
 	DynamicsProcessor.SetOutputGain(Settings.OutputGainDb);
-	DynamicsProcessor.SetChannelLinked(Settings.bStereoLinked);
 	DynamicsProcessor.SetAnalogMode(Settings.bAnalogMode);
+
+	DynamicsProcessor.SetChannelLinkMode(Settings.bStereoLinked ? Audio::EDynamicsProcessorChannelLinkMode::Average : Audio::EDynamicsProcessorChannelLinkMode::Disabled);
 }
 
 void FSourceEffectDynamicsProcessor::ProcessAudio(const FSoundEffectSourceInputData& InData, float* OutAudioBufferData)
diff --git a/Engine/Plugins/Runtime/TimeSynth/Source/TimeSynth/Private/TimeSynthComponent.cpp b/Engine/Plugins/Runtime/TimeSynth/Source/TimeSynth/Private/TimeSynthComponent.cpp
--- a/Engine/Plugins/Runtime/TimeSynth/Source/TimeSynth/Private/TimeSynthComponent.cpp
+++ b/Engine/Plugins/Runtime/TimeSynth/Source/TimeSynth/Private/TimeSynthComponent.cpp
@@ -311,7 +311,7 @@ bool UTimeSynthComponent::Init(int32& InSampleRate)
 	DynamicsProcessor.SetKneeBandwidth(10.0f);
 	DynamicsProcessor.SetInputGain(0.0f);
 	DynamicsProcessor.SetOutputGain(0.0f);
-	DynamicsProcessor.SetChannelLinked(true);
+	DynamicsProcessor.SetChannelLinkMode(Audio::EDynamicsProcessorChannelLinkMode::Average);
 	DynamicsProcessor.SetAnalogMode(true);
 	DynamicsProcessor.SetPeakMode(Audio::EPeakMode::Peak);
 	DynamicsProcessor.SetProcessingMode(Audio::EDynamicsProcessingMode::Compressor);
diff --git a/Engine/Source/Runtime/AudioMixer/Classes/SubmixEffects/AudioMixerSubmixEffectDynamicsProcessor.h b/Engine/Source/Runtime/AudioMixer/Classes/SubmixEffects/AudioMixerSubmixEffectDynamicsProcessor.h
--- a/Engine/Source/Runtime/AudioMixer/Classes/SubmixEffects/AudioMixerSubmixEffectDynamicsProcessor.h
+++ b/Engine/Source/Runtime/AudioMixer/Classes/SubmixEffects/AudioMixerSubmixEffectDynamicsProcessor.h
@@ -4,8 +4,11 @@
 
 #include "DSP/DynamicsProcessor.h"
 #include "Sound/SoundEffectSubmix.h"
+#include "Sound/SoundSubmix.h"
+
 #include "AudioMixerSubmixEffectDynamicsProcessor.generated.h"
 
+
 UENUM(BlueprintType)
 enum class ESubmixEffectDynamicsProcessorType : uint8
 {
@@ -25,72 +28,138 @@ enum class ESubmixEffectDynamicsPeakMode : uint8
 	Count UMETA(Hidden)
 };
 
-// A submix dynamics processor
+UENUM(BlueprintType)
+enum class ESubmixEffectDynamicsChannelLinkMode : uint8
+{
+	Disabled = 0,
+	Average,
+	Peak,
+	Count UMETA(Hidden)
+};
+
+USTRUCT(BlueprintType)
+struct AUDIOMIXER_API FSubmixEffectDynamicProcessorFilterSettings
+{
+	GENERATED_USTRUCT_BODY()
+
+	// Whether or not filter is enabled
+	UPROPERTY(EditAnywhere, BlueprintReadWrite, Category = Filter, meta = (DisplayName = "Enabled"))
+	uint8 bEnabled : 1;
+
+	// The cutoff frequency of the HPF applied to key signal
+	UPROPERTY(EditAnywhere, BlueprintReadWrite, Category = Filter, meta = (DisplayName = "Cutoff (Hz)", EditCondition = "bEnabled", ClampMin = "20.0", ClampMax = "20000.0", UIMin = "20.0", UIMax = "20000.0"))
+	float Cutoff;
+
+	// The gain of the filter shelf applied to the key signal
+	UPROPERTY(EditAnywhere, BlueprintReadWrite, Category = Filter, meta = (DisplayName = "Gain (dB)", EditCondition = "bEnabled", ClampMin = "-60.0", ClampMax = "6.0", UIMin = "-60.0", UIMax = "6.0"))
+	float GainDb;
+
+	FSubmixEffectDynamicProcessorFilterSettings()
+		: bEnabled(false)
+		, Cutoff(20.0f)
+		, GainDb(0.0f)
+	{
+	}
+};
+
+// Submix dynamics processor settings
 USTRUCT(BlueprintType)
 struct AUDIOMIXER_API FSubmixEffectDynamicsProcessorSettings
 {
 	GENERATED_USTRUCT_BODY()
 
-	UPROPERTY(EditAnywhere, BlueprintReadWrite, Category = "SourceEffect|Preset")
+	// Type of processor to apply
+	UPROPERTY(EditAnywhere, BlueprintReadWrite, Category = General, meta = (DisplayName = "Type"))
 	ESubmixEffectDynamicsProcessorType DynamicsProcessorType;
 
-	UPROPERTY(EditAnywhere, BlueprintReadWrite, Category = "SourceEffect|Preset")
+	// Mode of peak detection used on input key signal
+	UPROPERTY(EditAnywhere, BlueprintReadWrite, Category = Dynamics)
 	ESubmixEffectDynamicsPeakMode PeakMode;
 
-	// The amount of time to look ahead of the current audio. Allows for transients to be included in dynamics processing.
-	UPROPERTY(EditAnywhere, BlueprintReadWrite, Category = "SourceEffect|Preset",  meta = (ClampMin = "0.0", ClampMax = "50.0", UIMin = "0.0", UIMax = "50.0"))
-	float LookAheadMsec;
+	// Mode of peak detection if key signal is multi-channel
+	UPROPERTY(EditAnywhere, BlueprintReadWrite, Category = Dynamics)
+	ESubmixEffectDynamicsChannelLinkMode LinkMode;
 
-	// The amount of time to ramp into any dynamics processing effect in milliseconds.
-	UPROPERTY(EditAnywhere, BlueprintReadWrite, Category = "SourceEffect|Preset", meta = (ClampMin = "1.0", ClampMax = "300.0", UIMin = "1.0", UIMax = "200.0"))
-	float AttackTimeMsec;
-
-	// The amount of time to release the dynamics processing effect in milliseconds
-	UPROPERTY(EditAnywhere, BlueprintReadWrite, Category = "SourceEffect|Preset", meta = (ClampMin = "20.0", ClampMax = "5000.0", UIMin = "20.0", UIMax = "5000.0"))
-	float ReleaseTimeMsec;
+	// The input gain of the dynamics processor
+	UPROPERTY(EditAnywhere, BlueprintReadWrite, Category = General, meta = (DisplayName = "Input Gain (dB)", ClampMin = "-12.0", ClampMax = "20.0", UIMin = "-12.0", UIMax = "20.0"))
+	float InputGainDb;
 
 	// The threshold at which to perform a dynamics processing operation
-	UPROPERTY(EditAnywhere, BlueprintReadWrite, Category = "SourceEffect|Preset", meta = (ClampMin = "-60.0", ClampMax = "0.0", UIMin = "-60.0", UIMax = "0.0"))
+	UPROPERTY(EditAnywhere, BlueprintReadWrite, Category = Dynamics, meta = (DisplayName = "Threshold (dB)", ClampMin = "-60.0", ClampMax = "0.0", UIMin = "-60.0", UIMax = "0.0"))
 	float ThresholdDb;
 
-	// The dynamics processor ratio -- has different meaning depending on the processor type.
-	UPROPERTY(EditAnywhere, BlueprintReadWrite, Category = "SourceEffect|Preset", meta = (ClampMin = "1.0", ClampMax = "20.0", UIMin = "1.0", UIMax = "20.0"))
+	// The dynamics processor ratio used for compression/expansion
+	UPROPERTY(EditAnywhere, BlueprintReadWrite, Category = Dynamics, meta = (
+		EditCondition = "DynamicsProcessorType == ESubmixEffectDynamicsProcessorType::Compressor || DynamicsProcessorType == ESubmixEffectDynamicsProcessorType::Expander",
+		ClampMin = "1.0", ClampMax = "20.0", UIMin = "1.0", UIMax = "20.0"))
 	float Ratio;
 
-	// The knee bandwidth of the compressor to use in dB
-	UPROPERTY(EditAnywhere, BlueprintReadWrite, Category = "SourceEffect|Preset", meta = (ClampMin = "0.0", ClampMax = "20.0", UIMin = "0.0", UIMax = "20.0"))
+	// The knee bandwidth of the processor to use
+	UPROPERTY(EditAnywhere, BlueprintReadWrite, Category = Dynamics, meta = (DisplayName = "Knee (dB)", ClampMin = "0.0", ClampMax = "20.0", UIMin = "0.0", UIMax = "20.0"))
 	float KneeBandwidthDb;
 
-	// The input gain of the dynamics processor in dB
-	UPROPERTY(EditAnywhere, BlueprintReadWrite, Category = "SourceEffect|Preset", meta = (ClampMin = "-12.0", ClampMax = "20.0", UIMin = "-12.0", UIMax = "20.0"))
-	float InputGainDb;
+	// The amount of time to look ahead of the current audio (Allows for transients to be included in dynamics processing)
+	UPROPERTY(EditAnywhere, BlueprintReadWrite, Category = Response,  meta = (DisplayName = "Look Ahead (ms)", ClampMin = "0.0", ClampMax = "50.0", UIMin = "0.0", UIMax = "50.0"))
+	float LookAheadMsec;
+
+	// The amount of time to ramp into any dynamics processing effect
+	UPROPERTY(EditAnywhere, BlueprintReadWrite, Category = Response, meta = (DisplayName = "AttackTime (ms)", ClampMin = "1.0", ClampMax = "300.0", UIMin = "1.0", UIMax = "200.0"))
+	float AttackTimeMsec;
+
+	// The amount of time to release the dynamics processing effect
+	UPROPERTY(EditAnywhere, BlueprintReadWrite, Category = Response, meta = (DisplayName = "Release Time (ms)", ClampMin = "20.0", ClampMax = "5000.0", UIMin = "20.0", UIMax = "5000.0"))
+	float ReleaseTimeMsec;
+
+	// (Coming soon) If set, uses output of provided submix as modulator of input signal for dynamics processor (Uses input signal as default modulator)
+	UPROPERTY(VisibleAnywhere, Category = Sidechain)
+	USoundSubmix* ExternalSubmix;
+
+	UPROPERTY()
+	uint8 bChannelLinked_DEPRECATED : 1;
+
+	// Toggles treating the attack and release envelopes as analog-style vs digital-style (Analog will respond a bit more naturally/slower)
+	UPROPERTY(EditAnywhere, BlueprintReadWrite, Category = Response)
+	uint8 bAnalogMode : 1;
+
+	// Audition the key modulation signal, bypassing enveloping and processing the input signal.
+	UPROPERTY(EditAnywhere, BlueprintReadWrite, Category = Sidechain, meta = (DisplayName = "Key Audition"))
+	uint8 bKeyAudition : 1;
+
+	// Gain to apply to key signal (external signal if supplied or input signal if disabled)
+	UPROPERTY(EditAnywhere, BlueprintReadWrite, Category = Sidechain, meta = (DisplayName = "Key Gain (dB)", ClampMin = "-60.0", ClampMax = "20.0", UIMin = "-60.0", UIMax = "20.0"))
+	float KeyGainDb;
 
-	// The output gain of the dynamics processor in dB
-	UPROPERTY(EditAnywhere, BlueprintReadWrite, Category = "SourceEffect|Preset", meta = (ClampMin = "0.0", ClampMax = "20.0", UIMin = "0.0", UIMax = "20.0"))
+	// The output gain of the dynamics processor
+	UPROPERTY(EditAnywhere, BlueprintReadWrite, Category = Output, meta = (DisplayName = "Output Gain (dB)", ClampMin = "-60.0", ClampMax = "0.0", UIMin = "-60.0", UIMax = "20.0"))
 	float OutputGainDb;
 
-	// Whether or not to average all channels of audio before inputing into the dynamics processor
-	UPROPERTY(EditAnywhere, BlueprintReadWrite, Category = "SourceEffect|Preset")
-	uint32 bChannelLinked : 1;
+	// High Shelf filter settings for key signal (external signal if supplied or input signal if not)
+	UPROPERTY(EditAnywhere, BlueprintReadWrite, Category = Sidechain, meta = (DisplayName = "Key Highshelf"))
+	FSubmixEffectDynamicProcessorFilterSettings KeyHighshelf;
 
-	// Toggles treating the attack and release envelopes as analog-style vs digital-style. Analog will respond a bit more naturally/slower.
-	UPROPERTY(EditAnywhere, BlueprintReadWrite, Category = "SourceEffect|Preset")
-	uint32 bAnalogMode : 1;
+	// Low Shelf filter settings for key signal (external signal if supplied or input signal if not)
+	UPROPERTY(EditAnywhere, BlueprintReadWrite, Category = Sidechain, meta = (DisplayName = "Key Lowshelf"))
+	FSubmixEffectDynamicProcessorFilterSettings KeyLowshelf;
 
 	FSubmixEffectDynamicsProcessorSettings()
 		: DynamicsProcessorType(ESubmixEffectDynamicsProcessorType::Compressor)
 		, PeakMode(ESubmixEffectDynamicsPeakMode::RootMeanSquared)
-		, LookAheadMsec(3.0f)
-		, AttackTimeMsec(10.0f)
-		, ReleaseTimeMsec(100.0f)
+		, LinkMode(ESubmixEffectDynamicsChannelLinkMode::Average)
+		, InputGainDb(0.0f)
 		, ThresholdDb(-6.0f)
 		, Ratio(1.5f)
 		, KneeBandwidthDb(10.0f)
-		, InputGainDb(0.0f)
-		, OutputGainDb(0.0f)
-		, bChannelLinked(true)
+		, LookAheadMsec(3.0f)
+		, AttackTimeMsec(10.0f)
+		, ReleaseTimeMsec(100.0f)
+		, ExternalSubmix(nullptr)
+		, bChannelLinked_DEPRECATED(true)
 		, bAnalogMode(true)
+		, bKeyAudition(false)
+		, KeyGainDb(0.0f)
+		, OutputGainDb(0.0f)
 	{
+		KeyLowshelf.Cutoff = 20000.0f;
 	}
 };
 
@@ -124,9 +193,11 @@ public:
 
 	EFFECT_PRESET_METHODS(SubmixEffectDynamicsProcessor)
 
+	void Serialize(FStructuredArchive::FRecord Record) override;
+
 	UFUNCTION(BlueprintCallable, Category = "Audio|Effects")
 	void SetSettings(const FSubmixEffectDynamicsProcessorSettings& InSettings);
 
-	UPROPERTY(EditAnywhere, BlueprintReadWrite, Category = SubmixEffectPreset)
+	UPROPERTY(EditAnywhere, BlueprintReadWrite, Category = SubmixEffectPreset, meta = (ShowOnlyInnerProperties))
 	FSubmixEffectDynamicsProcessorSettings Settings;
 };
\ No newline at end of file
diff --git a/Engine/Source/Runtime/AudioMixer/Private/Effects/AudioMixerSubmixEffectDynamicsProcessor.cpp b/Engine/Source/Runtime/AudioMixer/Private/Effects/AudioMixerSubmixEffectDynamicsProcessor.cpp
--- a/Engine/Source/Runtime/AudioMixer/Private/Effects/AudioMixerSubmixEffectDynamicsProcessor.cpp
+++ b/Engine/Source/Runtime/AudioMixer/Private/Effects/AudioMixerSubmixEffectDynamicsProcessor.cpp
@@ -12,10 +12,16 @@ FSubmixEffectDynamicsProcessor::FSubmixEffectDynamicsProcessor()
 
 void FSubmixEffectDynamicsProcessor::Init(const FSoundEffectSubmixInitData& InitData)
 {
-	DynamicsProcessor.Init(InitData.SampleRate, 8);
+	// Use a full multichannel frame as a scratch processing buffer
+	static const int32 NumChannels = 8;
+
+	DynamicsProcessor.Init(InitData.SampleRate, NumChannels);
+
+	AudioInputFrame.Reset();
+	AudioInputFrame.AddZeroed(NumChannels);
 
 	AudioOutputFrame.Reset();
-	AudioOutputFrame.AddZeroed(8);
+	AudioOutputFrame.AddZeroed(NumChannels);
 }
 
 
@@ -67,8 +73,19 @@ void FSubmixEffectDynamicsProcessor::OnPresetChanged()
 	DynamicsProcessor.SetKneeBandwidth(Settings.KneeBandwidthDb);
 	DynamicsProcessor.SetInputGain(Settings.InputGainDb);
 	DynamicsProcessor.SetOutputGain(Settings.OutputGainDb);
-	DynamicsProcessor.SetChannelLinked(Settings.bChannelLinked);
 	DynamicsProcessor.SetAnalogMode(Settings.bAnalogMode);
+
+	DynamicsProcessor.SetKeyAudition(Settings.bKeyAudition);
+	DynamicsProcessor.SetKeyGain(Settings.KeyGainDb);
+	DynamicsProcessor.SetKeyHighshelfCutoffFrequency(Settings.KeyHighshelf.Cutoff);
+	DynamicsProcessor.SetKeyHighshelfEnabled(Settings.KeyHighshelf.bEnabled);
+	DynamicsProcessor.SetKeyHighshelfGain(Settings.KeyHighshelf.GainDb);
+	DynamicsProcessor.SetKeyLowshelfCutoffFrequency(Settings.KeyLowshelf.Cutoff);
+	DynamicsProcessor.SetKeyLowshelfEnabled(Settings.KeyLowshelf.bEnabled);
+	DynamicsProcessor.SetKeyLowshelfGain(Settings.KeyLowshelf.GainDb);
+
+	static_assert(static_cast<int32>(ESubmixEffectDynamicsChannelLinkMode::Count) == static_cast<int32>(Audio::EDynamicsProcessorChannelLinkMode::Count), "Enumerations must match");
+	DynamicsProcessor.SetChannelLinkMode(static_cast<Audio::EDynamicsProcessorChannelLinkMode>(Settings.LinkMode));
 }
 
 void FSubmixEffectDynamicsProcessor::OnProcessAudio(const FSoundEffectSubmixInputData& InData, FSoundEffectSubmixOutputData& OutData)
@@ -78,10 +95,6 @@ void FSubmixEffectDynamicsProcessor::OnProcessAudio(const FSoundEffectSubmixInpu
 	const Audio::AlignedFloatBuffer& InBuffer = *InData.AudioBuffer;
 	Audio::AlignedFloatBuffer& OutBuffer = *OutData.AudioBuffer;
 
-	// Use a full multichannel frame as a scratch processing buffer
-	AudioInputFrame.Reset();
-	AudioInputFrame.AddZeroed(8);
-
 	for (int32 Frame = 0; Frame < InData.NumFrames; ++Frame)
 	{
 		// Copy the data to the frame input
@@ -104,6 +117,21 @@ void FSubmixEffectDynamicsProcessor::OnProcessAudio(const FSoundEffectSubmixInpu
 	}
 }
 
+void USubmixEffectDynamicsProcessorPreset::Serialize(FStructuredArchive::FRecord Record)
+{
+	FArchive& UnderlyingArchive = Record.GetUnderlyingArchive();
+	if (UnderlyingArchive.IsLoading())
+	{
+		if (Settings.bChannelLinked_DEPRECATED)
+		{
+			Settings.LinkMode = ESubmixEffectDynamicsChannelLinkMode::Average;
+			Settings.bChannelLinked_DEPRECATED = 0;
+		}
+	}
+
+	Super::Serialize(Record);
+}
+
 void USubmixEffectDynamicsProcessorPreset::SetSettings(const FSubmixEffectDynamicsProcessorSettings& InSettings)
 {
 	UpdateSettings(InSettings);
diff --git a/Engine/Source/Runtime/SignalProcessing/Private/DynamicsProcessor.cpp b/Engine/Source/Runtime/SignalProcessing/Private/DynamicsProcessor.cpp
--- a/Engine/Source/Runtime/SignalProcessing/Private/DynamicsProcessor.cpp
+++ b/Engine/Source/Runtime/SignalProcessing/Private/DynamicsProcessor.cpp
@@ -14,9 +14,13 @@ namespace Audio
 		, HalfKneeBandwidthDb(5.0f)
 		, InputGain(1.0f)
 		, OutputGain(1.0f)
+		, KeyGain(1.0f)
 		, NumChannels(0)
-		, bIsChannelLinked(true)
+		, LinkMode(EDynamicsProcessorChannelLinkMode::Disabled)
 		, bIsAnalogMode(true)
+		, bKeyAuditionEnabled(false)
+		, bKeyHighshelfEnabled(false)
+		, bKeyLowshelfEnabled(false)
 	{
 		// The knee will have 2 points
 		KneePoints.Init(FVector2D(), 2);
@@ -29,15 +33,29 @@ namespace Audio
 	void FDynamicsProcessor::Init(const float SampleRate, const int32 InNumChannels)
 	{
 		NumChannels = InNumChannels;
+
+		LookaheadDelay.Reset();
+		LookaheadDelay.AddDefaulted(InNumChannels);
+
+		EnvFollower.Reset();
+		EnvFollower.AddDefaulted(InNumChannels);
+
 		for (int32 Channel = 0; Channel < InNumChannels; ++Channel)
 		{
-			int32 Index = LookaheadDelay.Add(FDelay());
-			LookaheadDelay[Index].Init(SampleRate, 0.1f);
-			LookaheadDelay[Index].SetDelayMsec(LookaheedDelayMsec);
+			LookaheadDelay[Channel].Init(SampleRate, 0.1f);
+			LookaheadDelay[Channel].SetDelayMsec(LookaheedDelayMsec);
 
-			Index = EnvFollower.Add(FEnvelopeFollower());
-			EnvFollower[Index].Init(SampleRate, AttackTimeMsec, ReleaseTimeMsec, EPeakMode::RootMeanSquared, bIsAnalogMode);
+			EnvFollower[Channel].Init(SampleRate, AttackTimeMsec, ReleaseTimeMsec, EPeakMode::RootMeanSquared, bIsAnalogMode);
 		}
+
+		InputLowshelfFilter.Init(SampleRate, InNumChannels, EBiquadFilter::LowShelf);
+		InputHighshelfFilter.Init(SampleRate, InNumChannels, EBiquadFilter::HighShelf);
+
+		DetectorOuts.Reset();
+		DetectorOuts.AddZeroed(InNumChannels);
+
+		Gain.Reset();
+		Gain.AddZeroed(InNumChannels);
 	}
 
 	void FDynamicsProcessor::SetLookaheadMsec(const float InLookAheadMsec)
@@ -88,14 +106,54 @@ namespace Audio
 		InputGain = ConvertToLinear(InInputGainDb);
 	}
 
+	void FDynamicsProcessor::SetKeyAudition(const bool InAuditionEnabled)
+	{
+		bKeyAuditionEnabled = InAuditionEnabled;
+	}
+
+	void FDynamicsProcessor::SetKeyGain(const float InKeyGain)
+	{
+		KeyGain = ConvertToLinear(InKeyGain);
+	}
+
+	void FDynamicsProcessor::SetKeyHighshelfCutoffFrequency(const float InCutoffFreq)
+	{
+		InputHighshelfFilter.SetFrequency(InCutoffFreq);
+	}
+
+	void FDynamicsProcessor::SetKeyHighshelfEnabled(const bool bInEnabled)
+	{
+		bKeyHighshelfEnabled = bInEnabled;
+	}
+
+	void FDynamicsProcessor::SetKeyHighshelfGain(const float InGainDb)
+	{
+		InputHighshelfFilter.SetGainDB(InGainDb);
+	}
+
+	void FDynamicsProcessor::SetKeyLowshelfCutoffFrequency(const float InCutoffFreq)
+	{
+		InputLowshelfFilter.SetFrequency(InCutoffFreq);
+	}
+
+	void FDynamicsProcessor::SetKeyLowshelfEnabled(const bool bInEnabled)
+	{
+		bKeyLowshelfEnabled = bInEnabled;
+	}
+
+	void FDynamicsProcessor::SetKeyLowshelfGain(const float InGainDb)
+	{
+		InputLowshelfFilter.SetGainDB(InGainDb);
+	}
+
 	void FDynamicsProcessor::SetOutputGain(const float InOutputGainDb)
 	{
 		OutputGain = ConvertToLinear(InOutputGainDb);
 	}
 
-	void FDynamicsProcessor::SetChannelLinked(const bool bInIsChannelLinked)
+	void FDynamicsProcessor::SetChannelLinkMode(const EDynamicsProcessorChannelLinkMode InLinkMode)
 	{
-		bIsChannelLinked = bInIsChannelLinked;
+		LinkMode = InLinkMode;
 	}
 
 	void FDynamicsProcessor::SetAnalogMode(const bool bInIsAnalogMode)
@@ -122,59 +180,97 @@ namespace Audio
 
 	void FDynamicsProcessor::ProcessAudioFrame(const float* InFrame, float* OutFrame)
 	{
-		// Get detector outputs
-		DetectorOuts.Reset();
-		 
-		for (int32 Channel = 0; Channel < NumChannels; ++Channel)
-		{
-			DetectorOuts.Add(EnvFollower[Channel].ProcessAudio(InputGain * InFrame[Channel]));
-		}
-
-		Gain.Reset();
+		checkSlow(NumChannels == DetectorOuts.Num());
+		checkSlow(NumChannels == Gain.Num());
 
-		if (bIsChannelLinked)
+		// Get detector outputs
+		const float* DetectorIn = InFrame;
+		if (NumChannels > 0)
 		{
-			// average all channels
-			float DetectorOutLinked = 0.0f;
-			for (int32 Channel = 0; Channel < NumChannels; ++Channel)
+			if (bKeyLowshelfEnabled)
 			{
-				DetectorOutLinked += DetectorOuts[Channel];
+				InputLowshelfFilter.ProcessAudioFrame(DetectorIn, DetectorOuts.GetData());
+				DetectorIn = DetectorOuts.GetData();
 			}
-			DetectorOutLinked /= (float)NumChannels;
-
-			// convert to decibels
-			const float DetectorOutLinkedDb = ConvertToDecibels(DetectorOutLinked);
 
-			const float ComputedGain = ComputeGain(DetectorOutLinkedDb);
+			if (bKeyHighshelfEnabled)
+			{
+				InputHighshelfFilter.ProcessAudioFrame(DetectorIn, DetectorOuts.GetData());
+				DetectorIn = DetectorOuts.GetData();
+			}
+		}
 
+		const float DetectorGain = KeyGain * InputGain;
+		if (bKeyAuditionEnabled)
+		{
 			for (int32 Channel = 0; Channel < NumChannels; ++Channel)
 			{
-				Gain.Add(ComputedGain);
+				OutFrame[Channel] = DetectorGain * DetectorIn[Channel];
 			}
 		}
 		else
 		{
-			// compute gain individually per channel
 			for (int32 Channel = 0; Channel < NumChannels; ++Channel)
 			{
-				// convert to decibels
-				const float DetectorOutLinkedDb = ConvertToDecibels(DetectorOuts[Channel]);
+				DetectorOuts[Channel] = EnvFollower[Channel].ProcessAudio(DetectorGain * DetectorIn[Channel]);
+			}
+
+			switch (LinkMode)
+			{
+				case EDynamicsProcessorChannelLinkMode::Average:
+				{
+					float DetectorOutLinked = 0.0f;
+					for (int32 Channel = 0; Channel < NumChannels; ++Channel)
+					{
+						DetectorOutLinked += DetectorOuts[Channel];
+					}
+					DetectorOutLinked /= static_cast<float>(NumChannels);
+					const float DetectorOutLinkedDb = ConvertToDecibels(DetectorOutLinked);
+					const float ComputedGain = ComputeGain(DetectorOutLinkedDb);
+					for (int32 Channel = 0; Channel < NumChannels; ++Channel)
+					{
+						Gain[Channel] = ComputedGain;
+					}
+				}
+				break;
 
-				const float ComputedGain = ComputeGain(DetectorOutLinkedDb);
+				case EDynamicsProcessorChannelLinkMode::Peak:
+				{
+					float DetectorOutLinked = FMath::Max<float>(DetectorOuts);
+					const float DetectorOutLinkedDb = ConvertToDecibels(DetectorOutLinked);
+					const float ComputedGain = ComputeGain(DetectorOutLinkedDb);
+					for (int32 Channel = 0; Channel < NumChannels; ++Channel)
+					{
+						Gain[Channel] = ComputedGain;
+					}
+				}
+				break;
 
-				Gain.Add(ComputedGain);
+				case EDynamicsProcessorChannelLinkMode::Disabled:
+				default:
+				{
+					// compute gain individually per channel
+					for (int32 Channel = 0; Channel < NumChannels; ++Channel)
+					{
+						// convert to decibels
+						const float DetectorOutLinkedDb = ConvertToDecibels(DetectorOuts[Channel]);
+						const float ComputedGain = ComputeGain(DetectorOutLinkedDb);
+						Gain[Channel] = ComputedGain;
+					}
+				}
+				break;
 			}
-		}
 
-		for (int32 Channel = 0; Channel < NumChannels; ++Channel)
-		{
-			// Write and read into the look ahead delay line.
-			// We apply the compression output of the direct input to the output of this delay line
-			// This way sharp transients can be "caught" with the gain.
-			float LookaheadOutput = LookaheadDelay[Channel].ProcessAudioSample(InFrame[Channel]);
+			for (int32 Channel = 0; Channel < NumChannels; ++Channel)
+			{
+				// Write and read into the look ahead delay line.
+				// We apply the compression output of the direct input to the output of this delay line
+				// This way sharp transients can be "caught" with the gain.
+				float LookaheadOutput = LookaheadDelay[Channel].ProcessAudioSample(InFrame[Channel]);
 
-			// Write into the output with the computed gain value
-			OutFrame[Channel] = Gain[Channel] * LookaheadOutput * OutputGain * InputGain;
+				// Write into the output with the computed gain value
+				OutFrame[Channel] = Gain[Channel] * LookaheadOutput * OutputGain * InputGain;
+			}
 		}
 	}
 
diff --git a/Engine/Source/Runtime/SignalProcessing/Private/Filter.cpp b/Engine/Source/Runtime/SignalProcessing/Private/Filter.cpp
--- a/Engine/Source/Runtime/SignalProcessing/Private/Filter.cpp
+++ b/Engine/Source/Runtime/SignalProcessing/Private/Filter.cpp
@@ -126,7 +126,9 @@ namespace Audio
 		if (GainDB != InGainDB)
 		{
 			GainDB = InGainDB;
-			if (FilterType == EBiquadFilter::ParametricEQ)
+			if (FilterType == EBiquadFilter::ParametricEQ
+				|| FilterType == EBiquadFilter::HighShelf
+				|| FilterType == EBiquadFilter::LowShelf)
 			{
 				CalculateBiquadCoefficients();
 			}
@@ -205,7 +207,6 @@ namespace Audio
 			case EBiquadFilter::ParametricEQ:
 			{
 				const float Amp = FMath::Pow(10.0f, GainDB / 40.0f);
-				const float Beta = FMath::Sqrt(2.0f * Amp);
 
 				a0 = 1.0f + (Alpha * Amp);
 				a1 = -2.0f * Cs;
@@ -216,6 +217,34 @@ namespace Audio
 			}
 			break;
 
+			case EBiquadFilter::HighShelf:
+			{
+				const float Amp = FMath::Pow(10.0f, GainDB / 40.0f);
+				const float BetaSn = Sn * FMath::Sqrt(2.0f * Amp);
+
+				a0 = Amp * ((Amp + 1) + (Amp - 1) * Cs + BetaSn);
+				a1 = -2 * Amp * ((Amp - 1) + (Amp + 1) * Cs);
+				a2 = Amp * ((Amp + 1) + (Amp - 1) * Cs - BetaSn);
+				b0 = (Amp + 1) - (Amp - 1) * Cs + BetaSn;
+				b1 = 2 * ((Amp - 1) - (Amp + 1) * Cs);
+				b2 = (Amp + 1) - (Amp - 1) * Cs - BetaSn;
+			}
+			break;
+
+			case EBiquadFilter::LowShelf:
+			{
+				const float Amp = FMath::Pow(10.0f, GainDB / 40.0f);
+				const float BetaSn = Sn * FMath::Sqrt(2.0f * Amp);
+
+				a0 = Amp * ((Amp + 1) - (Amp - 1) * Cs + BetaSn);
+				a1 = 2 * Amp * ((Amp - 1) - (Amp + 1) * Cs);
+				a2 = Amp * ((Amp + 1) - (Amp - 1) * Cs - BetaSn);
+				b0 = (Amp + 1) + (Amp - 1) * Cs + BetaSn;
+				b1 = -2 * ((Amp - 1) + (Amp + 1) * Cs);
+				b2 = (Amp + 1) + (Amp - 1) * Cs - BetaSn;
+			}
+			break;
+
 			case EBiquadFilter::AllPass:
 			{
 				a0 = 1.0f - Alpha;
diff --git a/Engine/Source/Runtime/SignalProcessing/Private/Granulator.cpp b/Engine/Source/Runtime/SignalProcessing/Private/Granulator.cpp
--- a/Engine/Source/Runtime/SignalProcessing/Private/Granulator.cpp
+++ b/Engine/Source/Runtime/SignalProcessing/Private/Granulator.cpp
@@ -423,7 +423,7 @@ namespace Audio
 		DynamicsProcessor.SetKneeBandwidth(10.0f);
 		DynamicsProcessor.SetInputGain(0.0f);
 		DynamicsProcessor.SetOutputGain(0.0f);
-		DynamicsProcessor.SetChannelLinked(true);
+		DynamicsProcessor.SetChannelLinkMode(EDynamicsProcessorChannelLinkMode::Average);
 		DynamicsProcessor.SetAnalogMode(true);
 		DynamicsProcessor.SetPeakMode(EPeakMode::Peak);
 		DynamicsProcessor.SetProcessingMode(EDynamicsProcessingMode::Compressor);
diff --git a/Engine/Source/Runtime/SignalProcessing/Public/DSP/DynamicsProcessor.h b/Engine/Source/Runtime/SignalProcessing/Public/DSP/DynamicsProcessor.h
--- a/Engine/Source/Runtime/SignalProcessing/Public/DSP/DynamicsProcessor.h
+++ b/Engine/Source/Runtime/SignalProcessing/Public/DSP/DynamicsProcessor.h
@@ -5,6 +5,7 @@
 #include "CoreMinimal.h"
 #include "DSP/EnvelopeFollower.h"
 #include "DSP/Delay.h"
+#include "Filter.h"
 
 namespace Audio
 {
@@ -21,6 +22,14 @@ namespace Audio
 		};
 	}
 
+	enum class EDynamicsProcessorChannelLinkMode : uint8
+	{
+		Disabled,
+		Average,
+		Peak,
+		Count
+	};
+
 	// Dynamic range compressor
 	// https://en.wikipedia.org/wiki/Dynamic_range_compression
 	class SIGNALPROCESSING_API FDynamicsProcessor
@@ -38,8 +47,16 @@ namespace Audio
 		void SetRatio(const float InCompressionRatio);
 		void SetKneeBandwidth(const float InKneeBandwidthDb);
 		void SetInputGain(const float InInputGainDb);
+		void SetKeyAudition(const bool InAuditionEnabled);
+		void SetKeyGain(const float InKeyGain);
+		void SetKeyHighshelfCutoffFrequency(const float InCutoffFreq);
+		void SetKeyHighshelfEnabled(const bool bInEnabled);
+		void SetKeyHighshelfGain(const float InGainDb);
+		void SetKeyLowshelfCutoffFrequency(const float InCutoffFreq);
+		void SetKeyLowshelfEnabled(const bool bInEnabled);
+		void SetKeyLowshelfGain(const float InGainDb);
 		void SetOutputGain(const float InOutputGainDb);
-		void SetChannelLinked(const bool bInIsChannelLinked);
+		void SetChannelLinkMode(const EDynamicsProcessorChannelLinkMode InLinkMode);
 		void SetAnalogMode(const bool bInIsAnalogMode);
 		void SetPeakMode(const EPeakMode::Type InEnvelopeFollowerModeType);
 		void SetProcessingMode(const EDynamicsProcessingMode::Type ProcessingMode);
@@ -52,6 +69,12 @@ namespace Audio
 
 		float ComputeGain(const float InEnvFollowerDb);
 
+		// (Optional) Low-pass filter for input signal
+		FBiquadFilter InputLowshelfFilter;
+
+		// (Optional) High-pass filter for input signal
+		FBiquadFilter InputHighshelfFilter;
+
 		EDynamicsProcessingMode::Type ProcessingMode;
 
 		// Lookahead delay lines
@@ -63,7 +86,10 @@ namespace Audio
 		// Points in the knee used for lagrangian interpolation
 		TArray<FVector2D> KneePoints;
 
+		// Channel values of cached detector sample
 		TArray<float> DetectorOuts;
+
+		// Channel values of cached gain sample
 		TArray<float> Gain;
 
 		// How far ahead to look in the audio
@@ -78,7 +104,7 @@ namespace Audio
 		// Amplitude threshold above which gain will be reduced
 		float ThresholdDb;
 
-		// Amount of gain reduction.
+		// Amount of gain reduction
 		float Ratio;
 
 		// Defines how hard or soft the gain reduction blends from no gain reduction to gain reduction (determined by the ratio)
@@ -90,14 +116,25 @@ namespace Audio
 		// Amount of output gain
 		float OutputGain;
 
+		// Gain of key detector signal in dB
+		float KeyGain;
+
 		// Number of channels to use for the dynamics processor
 		int32 NumChannels;
 
-		// Whether or not the stereo channels are linked
-		bool bIsChannelLinked;
+		// Whether or not input channels are linked, and if so, how to calculate gain
+		EDynamicsProcessorChannelLinkMode LinkMode;
 
 		// Whether or not we're in analog mode
 		bool bIsAnalogMode;
 
+		// Whether or not to bypass processor and only output key modulator
+		bool bKeyAuditionEnabled;
+
+		// Whether or not key high-pass filter is enabled
+		bool bKeyHighshelfEnabled;
+
+		// Whether or not key low-pass filter is enabled
+		bool bKeyLowshelfEnabled;
 	};
 }
diff --git a/Engine/Source/Runtime/SignalProcessing/Public/DSP/Filter.h b/Engine/Source/Runtime/SignalProcessing/Public/DSP/Filter.h
--- a/Engine/Source/Runtime/SignalProcessing/Public/DSP/Filter.h
+++ b/Engine/Source/Runtime/SignalProcessing/Public/DSP/Filter.h
@@ -18,6 +18,8 @@ namespace Audio
 			Bandpass,
 			Notch,
 			ParametricEQ,
+			LowShelf,
+			HighShelf,
 			AllPass,
 		};
 	}
