Ensure parent submix structure is registered prior to current submix if missing

diff --git a/Engine/Source/Runtime/AudioMixer/Private/AudioMixerDevice.cpp b/Engine/Source/Runtime/AudioMixer/Private/AudioMixerDevice.cpp
--- a/Engine/Source/Runtime/AudioMixer/Private/AudioMixerDevice.cpp
+++ b/Engine/Source/Runtime/AudioMixer/Private/AudioMixerDevice.cpp
@@ -1001,6 +1001,16 @@ namespace Audio
 
 		if (!bIsMasterSubmix)
 		{
+			// Ensure parent structure is registered prior to current submix if missing
+			if (InSoundSubmix->ParentSubmix)
+			{
+				FMixerSubmixPtr ParentSubmix = GetSubmixInstance(InSoundSubmix->ParentSubmix).Pin();
+				if (!ParentSubmix.IsValid())
+				{
+					RegisterSoundSubmix(InSoundSubmix->ParentSubmix, bInit);
+				}
+			}
+
 			LoadSoundSubmix(*InSoundSubmix);
 		}
 
@@ -1018,10 +1028,8 @@ namespace Audio
 
 	void FMixerDevice::LoadSoundSubmix(USoundSubmix& InSoundSubmix)
 	{
-		// Ensure submix not already registered by first checking master submixes and then additional mixes.
+		// If submix not already found, load it.
 		FMixerSubmixPtr MixerSubmix = GetSubmixInstance(&InSoundSubmix).Pin();
-
-		// If submix not already found, register it.
 		if (!MixerSubmix.IsValid())
 		{
 			MixerSubmix = MakeShared<FMixerSubmix, ESPMode::ThreadSafe>(this);
