Fix strict compile bug from removing monolithic Audio.h include

diff --git a/Engine/Source/Runtime/AudioMixer/Classes/SubmixEffects/AudioMixerSubmixEffectReverbFast.h b/Engine/Source/Runtime/AudioMixer/Classes/SubmixEffects/AudioMixerSubmixEffectReverbFast.h
--- a/Engine/Source/Runtime/AudioMixer/Classes/SubmixEffects/AudioMixerSubmixEffectReverbFast.h
+++ b/Engine/Source/Runtime/AudioMixer/Classes/SubmixEffects/AudioMixerSubmixEffectReverbFast.h
@@ -3,9 +3,10 @@
 #pragma once
 
 #include "AudioEffect.h"
-#include "Sound/SoundEffectSubmix.h"
+#include "Curves/RichCurve.h"
 #include "DSP/Amp.h"
 #include "DSP/ReverbFast.h"
+#include "Sound/SoundEffectSubmix.h"
 
 #include "AudioMixerSubmixEffectReverbFast.generated.h"
 
diff --git a/Engine/Source/Runtime/AudioMixer/Private/AudioMixerEffectsManager.h b/Engine/Source/Runtime/AudioMixer/Private/AudioMixerEffectsManager.h
--- a/Engine/Source/Runtime/AudioMixer/Private/AudioMixerEffectsManager.h
+++ b/Engine/Source/Runtime/AudioMixer/Private/AudioMixerEffectsManager.h
@@ -1,15 +1,17 @@
 // Copyright Epic Games, Inc. All Rights Reserved.
 
 /*=============================================================================
-AudioMixerEffectsManager.h: Implementation of backwarwds compatible effects manager
-for the multiplatform audio mixer
+AudioMixerEffectsManager.h: Implementation of backwards compatible effects
+manager for the multi-platform audio mixer device.
 =============================================================================*/
 
 #pragma once
 
 #include "AudioEffect.h"
+#include "Curves/CurveFloat.h"
 #include "Sound/SoundEffectSubmix.h"
 
+
 namespace Audio
 {
 	class FMixerDevice;
@@ -33,4 +35,4 @@ namespace Audio
 		FSoundEffectSubmixPtr InvalidReverbEffect;
 		FSoundEffectSubmixPtr InvalidEQEffect;
 	};
-}
+} // namespace Audio
