Removing call to SetStereoBleed. The call is not threadsafe with editor shutdown and the result is not used in the new audio engine.

diff --git a/Engine/Source/Runtime/AudioMixer/Private/AudioMixerSource.cpp b/Engine/Source/Runtime/AudioMixer/Private/AudioMixerSource.cpp
--- a/Engine/Source/Runtime/AudioMixer/Private/AudioMixerSource.cpp
+++ b/Engine/Source/Runtime/AudioMixer/Private/AudioMixerSource.cpp
@@ -1024,8 +1024,6 @@ namespace Audio
 
 	void FMixerSource::UpdateChannelMaps()
 	{
-		SetStereoBleed();
-
 		SetLFEBleed();
 
 		int32 NumOutputDeviceChannels = MixerDevice->GetNumDeviceChannels();
