Fix for assret being hit with timesynth when Precache gets called multiple times on the same asset (TimeSynth now performs Precach synchronously)

diff --git a/Engine/Source/Runtime/AudioMixer/Private/SoundWaveDecoder.cpp b/Engine/Source/Runtime/AudioMixer/Private/SoundWaveDecoder.cpp
--- a/Engine/Source/Runtime/AudioMixer/Private/SoundWaveDecoder.cpp
+++ b/Engine/Source/Runtime/AudioMixer/Private/SoundWaveDecoder.cpp
@@ -454,13 +454,7 @@ namespace Audio
 
 		// Start the soundwave precache
 		const ESoundWavePrecacheState PrecacheState = InitData.SoundWave->GetPrecacheState();
-		if (PrecacheState == ESoundWavePrecacheState::NotStarted)
-		{
-			AudioDevice->Precache(InitData.SoundWave);
-			PrecachingSources.Add(InitData.Handle.Id, InitData);
-			return true;
-		}
-		else if (PrecacheState != ESoundWavePrecacheState::Done)
+		if (PrecacheState == ESoundWavePrecacheState::InProgress)
 		{
 			if (!PrecachingSources.Contains(InitData.Handle.Id))
 			{
@@ -470,6 +464,11 @@ namespace Audio
 		}
 		else
 		{
+			if (PrecacheState == ESoundWavePrecacheState::NotStarted)
+			{
+				AudioDevice->Precache(InitData.SoundWave, true);
+			}
+			check(InitData.SoundWave->GetPrecacheState() == ESoundWavePrecacheState::Done);
 			return InitDecodingSourceInternal(InitData);
 		}
 
