Submix sidechaining
- Minor: Add cvar au.DisableAppVolume for debugging convenience

diff --git a/Engine/Source/Runtime/AudioMixer/Classes/SubmixEffects/AudioMixerSubmixEffectDynamicsProcessor.h b/Engine/Source/Runtime/AudioMixer/Classes/SubmixEffects/AudioMixerSubmixEffectDynamicsProcessor.h
--- a/Engine/Source/Runtime/AudioMixer/Classes/SubmixEffects/AudioMixerSubmixEffectDynamicsProcessor.h
+++ b/Engine/Source/Runtime/AudioMixer/Classes/SubmixEffects/AudioMixerSubmixEffectDynamicsProcessor.h
@@ -1,14 +1,21 @@
 // Copyright Epic Games, Inc. All Rights Reserved.
-
 #pragma once
 
+#include "AudioDevice.h"
 #include "DSP/DynamicsProcessor.h"
 #include "Sound/SoundEffectSubmix.h"
 #include "Sound/SoundSubmix.h"
+#include "Sound/SoundSubmixSend.h"
 
 #include "AudioMixerSubmixEffectDynamicsProcessor.generated.h"
 
 
+namespace Audio
+{
+	// Forward Declarations
+	class FMixerDevice;
+}
+
 UENUM(BlueprintType)
 enum class ESubmixEffectDynamicsProcessorType : uint8
 {
@@ -110,8 +117,8 @@ struct AUDIOMIXER_API FSubmixEffectDynamicsProcessorSettings
 	UPROPERTY(EditAnywhere, BlueprintReadWrite, Category = Response, meta = (DisplayName = "Release Time (ms)", ClampMin = "20.0", ClampMax = "5000.0", UIMin = "20.0", UIMax = "5000.0"))
 	float ReleaseTimeMsec;
 
-	// (Coming soon) If set, uses output of provided submix as modulator of input signal for dynamics processor (Uses input signal as default modulator)
-	UPROPERTY(VisibleAnywhere, Category = Sidechain)
+	// If set, uses output of provided submix as modulator of input signal for dynamics processor (Uses input signal as default modulator)
+	UPROPERTY(EditAnywhere, BlueprintReadWrite, Category = Sidechain)
 	USoundSubmix* ExternalSubmix;
 
 	UPROPERTY()
@@ -164,11 +171,15 @@ struct AUDIOMIXER_API FSubmixEffectDynamicsProcessorSettings
 };
 
 
-class AUDIOMIXER_API FSubmixEffectDynamicsProcessor : public FSoundEffectSubmix
+class AUDIOMIXER_API FSubmixEffectDynamicsProcessor :
+	public FSoundEffectSubmix,
+	public ISubmixBufferListener
 {
 public:
 	FSubmixEffectDynamicsProcessor();
 
+	virtual ~FSubmixEffectDynamicsProcessor();
+
 	// Called on an audio effect at initialization on main thread before audio processing begins.
 	virtual void Init(const FSoundEffectSubmixInitData& InSampleRate) override;
 
@@ -179,8 +190,24 @@ public:
 	virtual void OnPresetChanged() override;
 
 protected:
+	// ISubmixBufferListener interface
+	virtual void OnNewSubmixBuffer(
+		const USoundSubmix* InOwningSubmix,
+		float*				InAudioData,
+		int32				InNumSamples,
+		int32				InNumChannels,
+		const int32			InSampleRate,
+		double				InAudioClock) override;
+
+	TWeakObjectPtr<USoundSubmix> ExternalSubmix;
+	Audio::AlignedFloatBuffer AudioExternal;
+
+	TArray<float> AudioKeyFrame;
 	TArray<float> AudioInputFrame;
 	TArray<float> AudioOutputFrame;
+
+	Audio::FDeviceId DeviceId;
+
 	Audio::FDynamicsProcessor DynamicsProcessor;
 };
 
diff --git a/Engine/Source/Runtime/AudioMixer/Private/AudioMixerSubmix.cpp b/Engine/Source/Runtime/AudioMixer/Private/AudioMixerSubmix.cpp
--- a/Engine/Source/Runtime/AudioMixer/Private/AudioMixerSubmix.cpp
+++ b/Engine/Source/Runtime/AudioMixer/Private/AudioMixerSubmix.cpp
@@ -110,6 +110,7 @@ namespace Audio
 				FSoundEffectSubmix* SubmixEffect = static_cast<FSoundEffectSubmix*>(EffectPreset->CreateNewEffect());
 				
 				FSoundEffectSubmixInitData InitData;
+				InitData.DeviceHandle = MixerDevice->DeviceHandle;
 				InitData.SampleRate = MixerDevice->GetSampleRate();
 				InitData.PresetSettings = nullptr;
 
diff --git a/Engine/Source/Runtime/AudioMixer/Private/Effects/AudioMixerSubmixEffectDynamicsProcessor.cpp b/Engine/Source/Runtime/AudioMixer/Private/Effects/AudioMixerSubmixEffectDynamicsProcessor.cpp
--- a/Engine/Source/Runtime/AudioMixer/Private/Effects/AudioMixerSubmixEffectDynamicsProcessor.cpp
+++ b/Engine/Source/Runtime/AudioMixer/Private/Effects/AudioMixerSubmixEffectDynamicsProcessor.cpp
@@ -1,29 +1,68 @@
 // Copyright Epic Games, Inc. All Rights Reserved.
 
 #include "SubmixEffects/AudioMixerSubmixEffectDynamicsProcessor.h"
+
+#include "AudioDeviceManager.h"
+#include "AudioMixerDevice.h"
 #include "ProfilingDebugging/CsvProfiler.h"
 
 // Link to "Audio" profiling category
 CSV_DECLARE_CATEGORY_MODULE_EXTERN(AUDIOMIXERCORE_API, Audio);
 
+namespace
+{
+	// Use a full multichannel frame as a scratch processing buffer
+	static const int32 ProcessorScratchNumChannels = 8;
+} // namespace <>
+
 FSubmixEffectDynamicsProcessor::FSubmixEffectDynamicsProcessor()
 {
 }
 
+FSubmixEffectDynamicsProcessor::~FSubmixEffectDynamicsProcessor()
+{
+	if (ExternalSubmix.IsValid())
+	{
+		Audio::FMixerDevice* MixerDevice = nullptr;
+		if (FAudioDeviceManager* DeviceManager = FAudioDeviceManager::Get())
+		{
+			MixerDevice = static_cast<Audio::FMixerDevice*>(DeviceManager->GetAudioDevice(DeviceId));
+		}
+
+		if (MixerDevice)
+		{
+			MixerDevice->UnregisterSubmixBufferListener(this, ExternalSubmix.Get());
+		}
+	}
+}
+
 void FSubmixEffectDynamicsProcessor::Init(const FSoundEffectSubmixInitData& InitData)
 {
-	// Use a full multichannel frame as a scratch processing buffer
-	static const int32 NumChannels = 8;
+	DynamicsProcessor.Init(InitData.SampleRate, ProcessorScratchNumChannels);
 
-	DynamicsProcessor.Init(InitData.SampleRate, NumChannels);
+	AudioKeyFrame.Reset();
+	AudioKeyFrame.AddZeroed(ProcessorScratchNumChannels);
 
 	AudioInputFrame.Reset();
-	AudioInputFrame.AddZeroed(NumChannels);
+	AudioInputFrame.AddZeroed(ProcessorScratchNumChannels);
 
 	AudioOutputFrame.Reset();
-	AudioOutputFrame.AddZeroed(NumChannels);
+	AudioOutputFrame.AddZeroed(ProcessorScratchNumChannels);
+
+	DeviceId = InitData.DeviceHandle;
 }
 
+void FSubmixEffectDynamicsProcessor::OnNewSubmixBuffer(
+	const USoundSubmix* InOwningSubmix,
+	float*				InAudioData,
+	int32				InNumSamples,
+	int32				InNumChannels,
+	const int32			InSampleRate,
+	double				InAudioClock)
+{
+	AudioExternal.Reset();
+	AudioExternal.Append(InAudioData, InNumSamples);
+}
 
 void FSubmixEffectDynamicsProcessor::OnPresetChanged()
 {
@@ -86,6 +125,30 @@ void FSubmixEffectDynamicsProcessor::OnPresetChanged()
 
 	static_assert(static_cast<int32>(ESubmixEffectDynamicsChannelLinkMode::Count) == static_cast<int32>(Audio::EDynamicsProcessorChannelLinkMode::Count), "Enumerations must match");
 	DynamicsProcessor.SetChannelLinkMode(static_cast<Audio::EDynamicsProcessorChannelLinkMode>(Settings.LinkMode));
+
+	Audio::FMixerDevice* MixerDevice = nullptr;
+	if (FAudioDeviceManager* DeviceManager = FAudioDeviceManager::Get())
+	{
+		MixerDevice = static_cast<Audio::FMixerDevice*>(DeviceManager->GetAudioDevice(DeviceId));
+	}
+
+	if (MixerDevice)
+	{
+		if (ExternalSubmix.Get() != Settings.ExternalSubmix)
+		{
+			if (ExternalSubmix.IsValid())
+			{
+				MixerDevice->UnregisterSubmixBufferListener(this, ExternalSubmix.Get());
+			}
+
+			ExternalSubmix = Settings.ExternalSubmix;
+
+			if (ExternalSubmix.IsValid())
+			{
+				MixerDevice->RegisterSubmixBufferListener(this, ExternalSubmix.Get());
+			}
+		}
+	}
 }
 
 void FSubmixEffectDynamicsProcessor::OnProcessAudio(const FSoundEffectSubmixInputData& InData, FSoundEffectSubmixOutputData& OutData)
@@ -103,10 +166,18 @@ void FSubmixEffectDynamicsProcessor::OnProcessAudio(const FSoundEffectSubmixInpu
 		{
 			const int32 SampleIndex = SampleIndexOfFrame + Channel;
 			AudioInputFrame[Channel] = InBuffer[SampleIndex];
+			if (ExternalSubmix.IsValid())
+			{
+				AudioKeyFrame[Channel] = AudioExternal.Num() > SampleIndex ? AudioExternal[SampleIndex] : 0;
+			}
+			else
+			{
+				AudioKeyFrame[Channel] = InBuffer[SampleIndex];
+			}
 		}
 
 		// Process
-		DynamicsProcessor.ProcessAudio(AudioInputFrame.GetData(), InData.NumChannels, AudioOutputFrame.GetData());
+		DynamicsProcessor.ProcessAudio(AudioInputFrame.GetData(), InData.NumChannels, AudioOutputFrame.GetData(), AudioKeyFrame.GetData());
 
 		// Copy the data to the frame output
 		for (int32 Channel = 0; Channel < InData.NumChannels; ++Channel)
diff --git a/Engine/Source/Runtime/Engine/Classes/Sound/SoundEffectSubmix.h b/Engine/Source/Runtime/Engine/Classes/Sound/SoundEffectSubmix.h
--- a/Engine/Source/Runtime/Engine/Classes/Sound/SoundEffectSubmix.h
+++ b/Engine/Source/Runtime/Engine/Classes/Sound/SoundEffectSubmix.h
@@ -3,13 +3,17 @@
 
 #include "CoreMinimal.h"
 
+#include "AudioDefines.h"
 #include "Sound/SoundEffectPreset.h"
 #include "Sound/SoundEffectBase.h"
 #include "UObject/ObjectMacros.h"
 
 #include "SoundEffectSubmix.generated.h"
 
+
+class FAudioDevice;
 class FSoundEffectSubmix;
+
 struct FAudioEffectParameters;
 
 /** Preset of a submix effect that can be shared between sounds. */
@@ -25,8 +29,16 @@ class ENGINE_API USoundEffectSubmixPreset : public USoundEffectPreset
 /** Struct which has data needed to initialize the submix effect. */
 struct FSoundEffectSubmixInitData
 {
+	uint32 DeviceHandle;
 	void* PresetSettings;
 	float SampleRate;
+
+	FSoundEffectSubmixInitData()
+		: DeviceHandle(INDEX_NONE)
+		, PresetSettings(nullptr)
+		, SampleRate(0)
+	{
+	}
 };
 
 /** Struct which supplies audio data to submix effects on game thread. */
diff --git a/Engine/Source/Runtime/Engine/Private/AudioDevice.cpp b/Engine/Source/Runtime/Engine/Private/AudioDevice.cpp
--- a/Engine/Source/Runtime/Engine/Private/AudioDevice.cpp
+++ b/Engine/Source/Runtime/Engine/Private/AudioDevice.cpp
@@ -2,6 +2,7 @@
 #include "AudioDevice.h"
 
 #include "ActiveSound.h"
+#include "Audio.h"
 #include "AudioCompressionSettingsUtils.h"
 #include "AudioDecompress.h"
 #include "AudioDefines.h"
@@ -78,6 +79,14 @@ FAutoConsoleVariableRef CVarForceRealtimeDecompression(
 	TEXT("0: Allow full decompression on load, 1: force realtime decompression."),
 	ECVF_Default);
 
+static int32 DisableAppVolumeCvar = 0;
+FAutoConsoleVariableRef CVarDisableAppVolume(
+	TEXT("au.DisableAppVolume"),
+	DisableAppVolumeCvar,
+	TEXT("Disables application volume when set to 1.\n")
+	TEXT("0: App volume enabled, 1: App volume disabled"),
+	ECVF_Default);
+
 static int32 DisableAutomaticPrecacheCvar = 0;
 FAutoConsoleVariableRef CVarDisableAutomaticPrecache(
 	TEXT("au.DisableAutomaticPrecache"),
@@ -4151,7 +4160,12 @@ void FAudioDevice::Update(bool bGameTicking)
 	bIsStoppingVoicesEnabled = !DisableStoppingVoicesCvar;
 
 	// Update the master volume
-	MasterVolume = GetTransientMasterVolume() * FApp::GetVolumeMultiplier();
+	MasterVolume = GetTransientMasterVolume();
+	
+	if (!DisableAppVolumeCvar)
+	{
+		MasterVolume *= FApp::GetVolumeMultiplier();
+	}
 
 	UpdateAudioPluginSettingsObjectCache();
 
diff --git a/Engine/Source/Runtime/Engine/Private/AudioDeviceManager.cpp b/Engine/Source/Runtime/Engine/Private/AudioDeviceManager.cpp
--- a/Engine/Source/Runtime/Engine/Private/AudioDeviceManager.cpp
+++ b/Engine/Source/Runtime/Engine/Private/AudioDeviceManager.cpp
@@ -1,7 +1,9 @@
 // Copyright Epic Games, Inc. All Rights Reserved.
 #include "AudioDeviceManager.h"
 
+#include "Audio.h"
 #include "Audio/AudioDebug.h"
+#include "AudioDefines.h"
 #include "AudioDevice.h"
 #include "AudioMixerDevice.h"
 #include "Sound/AudioSettings.h"
@@ -621,6 +623,16 @@ FAudioDevice* FAudioDeviceManager::GetAudioDevice(Audio::FDeviceId Handle)
 	return AudioDevice;
 }
 
+FAudioDeviceManager* FAudioDeviceManager::Get()
+{
+	if (GEngine)
+	{
+		return GEngine->GetAudioDeviceManager();
+	}
+
+	return nullptr;
+}
+
 FAudioDevice* FAudioDeviceManager::GetActiveAudioDevice()
 {
 	if (ActiveAudioDeviceHandle != INDEX_NONE)
diff --git a/Engine/Source/Runtime/Engine/Public/AudioDeviceManager.h b/Engine/Source/Runtime/Engine/Public/AudioDeviceManager.h
--- a/Engine/Source/Runtime/Engine/Public/AudioDeviceManager.h
+++ b/Engine/Source/Runtime/Engine/Public/AudioDeviceManager.h
@@ -1,7 +1,7 @@
 // Copyright Epic Games, Inc. All Rights Reserved.
 #pragma once
 
-#include "Audio.h"
+#include "AudioDefines.h"
 #include "Audio/AudioDebug.h"
 #include "AudioThread.h"
 #include "Containers/Queue.h"
@@ -103,6 +103,8 @@ public:
 	*/
 	FAudioDevice* GetAudioDevice(Audio::FDeviceId Handle);
 
+	static FAudioDeviceManager* Get();
+
 	/**
 	* Returns a ptr to the active audio device. If there is no active
 	* device then it will return the main audio device.
diff --git a/Engine/Source/Runtime/SignalProcessing/Private/DynamicsProcessor.cpp b/Engine/Source/Runtime/SignalProcessing/Private/DynamicsProcessor.cpp
--- a/Engine/Source/Runtime/SignalProcessing/Private/DynamicsProcessor.cpp
+++ b/Engine/Source/Runtime/SignalProcessing/Private/DynamicsProcessor.cpp
@@ -178,13 +178,13 @@ namespace Audio
 		ProcessingMode = InProcessingMode;
 	}
 
-	void FDynamicsProcessor::ProcessAudioFrame(const float* InFrame, float* OutFrame)
+	void FDynamicsProcessor::ProcessAudioFrame(const float* InFrame, float* OutFrame, const float* InKeyFrame)
 	{
 		checkSlow(NumChannels == DetectorOuts.Num());
 		checkSlow(NumChannels == Gain.Num());
 
 		// Get detector outputs
-		const float* DetectorIn = InFrame;
+		const float* DetectorIn = InKeyFrame;
 		if (NumChannels > 0)
 		{
 			if (bKeyLowshelfEnabled)
@@ -274,11 +274,12 @@ namespace Audio
 		}
 	}
 
-	void FDynamicsProcessor::ProcessAudio(const float* InBuffer, const int32 InNumSamples, float* OutBuffer)
+	void FDynamicsProcessor::ProcessAudio(const float* InBuffer, const int32 InNumSamples, float* OutBuffer, const float* InKeyBuffer)
 	{
 		for (int32 SampleIndex = 0; SampleIndex < InNumSamples; SampleIndex += NumChannels)
 		{
-			ProcessAudioFrame(&InBuffer[SampleIndex], &OutBuffer[SampleIndex]);
+			const float* KeyFrame = InKeyBuffer ? &InKeyBuffer[SampleIndex] : &InBuffer[SampleIndex];
+			ProcessAudioFrame(&InBuffer[SampleIndex], &OutBuffer[SampleIndex], KeyFrame);
 		}
 	}
 
@@ -332,6 +333,4 @@ namespace Audio
 		OutputGainDb = FMath::Min(0.0f, OutputGainDb);
 		return ConvertToLinear(OutputGainDb);
 	}
-
-
 }
diff --git a/Engine/Source/Runtime/SignalProcessing/Public/DSP/DynamicsProcessor.h b/Engine/Source/Runtime/SignalProcessing/Public/DSP/DynamicsProcessor.h
--- a/Engine/Source/Runtime/SignalProcessing/Public/DSP/DynamicsProcessor.h
+++ b/Engine/Source/Runtime/SignalProcessing/Public/DSP/DynamicsProcessor.h
@@ -61,8 +61,8 @@ namespace Audio
 		void SetPeakMode(const EPeakMode::Type InEnvelopeFollowerModeType);
 		void SetProcessingMode(const EDynamicsProcessingMode::Type ProcessingMode);
 
-		void ProcessAudioFrame(const float* InFrame, float* OutFrame);
-		void ProcessAudio(const float* InBuffer, const int32 InNumSamples, float* OutBuffer);
+		void ProcessAudioFrame(const float* InFrame, float* OutFrame, const float* InKeyFrame);
+		void ProcessAudio(const float* InBuffer, const int32 InNumSamples, float* OutBuffer, const float* InKeyBuffer = nullptr);
 
 
 	protected:
