Missing initializer

diff --git a/Engine/Source/Runtime/AudioMixer/Private/Effects/AudioMixerSubmixEffectDynamicsProcessor.cpp b/Engine/Source/Runtime/AudioMixer/Private/Effects/AudioMixerSubmixEffectDynamicsProcessor.cpp
--- a/Engine/Source/Runtime/AudioMixer/Private/Effects/AudioMixerSubmixEffectDynamicsProcessor.cpp
+++ b/Engine/Source/Runtime/AudioMixer/Private/Effects/AudioMixerSubmixEffectDynamicsProcessor.cpp
@@ -16,6 +16,7 @@ namespace
 } // namespace <>
 
 FSubmixEffectDynamicsProcessor::FSubmixEffectDynamicsProcessor()
+	: DeviceId(INDEX_NONE)
 {
 }
 
