Fix for ensure that could fire on audio mixer device initialization with the newly RAII'd AudioDeviceManager. FAudioMixerDevice::InitializeHardware has always run on the game thread, but the ensure at the beginning of the function happened to pass because the audio thread hadn't been started yet. This makes that behavior explicit.

diff --git a/Engine/Source/Runtime/AudioMixer/Private/AudioMixerDevice.cpp b/Engine/Source/Runtime/AudioMixer/Private/AudioMixerDevice.cpp
--- a/Engine/Source/Runtime/AudioMixer/Private/AudioMixerDevice.cpp
+++ b/Engine/Source/Runtime/AudioMixer/Private/AudioMixerDevice.cpp
@@ -138,7 +138,7 @@ namespace Audio
 
 	bool FMixerDevice::InitializeHardware()
 	{
-		AUDIO_MIXER_CHECK_GAME_THREAD(this);
+		ensure(IsInGameThread());
 	
 		LLM_SCOPE(ELLMTag::AudioMixer);
 
