Fix for Time Synth crash (race condition with GC).  Sonud Source Decodes for the TimeSynth are now synchronous.

diff --git a/Engine/Source/Runtime/AudioMixer/Private/AudioMixerSourceBuffer.cpp b/Engine/Source/Runtime/AudioMixer/Private/AudioMixerSourceBuffer.cpp
--- a/Engine/Source/Runtime/AudioMixer/Private/AudioMixerSourceBuffer.cpp
+++ b/Engine/Source/Runtime/AudioMixer/Private/AudioMixerSourceBuffer.cpp
@@ -55,7 +55,7 @@ namespace Audio
 		return CurrentSample >= NumSamples;
 	}
 
-	TSharedPtr<FMixerSourceBuffer> FMixerSourceBuffer::Create(FMixerBuffer& InBuffer, USoundWave& InWave, ELoopingMode InLoopingMode, bool bInIsSeeking)
+	TSharedPtr<FMixerSourceBuffer> FMixerSourceBuffer::Create(FMixerBuffer& InBuffer, USoundWave& InWave, ELoopingMode InLoopingMode, bool bInIsSeeking, bool bInForceSyncDecode)
 	{
 		LLM_SCOPE(ELLMTag::AudioMixer);
 
@@ -71,11 +71,11 @@ namespace Audio
 			return nullptr;
 		}
 
-		TSharedPtr<FMixerSourceBuffer> NewSourceBuffer = MakeShareable(new FMixerSourceBuffer(InBuffer, InWave, InLoopingMode, bInIsSeeking));
+		TSharedPtr<FMixerSourceBuffer> NewSourceBuffer = MakeShareable(new FMixerSourceBuffer(InBuffer, InWave, InLoopingMode, bInIsSeeking, bInForceSyncDecode));
 		return NewSourceBuffer;
 	}
 
-	FMixerSourceBuffer::FMixerSourceBuffer(FMixerBuffer& InBuffer, USoundWave& InWave, ELoopingMode InLoopingMode, bool bInIsSeeking)
+	FMixerSourceBuffer::FMixerSourceBuffer(FMixerBuffer& InBuffer, USoundWave& InWave, ELoopingMode InLoopingMode, bool bInIsSeeking, bool bInForceSyncDecode)
 		: NumBuffersQeueued(0)
 		, CurrentBuffer(0)
 		, SoundWave(&InWave)
@@ -92,6 +92,7 @@ namespace Audio
 		, bLoopCallback(false)
 		, bProcedural(InWave.bProcedural)
 		, bIsBus(InWave.bIsBus)
+		, bForceSyncDecode(bInForceSyncDecode)
 	{
 		InWave.AddPlayingSource(this);
 
@@ -343,6 +344,7 @@ namespace Audio
 		NewTaskData.bSkipFirstBuffer = (BufferReadMode == EBufferReadMode::AsynchronousSkipFirstFrame);
 		NewTaskData.NumFramesToDecode = MONO_PCM_BUFFER_SAMPLES;
 		NewTaskData.NumPrecacheFrames = NumPrecacheFrames;
+		NewTaskData.bForceSyncDecode = bForceSyncDecode;
 
 		check(!AsyncRealtimeAudioTask);
 		AsyncRealtimeAudioTask = CreateAudioTask(NewTaskData);
diff --git a/Engine/Source/Runtime/AudioMixer/Private/AudioMixerSourceBuffer.h b/Engine/Source/Runtime/AudioMixer/Private/AudioMixerSourceBuffer.h
--- a/Engine/Source/Runtime/AudioMixer/Private/AudioMixerSourceBuffer.h
+++ b/Engine/Source/Runtime/AudioMixer/Private/AudioMixerSourceBuffer.h
@@ -49,7 +49,7 @@ namespace Audio
 	class FMixerSourceBuffer : public ISoundWaveClient
 	{
 	public:		
-		static FMixerSourceBufferPtr Create(FMixerBuffer& InBuffer, USoundWave& InWave, ELoopingMode InLoopingMode, bool bInIsSeeking);
+		static FMixerSourceBufferPtr Create(FMixerBuffer& InBuffer, USoundWave& InWave, ELoopingMode InLoopingMode, bool bInIsSeeking, bool bInForceSyncDecode = false);
 
 		~FMixerSourceBuffer();
 
@@ -96,7 +96,7 @@ namespace Audio
 		void OnEndGenerate();
 		void ClearWave() { SoundWave = nullptr; }
 	private:
-		FMixerSourceBuffer(FMixerBuffer& InBuffer, USoundWave& InWave, ELoopingMode InLoopingMode, bool bInIsSeeking);
+		FMixerSourceBuffer(FMixerBuffer& InBuffer, USoundWave& InWave, ELoopingMode InLoopingMode, bool bInIsSeeking, bool bInForceSyncDecode = false);
 
 		void SubmitInitialPCMBuffers();
 		void SubmitInitialRealtimeBuffers();
@@ -128,6 +128,7 @@ namespace Audio
 		uint32 bLoopCallback : 1;
 		uint32 bProcedural : 1;
 		uint32 bIsBus : 1;
+		uint32 bForceSyncDecode : 1;
 	
 		virtual void OnBeginDestroy(class USoundWave* Wave) override;
 		virtual bool OnIsReadyForFinishDestroy(class USoundWave* Wave) const override;
diff --git a/Engine/Source/Runtime/AudioMixer/Private/AudioMixerSourceDecode.cpp b/Engine/Source/Runtime/AudioMixer/Private/AudioMixerSourceDecode.cpp
--- a/Engine/Source/Runtime/AudioMixer/Private/AudioMixerSourceDecode.cpp
+++ b/Engine/Source/Runtime/AudioMixer/Private/AudioMixerSourceDecode.cpp
@@ -249,7 +249,7 @@ public:
 	FProceduralDecodeHandle(const FProceduralAudioTaskData& InJobData)
 	{
 		Task = new FAsyncTask<FAsyncDecodeWorker>(InJobData);
-        if (ForceSyncAudioDecodesCvar)
+        if (ForceSyncAudioDecodesCvar || InJobData.bForceSyncDecode)
         {
             Task->StartSynchronousTask();
             return;
@@ -277,7 +277,7 @@ public:
 	FDecodeHandle(const FDecodeAudioTaskData& InJobData)
 	{
 		Task = new FAsyncTask<FAsyncDecodeWorker>(InJobData);
-        if (ForceSyncAudioDecodesCvar)
+        if (ForceSyncAudioDecodesCvar || InJobData.bForceSyncDecode)
         {
             Task->StartSynchronousTask();
             return;
diff --git a/Engine/Source/Runtime/AudioMixer/Private/AudioMixerSourceDecode.h b/Engine/Source/Runtime/AudioMixer/Private/AudioMixerSourceDecode.h
--- a/Engine/Source/Runtime/AudioMixer/Private/AudioMixerSourceDecode.h
+++ b/Engine/Source/Runtime/AudioMixer/Private/AudioMixerSourceDecode.h
@@ -29,11 +29,15 @@ namespace Audio
 		// The number of channels of the procedural buffer
 		int32 NumChannels;
 
+		// Force decodes to execute synchronously
+		bool bForceSyncDecode;
+
 		FProceduralAudioTaskData()
 			: ProceduralSoundWave(nullptr)
 			, AudioData(nullptr)
 			, NumSamples(0)
 			, NumChannels(0)
+			, bForceSyncDecode(false)
 		{}
 	};
 
@@ -64,6 +68,10 @@ namespace Audio
 		// Whether or not to skip the first buffer
 		bool bSkipFirstBuffer;
 
+		// Force this decoding operation to occur synchronously,
+		// regardless of the value of au.ForceSyncAudioDecodes. (used by time synth)
+		bool bForceSyncDecode;
+
 		FDecodeAudioTaskData()
 			: AudioData(nullptr)
 			, DecompressionState(nullptr)
@@ -73,6 +81,7 @@ namespace Audio
 			, NumFramesToDecode(0)
 			, bLoopingMode(false)
 			, bSkipFirstBuffer(false)
+			, bForceSyncDecode(false)
 		{}
 	};
 
diff --git a/Engine/Source/Runtime/AudioMixer/Private/SoundWaveDecoder.cpp b/Engine/Source/Runtime/AudioMixer/Private/SoundWaveDecoder.cpp
--- a/Engine/Source/Runtime/AudioMixer/Private/SoundWaveDecoder.cpp
+++ b/Engine/Source/Runtime/AudioMixer/Private/SoundWaveDecoder.cpp
@@ -12,6 +12,7 @@ namespace Audio
 		, MixerBuffer(nullptr)
 		, SampleRate(INDEX_NONE)
 		, SeekTime(InitData.SeekTime)
+		, bForceSyncDecode(InitData.bForceSyncDecode)
 	{
 		SourceInfo.VolumeParam.Init();
 		SourceInfo.VolumeParam.SetValue(InitData.VolumeScale);
@@ -45,7 +46,7 @@ namespace Audio
 		const ELoopingMode LoopingMode = SoundWave->bLooping ? ELoopingMode::LOOP_Forever : ELoopingMode::LOOP_Never;
 		const bool bIsSeeking = SeekTime > 0.0f;
 
-		MixerSourceBuffer = FMixerSourceBuffer::Create(*MixerBuffer, *SoundWave, LoopingMode, bIsSeeking);
+		MixerSourceBuffer = FMixerSourceBuffer::Create(*MixerBuffer, *SoundWave, LoopingMode, bIsSeeking, bForceSyncDecode);
 		return MixerSourceBuffer.IsValid();
 	}
 
@@ -141,6 +142,11 @@ namespace Audio
 		SourceInfo.VolumeResetFrame = SourceInfo.NumFramesGenerated + NumFrames;
 	}
 
+	void FDecodingSoundSource::SetForceSyncDecode(bool bShouldForceSyncDecode)
+	{
+		bForceSyncDecode = bShouldForceSyncDecode;
+	}
+
 	void FDecodingSoundSource::ReadFrame()
 	{
 		bool bNextFrameOutOfRange = (SourceInfo.CurrentFrameIndex + 1) >= SourceInfo.CurrentAudioChunkNumFrames;
@@ -355,6 +361,11 @@ namespace Audio
 		return true;
 	}
 
+	void FDecodingSoundSource::AddReferencedObjects(FReferenceCollector & Collector)
+	{
+		Collector.AddReferencedObject(SoundWave);
+	}
+
 	FSoundSourceDecoder::FSoundSourceDecoder()
 		: AudioThreadId(0)
 		, AudioDevice(nullptr)
@@ -364,7 +375,7 @@ namespace Audio
 
 	FSoundSourceDecoder::~FSoundSourceDecoder()
 	{
-
+		
 	}
 
 	void FSoundSourceDecoder::Init(FAudioDevice* InAudioDevice, int32 InSampleRate)
@@ -405,6 +416,7 @@ namespace Audio
 
 		if (DecodingSoundWaveDataPtr->PreInit(SampleRate))
 		{
+			DecodingSoundWaveDataPtr->SetForceSyncDecode(InitData.bForceSyncDecode);
 			InitializingDecodingSources.Add(InitData.Handle.Id, DecodingSoundWaveDataPtr);
 
 			// Add this decoding sound wave to a data structure we can access safely from audio render thread
diff --git a/Engine/Source/Runtime/AudioMixer/Public/SoundWaveDecoder.h b/Engine/Source/Runtime/AudioMixer/Public/SoundWaveDecoder.h
--- a/Engine/Source/Runtime/AudioMixer/Public/SoundWaveDecoder.h
+++ b/Engine/Source/Runtime/AudioMixer/Public/SoundWaveDecoder.h
@@ -33,16 +33,18 @@ namespace Audio
 			, SeekTime(0.0f)
 			, PitchScale(1.0f)
 			, VolumeScale(1.0f)
+			, bForceSyncDecode(false)
 		{}
 
 		USoundWave* SoundWave;
 		float SeekTime;
 		float PitchScale;
 		float VolumeScale;
+		bool bForceSyncDecode;
 		FDecodingSoundSourceHandle Handle;
 	};
 
-	class FDecodingSoundSource
+	class FDecodingSoundSource : public FGCObject
 	{
 	public:
 		FDecodingSoundSource(FAudioDevice* AudioDevice, const FSourceDecodeInit& InitData);
@@ -69,9 +71,15 @@ namespace Audio
 		// Sets the volume scale
 		void SetVolumeScale(float InVolumeScale, uint32 NumFrames = 512);
 
+		// Sets the ForceSyncDecode flag. (Decodes for this soundwave will not happen in an async task if true)
+		void SetForceSyncDecode(bool bShouldForceSyncDecode);
+
 		// Get audio buffer
 		bool GetAudioBuffer(const int32 InNumFrames, const int32 InNumChannels, AlignedFloatBuffer& OutAudioBuffer);
 
+		// Override for FGCObject:
+		virtual void AddReferencedObjects(FReferenceCollector & Collector) override;
+
 	private:
 
 		void ReadFrame();
@@ -83,7 +91,7 @@ namespace Audio
 		// The sound wave object with which this sound is generating
 		USoundWave* SoundWave;
 
-		// Mixer buffer object which is a conveince wrapper around some buffer initialization and management
+		// Mixer buffer object which is a convenience wrapper around some buffer initialization and management
 		FMixerBuffer* MixerBuffer;
 
 		// Object which handles bulk of decoding operations
@@ -98,6 +106,9 @@ namespace Audio
 		// If we've initialized	
 		FThreadSafeBool bInitialized;
 
+		// force the decoding of this sound to be synchronous
+		bool bForceSyncDecode{false};
+
 		// Object used for source generation from decoded buffers
 		struct FSourceInfo
 		{
diff --git a/Engine/Source/Runtime/Engine/Private/SoundWave.cpp b/Engine/Source/Runtime/Engine/Private/SoundWave.cpp
--- a/Engine/Source/Runtime/Engine/Private/SoundWave.cpp
+++ b/Engine/Source/Runtime/Engine/Private/SoundWave.cpp
@@ -1618,9 +1618,12 @@ bool USoundWave::IsReadyForFinishDestroy()
 	{
 		FScopeLock Lock(&SourcesPlayingCs);
 		
-		for (ISoundWaveClient* i : SourcesPlaying)
+		for (ISoundWaveClient* SoundWaveClientPtr : SourcesPlaying)
 		{
-			i->OnIsReadyForFinishDestroy(this);
+			if(SoundWaveClientPtr)
+			{
+				SoundWaveClientPtr->OnIsReadyForFinishDestroy(this);
+			}
 		}
 	}
 
