Removing private header includes in AudioMixer module

diff --git a/Engine/Source/Runtime/AudioMixer/AudioMixer.Build.cs b/Engine/Source/Runtime/AudioMixer/AudioMixer.Build.cs
--- a/Engine/Source/Runtime/AudioMixer/AudioMixer.Build.cs
+++ b/Engine/Source/Runtime/AudioMixer/AudioMixer.Build.cs
@@ -18,9 +18,6 @@ namespace UnrealBuildTool.Rules
 				}
 			);
 
-            PublicIncludePaths.Add("Runtime/AudioMixer/Private");
-
-
 			PublicDependencyModuleNames.AddRange(
 				new string[]
 				{
diff --git a/Engine/Source/Runtime/AudioMixer/Private/AudioMixerDevice.cpp b/Engine/Source/Runtime/AudioMixer/Private/AudioMixerDevice.cpp
--- a/Engine/Source/Runtime/AudioMixer/Private/AudioMixerDevice.cpp
+++ b/Engine/Source/Runtime/AudioMixer/Private/AudioMixerDevice.cpp
@@ -3,6 +3,7 @@
 #include "AudioMixerDevice.h"
 
 #include "AudioMixerSource.h"
+#include "AudioMixerSourceManager.h"
 #include "AudioMixerSubmix.h"
 #include "AudioMixerSourceVoice.h"
 #include "AudioPluginUtilities.h"
@@ -46,7 +47,6 @@ namespace Audio
 		, AudioClockDelta(0.0)
 		, AudioClock(0.0)
 		, PreviousMasterVolume((float)INDEX_NONE)
-		, SourceManager(this)
 		, GameOrAudioThreadId(INDEX_NONE)
 		, AudioPlatformThreadId(INDEX_NONE)
 		, bDebugOutputEnabled(false)
@@ -54,6 +54,8 @@ namespace Audio
 	{
 		// This audio device is the audio mixer
 		bAudioMixerModuleLoaded = true;
+
+		SourceManager = MakeUnique<FMixerSourceManager>(this);
 	}
 
 	FMixerDevice::~FMixerDevice()
@@ -83,7 +85,7 @@ namespace Audio
 			ListenerTransforms.Add(Listener.Transform);
 		}
 
-		SourceManager.SetListenerTransforms(ListenerTransforms);
+		SourceManager->SetListenerTransforms(ListenerTransforms);
 	}
 
 	void FMixerDevice::ResetAudioRenderingThreadId()
@@ -213,7 +215,7 @@ namespace Audio
 				// TODO: Migrate this to project settings properly
 				SourceManagerInitParams.NumSourceWorkers = 4;
 
-				SourceManager.Init(SourceManagerInitParams);
+				SourceManager->Init(SourceManagerInitParams);
 
 				AudioClock = 0.0;
 				AudioClockDelta = (double)OpenStreamParams.NumFrames / OpenStreamParams.SampleRate;
@@ -304,7 +306,7 @@ namespace Audio
 
 		if (AudioMixerPlatform)
 		{
-			SourceManager.Update();
+			SourceManager->Update();
 
 			AudioMixerPlatform->UnregisterDeviceChangedListener();
 			AudioMixerPlatform->StopAudioStream();
@@ -349,7 +351,7 @@ namespace Audio
 			PanningMethod = AudioSettings->PanningMethod;
 		}
 
-		SourceManager.Update();
+		SourceManager->Update();
 
 		AudioMixerPlatform->OnHardwareUpdate();
 
@@ -362,7 +364,7 @@ namespace Audio
 			InitializeChannelAzimuthMap(PlatformInfo.NumChannels);
 
 			// Update the channel device count in case it changed
-			SourceManager.UpdateDeviceChannelCount(PlatformInfo.NumChannels);
+			SourceManager->UpdateDeviceChannelCount(PlatformInfo.NumChannels);
 
 			// Audio rendering was suspended in CheckAudioDeviceChange if it changed.
 			AudioMixerPlatform->ResumePlaybackOnNewDevice();
@@ -510,7 +512,7 @@ namespace Audio
 		PumpCommandQueue();
 
 		// Compute the next block of audio in the source manager
-		SourceManager.ComputeNextBlockOfSamples();
+		SourceManager->ComputeNextBlockOfSamples();
 
 		FMixerSubmixWeakPtr MasterSubmix = GetMasterSubmix();
 
@@ -526,7 +528,7 @@ namespace Audio
 		}
 
 		// Reset stopping sounds and clear their state after submixes have been mixed
-		SourceManager.ClearStoppingSounds();
+		SourceManager->ClearStoppingSounds();
 
 		// Do any debug output performing
 		if (bDebugOutputEnabled)
@@ -544,11 +546,11 @@ namespace Audio
 	{
 		// Make sure the source manager pumps any final commands on shutdown. These allow for cleaning up sources, interfacing with plugins, etc.
 		// Because we double buffer our command queues, we call this function twice to ensure all commands are successfully pumped.
-		SourceManager.PumpCommandQueue();
-		SourceManager.PumpCommandQueue();
+		SourceManager->PumpCommandQueue();
+		SourceManager->PumpCommandQueue();
 
 		// Make sure we force any pending release data to happen on shutdown
-		SourceManager.UpdatePendingReleaseData(true);
+		SourceManager->UpdatePendingReleaseData(true);
 	}
 
 	void FMixerDevice::LoadMasterSoundSubmix(EMasterSubmixType::Type InType, const FString& InDefaultName, bool bInDefaultMuteWhenBackgrounded, FSoftObjectPath& InObjectPath)
@@ -842,7 +844,7 @@ namespace Audio
 
 	void FMixerDevice::UpdateModulationControls(const uint32 InSourceId, const FSoundModulationControls& InControls)
 	{
-		SourceManager.UpdateModulationControls(InSourceId, InControls);
+		SourceManager->UpdateModulationControls(InSourceId, InControls);
 	}
 
 	void FMixerDevice::UpdateSourceEffectChain(const uint32 SourceEffectChainId, const TArray<FSourceEffectChainEntry>& SourceEffectChain, const bool bPlayEffectChainTails)
@@ -857,7 +859,7 @@ namespace Audio
 			SourceEffectChainOverrides.Add(SourceEffectChainId, SourceEffectChain);
 		}
 
-		SourceManager.UpdateSourceEffectChain(SourceEffectChainId, SourceEffectChain, bPlayEffectChainTails);
+		SourceManager->UpdateSourceEffectChain(SourceEffectChainId, SourceEffectChain, bPlayEffectChainTails);
 	}
 
 	void FMixerDevice::UpdateSubmixProperties(USoundSubmix* InSoundSubmix)
@@ -934,11 +936,11 @@ namespace Audio
 	{
 		if (IsInitialized() && (FPlatformProcess::SupportsMultithreading() && !AudioMixerPlatform->IsNonRealtime()))
 		{
-			SourceManager.FlushCommandQueue(bPumpSynchronously);
+			SourceManager->FlushCommandQueue(bPumpSynchronously);
 		}
 		else if (AudioMixerPlatform->IsNonRealtime())
 		{
-			SourceManager.FlushCommandQueue(true);
+			SourceManager->FlushCommandQueue(true);
 		}
 		else
 		{
@@ -946,10 +948,10 @@ namespace Audio
 			PumpCommandQueue();
 
 			// And also directly pump the source manager command queue
-			SourceManager.PumpCommandQueue();
-			SourceManager.PumpCommandQueue();
+			SourceManager->PumpCommandQueue();
+			SourceManager->PumpCommandQueue();
 
-			SourceManager.UpdatePendingReleaseData(true);
+			SourceManager->UpdatePendingReleaseData(true);
 		}
 	}
 
@@ -1137,7 +1139,7 @@ namespace Audio
 
 	int32 FMixerDevice::GetNumActiveSources() const
 	{
-		return SourceManager.GetNumActiveSources();
+		return SourceManager->GetNumActiveSources();
 	}
 
 	void FMixerDevice::Get3DChannelMap(const ESubmixChannelFormat InSubmixType, const FWaveInstance* InWaveInstance, float EmitterAzimith, float NormalizedOmniRadius, Audio::AlignedFloatBuffer& OutChannelMap)
@@ -1295,7 +1297,7 @@ namespace Audio
 
 	const TArray<FTransform>* FMixerDevice::GetListenerTransforms()
 	{
-		return SourceManager.GetListenerTransforms();
+		return SourceManager->GetListenerTransforms();
 	}
 
 	void FMixerDevice::StartRecording(USoundSubmix* InSubmix, float ExpectedRecordingDuration)
@@ -1683,7 +1685,7 @@ namespace Audio
 
 	FMixerSourceManager* FMixerDevice::GetSourceManager()
 	{
-		return &SourceManager;
+		return SourceManager.Get();
 	}
 
 	bool FMixerDevice::IsMainAudioDevice() const
diff --git a/Engine/Source/Runtime/AudioMixer/Private/AudioMixerEffectsManager.cpp b/Engine/Source/Runtime/AudioMixer/Private/AudioMixerEffectsManager.cpp
--- a/Engine/Source/Runtime/AudioMixer/Private/AudioMixerEffectsManager.cpp
+++ b/Engine/Source/Runtime/AudioMixer/Private/AudioMixerEffectsManager.cpp
@@ -2,14 +2,13 @@
 
 #include "AudioMixerEffectsManager.h"
 #include "AudioMixerDevice.h"
-#include "AudioMixerEffectsManager.h"
+#include "AudioMixerSubmix.h"
 #include "SubmixEffects/AudioMixerSubmixEffectReverb.h"
 #include "SubmixEffects/AudioMixerSubmixEffectReverbFast.h"
 #include "SubmixEffects/AudioMixerSubmixEffectEQ.h"
 
 namespace Audio
 {
-
 	FAudioMixerEffectsManager::FAudioMixerEffectsManager(FAudioDevice* InDevice)
 		: FAudioEffectsManager(InDevice)
 	{
diff --git a/Engine/Source/Runtime/AudioMixer/Private/AudioMixerEffectsManager.h b/Engine/Source/Runtime/AudioMixer/Private/AudioMixerEffectsManager.h
--- a/Engine/Source/Runtime/AudioMixer/Private/AudioMixerEffectsManager.h
+++ b/Engine/Source/Runtime/AudioMixer/Private/AudioMixerEffectsManager.h
@@ -14,8 +14,6 @@ manager for the multi-platform audio mixer device.
 
 namespace Audio
 {
-	class FMixerDevice;
-
 	class FAudioMixerEffectsManager : public FAudioEffectsManager
 	{
 	public:
diff --git a/Engine/Source/Runtime/AudioMixer/Private/AudioMixerSourceManager.h b/Engine/Source/Runtime/AudioMixer/Private/AudioMixerSourceManager.h
--- a/Engine/Source/Runtime/AudioMixer/Private/AudioMixerSourceManager.h
+++ b/Engine/Source/Runtime/AudioMixer/Private/AudioMixerSourceManager.h
@@ -9,6 +9,7 @@
 #include "AudioMixerBuffer.h"
 #include "AudioMixerSubmix.h"
 #include "AudioMixerBus.h"
+#include "AudioMixerDevice.h"
 #include "DSP/InterpolatedOnePole.h"
 #include "DSP/Filter.h"
 #include "DSP/InterpolatedOnePole.h"
@@ -16,7 +17,6 @@
 #include "DSP/ParamInterpolator.h"
 #include "DSP/BufferVectorOperations.h"
 #include "IAudioExtensionPlugin.h"
-#include "Containers/Queue.h"
 
 
 namespace Audio
@@ -45,8 +45,6 @@ namespace Audio
 		}
 	};
 
-	typedef TSharedPtr<FMixerSubmix, ESPMode::ThreadSafe> FMixerSubmixPtr;
-	typedef TWeakPtr<FMixerSubmix, ESPMode::ThreadSafe> FMixerSubmixWeakPtr;
 
 	class ISourceListener
 	{
diff --git a/Engine/Source/Runtime/AudioMixer/Private/AudioMixerSourceVoice.h b/Engine/Source/Runtime/AudioMixer/Private/AudioMixerSourceVoice.h
--- a/Engine/Source/Runtime/AudioMixer/Private/AudioMixerSourceVoice.h
+++ b/Engine/Source/Runtime/AudioMixer/Private/AudioMixerSourceVoice.h
@@ -2,6 +2,7 @@
 
 #pragma once
 
+#include "AudioMixerDevice.h"
 #include "AudioMixerBuffer.h"
 #include "AudioMixerSourceManager.h"
 
diff --git a/Engine/Source/Runtime/AudioMixer/Private/SoundWaveDecoder.cpp b/Engine/Source/Runtime/AudioMixer/Private/SoundWaveDecoder.cpp
--- a/Engine/Source/Runtime/AudioMixer/Private/SoundWaveDecoder.cpp
+++ b/Engine/Source/Runtime/AudioMixer/Private/SoundWaveDecoder.cpp
@@ -1,8 +1,11 @@
 // Copyright Epic Games, Inc. All Rights Reserved.
 
 #include "SoundWaveDecoder.h"
-#include "Engine/Public/AudioThread.h"
+#include "AudioThread.h"
+#include "AudioDecompress.h"
 #include "AudioMixer.h"
+#include "AudioMixerBuffer.h"
+#include "AudioMixerSourceBuffer.h"
 
 namespace Audio
 {
diff --git a/Engine/Source/Runtime/AudioMixer/Public/AudioMixerDevice.h b/Engine/Source/Runtime/AudioMixer/Public/AudioMixerDevice.h
--- a/Engine/Source/Runtime/AudioMixer/Public/AudioMixerDevice.h
+++ b/Engine/Source/Runtime/AudioMixer/Public/AudioMixerDevice.h
@@ -4,18 +4,24 @@
 
 #include "Audio.h"
 #include "AudioMixer.h"
-#include "AudioMixerSourceManager.h"
 #include "AudioDevice.h"
+#include "Sound/SoundSubmix.h"
+#include "DSP/BufferVectorOperations.h"
 
 
 // Forward Declarations
 class FOnSubmixEnvelopeBP;
 class IAudioMixerPlatformInterface;
-class USoundSubmix;
 
 namespace Audio
 {
+	// Audio Namespace Forward Declarations
+	class FMixerSourceManager;
 	class FMixerSourceVoice;
+	class FMixerSubmix;
+
+	typedef TSharedPtr<FMixerSubmix, ESPMode::ThreadSafe> FMixerSubmixPtr;
+	typedef TWeakPtr<FMixerSubmix, ESPMode::ThreadSafe> FMixerSubmixWeakPtr;
 
 	struct FChannelPositionInfo
 	{
@@ -306,7 +312,7 @@ namespace Audio
 		TMap<uint32, TArray<FSourceEffectChainEntry>> SourceEffectChainOverrides;
 
 		/** The mixer source manager. */
-		FMixerSourceManager SourceManager;
+		TUniquePtr<FMixerSourceManager> SourceManager;
 
 		/** ThreadId for the game thread (or if audio is running a separate thread, that ID) */
 		mutable int32 GameOrAudioThreadId;
diff --git a/Engine/Source/Runtime/AudioMixer/Public/SoundWaveDecoder.h b/Engine/Source/Runtime/AudioMixer/Public/SoundWaveDecoder.h
--- a/Engine/Source/Runtime/AudioMixer/Public/SoundWaveDecoder.h
+++ b/Engine/Source/Runtime/AudioMixer/Public/SoundWaveDecoder.h
@@ -5,19 +5,22 @@
 #include "CoreMinimal.h"
 #include "Engine/Public/AudioDevice.h"
 #include "Engine/Classes/Sound/SoundWave.h"
-#include "AudioMixer/Private/AudioMixerBuffer.h"
-#include "AudioMixer/Private/AudioMixerSourceBuffer.h"
 #include "DSP/SinOsc.h"
 #include "DSP/ParamInterpolator.h"
 #include "Misc/ScopeLock.h"
 
 
+
 class FAudioDevice;
 
 #define AUDIO_SOURCE_DECODER_DEBUG 0
 
 namespace Audio
 {
+	class FMixerBuffer;
+	class FMixerSourceBuffer;
+	struct FMixerSourceVoiceBuffer;
+
 	struct FDecodingSoundSourceHandle
 	{
 		FDecodingSoundSourceHandle()
