Speculative fix for crash in automated test. (deref of pointer nulled by GC)

diff --git a/Engine/Source/Runtime/AudioMixer/Private/AudioMixerSource.cpp b/Engine/Source/Runtime/AudioMixer/Private/AudioMixerSource.cpp
--- a/Engine/Source/Runtime/AudioMixer/Private/AudioMixerSource.cpp
+++ b/Engine/Source/Runtime/AudioMixer/Private/AudioMixerSource.cpp
@@ -352,6 +352,13 @@ namespace Audio
 			return;
 		}
 
+		// if MarkPendingKill() was called, WaveInstance->WaveData is null
+		if (!WaveInstance->WaveData)
+		{
+			StopNow();
+			return;
+		}
+
 		++TickCount;
 
 		UpdateModulation();
@@ -593,7 +600,13 @@ namespace Audio
 			// StopFade will stop a sound with a very small fade to avoid discontinuities
 			if (MixerSourceVoice && Playing)
 			{
-				if (bIsStoppingVoicesEnabled && !WaveInstance->WaveData->bProcedural)
+				// if MarkPendingKill() was called, WaveInstance->WaveData is null
+				if (!WaveInstance || !WaveInstance->WaveData)
+				{
+					StopNow();
+					return;
+				}
+				else if (bIsStoppingVoicesEnabled && !WaveInstance->WaveData->bProcedural)
 				{
 					// Let the wave instance know it's stopping
 					WaveInstance->SetStopping(true);
