Editgrate 4.25 audio features from project stream:

diff --git a/Engine/Source/Runtime/AudioMixer/Private/AudioMixerDevice.cpp b/Engine/Source/Runtime/AudioMixer/Private/AudioMixerDevice.cpp
--- a/Engine/Source/Runtime/AudioMixer/Private/AudioMixerDevice.cpp
+++ b/Engine/Source/Runtime/AudioMixer/Private/AudioMixerDevice.cpp
@@ -22,6 +22,7 @@
 #include "Runtime/HeadMountedDisplay/Public/IHeadMountedDisplayModule.h"
 #include "Misc/App.h"
 #include "ProfilingDebugging/CsvProfiler.h"
+#include "IAssetRegistry.h"
 #include "Async/Async.h"
 
 #if WITH_EDITOR
@@ -114,6 +115,89 @@ namespace Audio
 		return CurrentThreadId == AudioPlatformThreadId;
 	}
 
+	TArray<Audio::FChannelPositionInfo>* FMixerDevice::GetDefaultPositionMap(int32 NumChannels)
+	{
+		const Audio::FChannelPositionInfo* SpeakerPositions = GetDefaultChannelPositions();
+
+		if (!SpeakerPositions) // speaker maps are not yet initialized
+		{
+			return nullptr;
+		}
+
+		switch (NumChannels)
+		{
+			// Mono speaker directly in front of listener:
+			case 1:
+			{
+				// force angle on single channel if we are mono
+				static TArray<Audio::FChannelPositionInfo> MonoMap = { {EAudioMixerChannel::FrontCenter, 0, 0} };
+				return &MonoMap;
+			}
+
+			// Stereo speakers to front left and right of listener:
+			case 2:
+			{
+				static TArray<Audio::FChannelPositionInfo> StereoMap = { SpeakerPositions[EAudioMixerChannel::FrontLeft], SpeakerPositions[EAudioMixerChannel::FrontRight] };
+				return &StereoMap;
+			}
+
+			// Quadrophonic speakers at each corner.
+			case 4:
+			{
+				static TArray<Audio::FChannelPositionInfo> QuadMap = {
+														 SpeakerPositions[EAudioMixerChannel::FrontLeft] //left
+														,SpeakerPositions[EAudioMixerChannel::FrontRight] // right
+														,SpeakerPositions[EAudioMixerChannel::SideLeft] //Left Surround
+														,SpeakerPositions[EAudioMixerChannel::SideRight] //Right Surround
+				};
+				return &QuadMap;
+			}
+
+			// 5.1 speakers.
+			case 6:
+			{
+				static TArray<Audio::FChannelPositionInfo> FiveDotOneMap = {
+														 SpeakerPositions[EAudioMixerChannel::FrontLeft] //left
+														,SpeakerPositions[EAudioMixerChannel::FrontRight] // right
+														,SpeakerPositions[EAudioMixerChannel::FrontCenter] //center
+														,SpeakerPositions[EAudioMixerChannel::LowFrequency] //LFE
+														,SpeakerPositions[EAudioMixerChannel::SideLeft] //Left Rear
+														,SpeakerPositions[EAudioMixerChannel::SideRight] //Right Rear
+				};
+				return &FiveDotOneMap;
+			}
+
+			// 7.1 speakers.
+			case 8:
+			{
+				static TArray<Audio::FChannelPositionInfo> SevenDotOneMap = {
+														 SpeakerPositions[EAudioMixerChannel::FrontLeft] // left
+														,SpeakerPositions[EAudioMixerChannel::FrontRight] // right
+														,SpeakerPositions[EAudioMixerChannel::FrontCenter] //center
+														,SpeakerPositions[EAudioMixerChannel::LowFrequency] //LFE
+														,SpeakerPositions[EAudioMixerChannel::BackLeft] // Left Rear
+														,SpeakerPositions[EAudioMixerChannel::BackRight] // Right Rear
+														,SpeakerPositions[EAudioMixerChannel::SideLeft] // Left Surround
+														,SpeakerPositions[EAudioMixerChannel::SideRight] // Right Surround
+				};
+
+
+				return &SevenDotOneMap;
+			}
+
+			case 0:
+			default:
+			{
+				return nullptr;
+			}
+		}
+	}
+
+	bool FMixerDevice::IsEndpointSubmix(const USoundSubmixBase* InSubmix)
+	{
+		return InSubmix && (InSubmix->IsA<UEndpointSubmix>() || InSubmix->IsA<USoundfieldEndpointSubmix>());
+	}
+
 	void FMixerDevice::UpdateDeviceDeltaTime()
 	{
 		DeviceDeltaTime = GetGameDeltaTime();
@@ -232,17 +316,6 @@ namespace Audio
 					SpatializationPluginInterface->Initialize(PluginInitializationParams);
 				}
 
-				// Create a new ambisonics mixer.
-				IAudioSpatializationFactory* SpatializationPluginFactory = AudioPluginUtilities::GetDesiredSpatializationPlugin();
-				if (SpatializationPluginFactory != nullptr)
-				{
-					AmbisonicsMixer = SpatializationPluginFactory->CreateNewAmbisonicsMixer(this);
-					if (AmbisonicsMixer.IsValid())
-					{
-						AmbisonicsMixer->Initialize(PluginInitializationParams);
-					}
-				}
-
 				if (OcclusionInterface.IsValid())
 				{
 					OcclusionInterface->Initialize(PluginInitializationParams);
@@ -317,11 +390,6 @@ namespace Audio
 		// Reset existing submixes if they exist
 		MasterSubmixInstances.Reset();
 		Submixes.Reset();
-
-		if (AmbisonicsMixer.IsValid())
-		{
-			AmbisonicsMixer->Shutdown();
-		}
 	}
 
 	void FMixerDevice::UpdateHardwareTiming()
@@ -515,7 +583,6 @@ namespace Audio
 		SourceManager->ComputeNextBlockOfSamples();
 
 		FMixerSubmixWeakPtr MasterSubmix = GetMasterSubmix();
-
 		{
 			CSV_SCOPED_TIMING_STAT(Audio, Submixes);
 
@@ -523,7 +590,33 @@ namespace Audio
 			if (MasterSubmixPtr.IsValid())
 			{
 				// Process the audio output from the master submix
-				MasterSubmixPtr->ProcessAudio(ESubmixChannelFormat::Device, Output);
+				MasterSubmixPtr->ProcessAudio(Output);
+			}
+		}
+
+		{
+			CSV_SCOPED_TIMING_STAT(Audio, EndpointSubmixes);
+			FScopeLock ScopeLock(&EndpointSubmixesMutationLock);
+			for (FMixerSubmixPtr& Submix : DefaultEndpointSubmixes)
+			{
+				// If this hit, a submix was added to the default submix endpoint array
+				// even though it's not an endpoint, or a parent was set on an endpoint submix
+				// and it wasn't removed from DefaultEndpointSubmixes.
+				ensure(Submix->IsDefaultEndpointSubmix());
+
+				// Any endpoint submixes that don't specify an endpoint
+				// are summed into our master output.
+				Submix->ProcessAudio(Output);
+			}
+			
+			for (FMixerSubmixPtr& Submix : ExternalEndpointSubmixes)
+			{
+				// If this hit, a submix was added to the external submix endpoint array
+				// even though it's not an endpoint, or a parent was set on an endpoint submix
+				// and it wasn't removed from ExternalEndpointSubmixes.
+				ensure(Submix->IsExternalEndpointSubmix());
+
+				Submix->ProcessAudioAndSendToEndpoint();
 			}
 		}
 
@@ -582,7 +675,6 @@ namespace Audio
 			{
 				return;
 			}
-
 			OldSubmix->RemoveFromRoot();
 			FMixerSubmixPtr OldSubmixPtr = MasterSubmixInstances[TypeIndex];
 			if (OldSubmixPtr.IsValid())
@@ -645,6 +737,8 @@ namespace Audio
 		// Update/add new submix and instance to respective master arrays
 		MasterSubmixes[TypeIndex] = NewSubmix;
 		MasterSubmixInstances[TypeIndex] = NewMixerSubmix;
+
+		//Note: If we support using endpoint/soundfield submixes as a master submix in the future, we will need to call NewMixerSubmix->SetSoundfieldFactory here.
 		NewMixerSubmix->Init(NewSubmix, false /* bAllowReInit */);
 	}
 
@@ -659,9 +753,12 @@ namespace Audio
 			ReverbPluginSubmix->AddToRoot();
 
 			LoadSoundSubmix(*ReverbPluginSubmix);
+			InitSoundfieldAndEndpointDataForSubmix(*ReverbPluginSubmix, GetSubmixInstance(ReverbPluginSubmix).Pin(), false);
 
 			// Plugin must provide valid effect to enable reverb
 			FSoundEffectSubmixPtr ReverbPluginEffectSubmix = ReverbPluginInterface->GetEffectSubmix();
+			
+
 			if (ReverbPluginEffectSubmix.IsValid())
 			{
 				if (USoundEffectPreset* Preset = ReverbPluginEffectSubmix->GetPreset())
@@ -716,17 +813,17 @@ namespace Audio
 				LoadMasterSoundSubmix(EMasterSubmixType::EQ, TEXT("MasterEQSubmixDefault"), false /* DefaultMuteWhenBackgrounded */, AudioSettings->EQSubmix);
 			}
 
-			LoadMasterSoundSubmix(EMasterSubmixType::Ambisonics, TEXT("MasterAmbisonicSubmixDefault"), false /* DefaultMuteWhenBackgrounded */, AudioSettings->AmbisonicSubmix);
 			LoadPluginSoundSubmixes();
 
-			for (TObjectIterator<USoundSubmix> It; It; ++It)
+			for (TObjectIterator<USoundSubmixBase> It; It; ++It)
 			{
-				USoundSubmix* SubmixToLoad = *It;
+				USoundSubmixBase* SubmixToLoad = *It;
 				check(SubmixToLoad);
 
 				if (!IsMasterSubmixType(SubmixToLoad))
 				{
 					LoadSoundSubmix(*SubmixToLoad);
+					InitSoundfieldAndEndpointDataForSubmix(*SubmixToLoad, GetSubmixInstance(SubmixToLoad).Pin(), false);
 				}
 			}
 			bSubmixRegistrationDisabled = false;
@@ -743,7 +840,6 @@ namespace Audio
 			}, GET_STATID(STAT_InitSoundSubmixes));
 			return;
 		}
-
 		for (int32 i = 0; i < static_cast<int32>(EMasterSubmixType::Count); ++i)
 		{
 			if (DisableSubmixEffectEQCvar && i == static_cast<int32>(EMasterSubmixType::EQ))
@@ -751,7 +847,7 @@ namespace Audio
 				continue;
 			}
 
-			USoundSubmix* SoundSubmix = MasterSubmixes[i];
+			USoundSubmixBase* SoundSubmix = MasterSubmixes[i];
 			check(SoundSubmix);
 			FMixerSubmixPtr& MasterSubmixInstance = MasterSubmixInstances[i];
 
@@ -759,28 +855,33 @@ namespace Audio
 			{
 				RebuildSubmixLinks(*SoundSubmix, MasterSubmixInstance);
 			}
-		};
+		}
 
-		for (TPair<USoundSubmix*, FMixerSubmixPtr>& Entry : Submixes)
+		for (TPair<const USoundSubmixBase*, FMixerSubmixPtr>& Entry : Submixes)
 		{
-			USoundSubmix* SoundSubmix = Entry.Key;
+			const USoundSubmixBase* SoundSubmix = Entry.Key;
 			check(SoundSubmix);
 			FMixerSubmixPtr& SubmixInstance = Entry.Value;
 			RebuildSubmixLinks(*SoundSubmix, SubmixInstance);
 		}
 	}
 
-	void FMixerDevice::RebuildSubmixLinks(USoundSubmix& SoundSubmix, FMixerSubmixPtr& SubmixInstance)
+	void FMixerDevice::RebuildSubmixLinks(const USoundSubmixBase& SoundSubmix, FMixerSubmixPtr& SubmixInstance)
 	{
 		// Setup up the submix instance's parent and add the submix instance as a child
-		FMixerSubmixPtr ParentSubmixInstance = SoundSubmix.ParentSubmix
-			? GetSubmixInstance(SoundSubmix.ParentSubmix).Pin()
-			: GetMasterSubmix().Pin();
-
-		check(ParentSubmixInstance.IsValid());
+		FMixerSubmixPtr ParentSubmixInstance;
+		if (const USoundSubmixWithParentBase* SubmixWithParent = Cast<const USoundSubmixWithParentBase>(&SoundSubmix))
+		{
+			ParentSubmixInstance = SubmixWithParent->ParentSubmix
+				? GetSubmixInstance(SubmixWithParent->ParentSubmix).Pin()
+				: GetMasterSubmix().Pin();
+		}
 
-		SubmixInstance->SetParentSubmix(ParentSubmixInstance);
-		ParentSubmixInstance->AddChildSubmix(SubmixInstance);
+		if (ParentSubmixInstance.IsValid())
+		{
+			SubmixInstance->SetParentSubmix(ParentSubmixInstance);
+			ParentSubmixInstance->AddChildSubmix(SubmixInstance);
+		}
 	}
 
  	FAudioPlatformSettings FMixerDevice::GetPlatformSettings() const
@@ -813,11 +914,6 @@ namespace Audio
 		return MasterSubmixInstances[EMasterSubmixType::EQ];
 	}
 
-	FMixerSubmixWeakPtr FMixerDevice::GetMasterAmbisonicsSubmix()
-	{
-		return MasterSubmixInstances[EMasterSubmixType::Ambisonics];
-	}
-
 	void FMixerDevice::AddMasterSubmixEffect(uint32 SubmixEffectId, FSoundEffectSubmixPtr SoundEffectSubmix)
 	{
 		AudioRenderThreadCommand([this, SubmixEffectId, SoundEffectSubmix]()
@@ -862,17 +958,25 @@ namespace Audio
 		SourceManager->UpdateSourceEffectChain(SourceEffectChainId, SourceEffectChain, bPlayEffectChainTails);
 	}
 
-	void FMixerDevice::UpdateSubmixProperties(USoundSubmix* InSoundSubmix)
+	void FMixerDevice::UpdateSubmixProperties(USoundSubmixBase* InSoundSubmix)
 	{
 		check(InSoundSubmix);
 
+		// Output volume is only supported on USoundSubmixes.
+		USoundSubmix* CastedSubmix = Cast<USoundSubmix>(InSoundSubmix);
+
+		if (!CastedSubmix)
+		{
+			return;
+		}
+
 #if WITH_EDITOR
 		check(IsInAudioThread());
 
 		FMixerSubmixPtr MixerSubmix = GetSubmixInstance(InSoundSubmix).Pin();
 		if (MixerSubmix.IsValid())
 		{
-			const float NewVolume = InSoundSubmix->OutputVolume;
+			const float NewVolume = CastedSubmix->OutputVolume;
 			AudioRenderThreadCommand([MixerSubmix, NewVolume]()
 			{
 				MixerSubmix->SetOutputVolume(NewVolume);
@@ -955,7 +1059,7 @@ namespace Audio
 		}
 	}
 
-	bool FMixerDevice::IsMasterSubmixType(USoundSubmix* InSubmix) const
+	bool FMixerDevice::IsMasterSubmixType(const USoundSubmixBase* InSubmix) const
 	{
 		for (int32 i = 0; i < EMasterSubmixType::Count; ++i)
 		{
@@ -967,7 +1071,7 @@ namespace Audio
 		return false;
 	}
 
-	FMixerSubmixPtr FMixerDevice::GetMasterSubmixInstance(USoundSubmix* InSubmix)
+	FMixerSubmixPtr FMixerDevice::GetMasterSubmixInstance(const USoundSubmixBase* InSubmix)
 	{
 		check(MasterSubmixes.Num() == EMasterSubmixType::Count);
 		for (int32 i = 0; i < EMasterSubmixType::Count; ++i)
@@ -980,7 +1084,7 @@ namespace Audio
 		return nullptr;
 	}
 
-	void FMixerDevice::RegisterSoundSubmix(USoundSubmix* InSoundSubmix, bool bInit)
+	void FMixerDevice::RegisterSoundSubmix(const USoundSubmixBase* InSoundSubmix, bool bInit)
 	{
 		if (!InSoundSubmix || bSubmixRegistrationDisabled)
 		{
@@ -992,10 +1096,10 @@ namespace Audio
 			DECLARE_CYCLE_STAT(TEXT("FAudioThreadTask.RegisterSoundSubmix"), STAT_AudioRegisterSoundSubmix, STATGROUP_AudioThreadCommands);
 
 			FMixerDevice* MixerDevice = this;
-			FAudioThread::RunCommandOnAudioThread([MixerDevice, InSoundSubmix]()
+			FAudioThread::RunCommandOnAudioThread([MixerDevice, InSoundSubmix, bInit]()
 			{
 				CSV_SCOPED_TIMING_STAT(Audio, RegisterSubmix);
-				MixerDevice->RegisterSoundSubmix(InSoundSubmix);
+				MixerDevice->RegisterSoundSubmix(InSoundSubmix, bInit);
 			}, GET_STATID(STAT_AudioRegisterSoundSubmix));
 			return;
 		}
@@ -1005,12 +1109,13 @@ namespace Audio
 		if (!bIsMasterSubmix)
 		{
 			// Ensure parent structure is registered prior to current submix if missing
-			if (InSoundSubmix->ParentSubmix)
+			const USoundSubmixWithParentBase* SubmixWithParent = Cast<const USoundSubmixWithParentBase>(InSoundSubmix);
+			if (SubmixWithParent && SubmixWithParent->ParentSubmix)
 			{
-				FMixerSubmixPtr ParentSubmix = GetSubmixInstance(InSoundSubmix->ParentSubmix).Pin();
+				FMixerSubmixPtr ParentSubmix = GetSubmixInstance(SubmixWithParent->ParentSubmix).Pin();
 				if (!ParentSubmix.IsValid())
 				{
-					RegisterSoundSubmix(InSoundSubmix->ParentSubmix, bInit);
+					RegisterSoundSubmix(SubmixWithParent->ParentSubmix, bInit);
 				}
 			}
 
@@ -1018,18 +1123,19 @@ namespace Audio
 		}
 
 		FMixerSubmixPtr SubmixPtr = GetSubmixInstance(InSoundSubmix).Pin();
-		if (bInit)
-		{
-			SubmixPtr->Init(InSoundSubmix, true /* bAllowReInit */);
-		}
 
 		if (!bIsMasterSubmix)
 		{
 			RebuildSubmixLinks(*InSoundSubmix, SubmixPtr);
 		}
+
+		if (bInit)
+		{
+			InitSoundfieldAndEndpointDataForSubmix(*InSoundSubmix, SubmixPtr, true);
+		}
 	}
 
-	void FMixerDevice::LoadSoundSubmix(USoundSubmix& InSoundSubmix)
+	void FMixerDevice::LoadSoundSubmix(const USoundSubmixBase& InSoundSubmix)
 	{
 		// If submix not already found, load it.
 		FMixerSubmixPtr MixerSubmix = GetSubmixInstance(&InSoundSubmix).Pin();
@@ -1038,11 +1144,47 @@ namespace Audio
 			MixerSubmix = MakeShared<FMixerSubmix, ESPMode::ThreadSafe>(this);
 			Submixes.Add(&InSoundSubmix, MixerSubmix);
 		}
+	}
+
+	void FMixerDevice::InitSoundfieldAndEndpointDataForSubmix(const USoundSubmixBase& InSoundSubmix, FMixerSubmixPtr MixerSubmix, bool bAllowReInit)
+	{
+		{
+			FScopeLock ScopeLock(&EndpointSubmixesMutationLock);
+
+			// Check to see if this is an endpoint or soundfield submix:
+			if (const USoundfieldSubmix* SoundfieldSubmix = Cast<const USoundfieldSubmix>(&InSoundSubmix))
+			{
+				MixerSubmix->SetSoundfieldFactory(SoundfieldSubmix->GetSoundfieldFactoryForSubmix());
+			}
+			else if (const USoundfieldEndpointSubmix* SoundfieldEndpointSubmix = Cast<const USoundfieldEndpointSubmix>(&InSoundSubmix))
+			{
+				MixerSubmix->SetSoundfieldFactory(SoundfieldEndpointSubmix->GetSoundfieldEndpointForSubmix());
+			}
 
-		MixerSubmix->Init(&InSoundSubmix, false /* bAllowReInit */);
+			if (DefaultEndpointSubmixes.Contains(MixerSubmix))
+			{
+				DefaultEndpointSubmixes.RemoveSwap(MixerSubmix);
+			}
+
+			if (ExternalEndpointSubmixes.Contains(MixerSubmix))
+			{
+				ExternalEndpointSubmixes.RemoveSwap(MixerSubmix);
+			}
+
+			MixerSubmix->Init(&InSoundSubmix, bAllowReInit);
+
+			if (IsEndpointSubmix(&InSoundSubmix) && MixerSubmix->IsDefaultEndpointSubmix())
+			{
+				DefaultEndpointSubmixes.Add(MixerSubmix);
+			}
+			else if (MixerSubmix->IsExternalEndpointSubmix())
+			{
+				ExternalEndpointSubmixes.Add(MixerSubmix);
+			}
+		}
 	}
 
-	void FMixerDevice::UnregisterSoundSubmix(USoundSubmix* InSoundSubmix)
+	void FMixerDevice::UnregisterSoundSubmix(const USoundSubmixBase* InSoundSubmix)
 	{
 		if (!InSoundSubmix || bSubmixRegistrationDisabled || IsMasterSubmixType(InSoundSubmix))
 		{
@@ -1064,21 +1206,29 @@ namespace Audio
 		UnloadSoundSubmix(*InSoundSubmix);
 	}
 
-	void FMixerDevice::UnloadSoundSubmix(USoundSubmix& InSoundSubmix)
+	void FMixerDevice::UnloadSoundSubmix(const USoundSubmixBase& InSoundSubmix)
 	{
 		check(IsInAudioThread());
 
 		FMixerSubmixWeakPtr MasterSubmix = GetMasterSubmix();
-		FMixerSubmixPtr ParentSubmixInstance = InSoundSubmix.ParentSubmix
-			? GetSubmixInstance(InSoundSubmix.ParentSubmix).Pin()
-			: MasterSubmix.Pin();
+
+
+		FMixerSubmixPtr ParentSubmixInstance;
+		
+		// Check if this is a submix type that has a parent.
+		if (const USoundSubmixWithParentBase* InSoundSubmixWithParent = Cast<const USoundSubmixWithParentBase>(&InSoundSubmix))
+		{
+			ParentSubmixInstance = InSoundSubmixWithParent->ParentSubmix
+				? GetSubmixInstance(InSoundSubmixWithParent->ParentSubmix).Pin()
+				: MasterSubmix.Pin();
+		}
 
 		if (ParentSubmixInstance.IsValid())
 		{
 			ParentSubmixInstance->RemoveChildSubmix(GetSubmixInstance(&InSoundSubmix));
 		}
 
-		for (USoundSubmix* ChildSubmix : InSoundSubmix.ChildSubmixes)
+		for (USoundSubmixBase* ChildSubmix : InSoundSubmix.ChildSubmixes)
 		{
 			FMixerSubmixPtr ChildSubmixPtr = GetSubmixInstance(ChildSubmix).Pin();
 			if (ChildSubmixPtr.IsValid())
@@ -1089,6 +1239,20 @@ namespace Audio
 			}
 		}
 
+		FMixerSubmixWeakPtr MixerSubmixWeakPtr = GetSubmixInstance(&InSoundSubmix);
+		FMixerSubmixPtr MixerSubmix = MixerSubmixWeakPtr.Pin();
+
+		if (MixerSubmix && MixerSubmix->IsDefaultEndpointSubmix())
+		{
+			FScopeLock ScopeLock(&EndpointSubmixesMutationLock);
+			DefaultEndpointSubmixes.Remove(MixerSubmix);
+		}
+		else if (MixerSubmix && MixerSubmix->IsExternalEndpointSubmix())
+		{
+			FScopeLock ScopeLock(&EndpointSubmixesMutationLock);
+			ExternalEndpointSubmixes.Add(MixerSubmix);
+		}
+
 		Submixes.Remove(&InSoundSubmix);
 	}
 
@@ -1100,17 +1264,36 @@ namespace Audio
 #endif
 	}
 
-	FMixerSubmixWeakPtr FMixerDevice::GetSubmixInstance(USoundSubmix* InSoundSubmix)
+	FMixerSubmixWeakPtr FMixerDevice::GetSubmixInstance(const USoundSubmixBase* SoundSubmix)
 	{
 		LLM_SCOPE(ELLMTag::AudioMixer);
 
-		FMixerSubmixPtr MixerSubmix = GetMasterSubmixInstance(InSoundSubmix);
+		FMixerSubmixPtr MixerSubmix = GetMasterSubmixInstance(SoundSubmix);
 		if (MixerSubmix.IsValid())
 		{
 			return MixerSubmix;
 		}
 
-		return Submixes.FindRef(InSoundSubmix);
+		return Submixes.FindRef(SoundSubmix);
+	}
+
+	ISoundfieldFactory* FMixerDevice::GetFactoryForSubmixInstance(USoundSubmix* SoundSubmix)
+	{
+		FMixerSubmixWeakPtr WeakSubmixPtr = GetSubmixInstance(SoundSubmix);
+		return GetFactoryForSubmixInstance(WeakSubmixPtr);
+	}
+
+	ISoundfieldFactory* FMixerDevice::GetFactoryForSubmixInstance(FMixerSubmixWeakPtr& SoundSubmixPtr)
+	{
+		FMixerSubmixPtr SubmixPtr = SoundSubmixPtr.Pin();
+		if (ensure(SubmixPtr.IsValid()))
+		{
+			return SubmixPtr->GetSoundfieldFactory();
+		}
+		else
+		{
+			return nullptr;
+		}
 	}
 
 	FMixerSourceVoice* FMixerDevice::GetMixerSourceVoice()
@@ -1142,13 +1325,13 @@ namespace Audio
 		return SourceManager->GetNumActiveSources();
 	}
 
-	void FMixerDevice::Get3DChannelMap(const ESubmixChannelFormat InSubmixType, const FWaveInstance* InWaveInstance, float EmitterAzimith, float NormalizedOmniRadius, Audio::AlignedFloatBuffer& OutChannelMap)
+	void FMixerDevice::Get3DChannelMap(const int32 InSubmixNumChannels, const FWaveInstance* InWaveInstance, float EmitterAzimith, float NormalizedOmniRadius, Audio::AlignedFloatBuffer& OutChannelMap)
 	{
 		// If we're center-channel only, then no need for spatial calculations, but need to build a channel map
 		if (InWaveInstance->bCenterChannelOnly)
 		{
-			int32 NumOutputChannels = GetNumChannelsForSubmixFormat(InSubmixType);
-			const TArray<EAudioMixerChannel::Type>& ChannelArray = GetChannelArrayForSubmixChannelFormat(InSubmixType);
+			int32 NumOutputChannels = InSubmixNumChannels;
+			const TArray<EAudioMixerChannel::Type>& ChannelArray = GetChannelArray();
 
 			// If we are only spatializing to stereo output
 			if (NumOutputChannels == 2)
@@ -1175,24 +1358,21 @@ namespace Audio
 		const FChannelPositionInfo* PrevChannelInfo = nullptr;
 		const FChannelPositionInfo* NextChannelInfo = nullptr;
 
-		const TArray<FChannelPositionInfo>* CurrentChannelAzimuthPositionsPtr = ChannelAzimuthPositions.Find(InSubmixType);
-		const TArray<FChannelPositionInfo>& CurrentChannelAzimuthPositions = *CurrentChannelAzimuthPositionsPtr;
-
-		for (int32 i = 0; i < CurrentChannelAzimuthPositions.Num(); ++i)
+		for (int32 i = 0; i < DeviceChannelAzimuthPositions.Num(); ++i)
 		{
-			const FChannelPositionInfo& ChannelPositionInfo = CurrentChannelAzimuthPositions[i];
+			const FChannelPositionInfo& ChannelPositionInfo = DeviceChannelAzimuthPositions[i];
 
 			if (Azimuth <= ChannelPositionInfo.Azimuth)
 			{
-				NextChannelInfo = &CurrentChannelAzimuthPositions[i];
+				NextChannelInfo = &DeviceChannelAzimuthPositions[i];
 
 				int32 PrevIndex = i - 1;
 				if (PrevIndex < 0)
 				{
-					PrevIndex = CurrentChannelAzimuthPositions.Num() - 1;
+					PrevIndex = DeviceChannelAzimuthPositions.Num() - 1;
 				}
 
-				PrevChannelInfo = &CurrentChannelAzimuthPositions[PrevIndex];
+				PrevChannelInfo = &DeviceChannelAzimuthPositions[PrevIndex];
 				break;
 			}
 		}
@@ -1200,8 +1380,8 @@ namespace Audio
 		// If we didn't find anything, that means our azimuth position is at the top of the mapping
 		if (PrevChannelInfo == nullptr)
 		{
-			PrevChannelInfo = &CurrentChannelAzimuthPositions[CurrentChannelAzimuthPositions.Num() - 1];
-			NextChannelInfo = &CurrentChannelAzimuthPositions[0];
+			PrevChannelInfo = &DeviceChannelAzimuthPositions[DeviceChannelAzimuthPositions.Num() - 1];
+			NextChannelInfo = &DeviceChannelAzimuthPositions[0];
 			AUDIO_MIXER_CHECK(PrevChannelInfo != NextChannelInfo);
 		}
 
@@ -1251,15 +1431,15 @@ namespace Audio
 
 		// Build the output channel map based on the current platform device output channel array 
 
-		int32 NumSpatialChannels = CurrentChannelAzimuthPositions.Num();
-		if (CurrentChannelAzimuthPositions.Num() > 4)
+		int32 NumSpatialChannels = DeviceChannelAzimuthPositions.Num();
+		if (DeviceChannelAzimuthPositions.Num() > 4)
 		{
 			NumSpatialChannels--;
 		}
 		float OmniPanFactor = 1.0f / NumSpatialChannels;
 
 		float DefaultEffectivePan = !OmniAmount ? 0.0f : FMath::Lerp(0.0f, OmniPanFactor, OmniAmount);
-		const TArray<EAudioMixerChannel::Type>& ChannelArray = GetChannelArrayForSubmixChannelFormat(InSubmixType);
+		const TArray<EAudioMixerChannel::Type>& ChannelArray = GetChannelArray();
 
 		for (EAudioMixerChannel::Type Channel : ChannelArray)
 		{
@@ -1289,12 +1469,6 @@ namespace Audio
 		}
 	}
 
-	uint32 FMixerDevice::GetNewUniqueAmbisonicsStreamID()
-	{
-		static uint32 AmbisonicsStreamIDCounter = 0;
-		return ++AmbisonicsStreamIDCounter;
-	}
-
 	const TArray<FTransform>* FMixerDevice::GetListenerTransforms()
 	{
 		return SourceManager->GetListenerTransforms();
diff --git a/Engine/Source/Runtime/AudioMixer/Private/AudioMixerSource.cpp b/Engine/Source/Runtime/AudioMixer/Private/AudioMixerSource.cpp
--- a/Engine/Source/Runtime/AudioMixer/Private/AudioMixerSource.cpp
+++ b/Engine/Source/Runtime/AudioMixer/Private/AudioMixerSource.cpp
@@ -219,6 +219,7 @@ namespace Audio
 					SubmixSend.Submix = SubmixPtr;
 					SubmixSend.SendLevel = 1.0f;
 					SubmixSend.bIsMainSend = true;
+					SubmixSend.SoundfieldFactory = MixerDevice->GetFactoryForSubmixInstance(SubmixSend.Submix);
 					InitParams.SubmixSends.Add(SubmixSend);
 				}
 
@@ -231,6 +232,7 @@ namespace Audio
 						SubmixSend.Submix = MixerDevice->GetSubmixInstance(SendInfo.SoundSubmix);
 						SubmixSend.SendLevel = SendInfo.SendLevel;
 						SubmixSend.bIsMainSend = false;
+						SubmixSend.SoundfieldFactory = MixerDevice->GetFactoryForSubmixInstance(SubmixSend.Submix);
 						InitParams.SubmixSends.Add(SubmixSend);
 					}
 				}
@@ -242,9 +244,7 @@ namespace Audio
 				FMixerSubmixPtr SubmixPtr = Send.Submix.Pin();
 				if (SubmixPtr.IsValid())
 				{
-					ESubmixChannelFormat SubmixChannelType = SubmixPtr->GetSubmixChannels();
-					ChannelMaps[(int32)SubmixChannelType].bUsed = true;
-					ChannelMaps[(int32)SubmixChannelType].ChannelMap.Reset();
+					ChannelMap.Reset();
 				}
 			}
 
@@ -799,11 +799,7 @@ namespace Audio
 		}
 
 		// Reset the source's channel maps
-		for (int32 i = 0; i < (int32)ESubmixChannelFormat::Count; ++i)
-		{
-			ChannelMaps[i].bUsed = false;
-			ChannelMaps[i].ChannelMap.Reset();
-		}
+		ChannelMap.Reset();
 
 		InitializationState = EMixerSourceInitializationState::NotInitialized;
 	}
@@ -904,7 +900,6 @@ namespace Audio
 		if (bReverbApplied)
 		{
 			float ReverbSendLevel = 0.0f;
-			ChannelMaps[(int32)ESubmixChannelFormat::Device].bUsed = true;
 
 			if (WaveInstance->ReverbSendMethod == EReverbSendMethod::Manual)
 			{
@@ -948,7 +943,7 @@ namespace Audio
 			if (SendInfo.SoundSubmix != nullptr)
 			{
 				FMixerSubmixWeakPtr SubmixInstance = MixerDevice->GetSubmixInstance(SendInfo.SoundSubmix);
-				float SendLevel = 0.0f;
+				float SendLevel = 1.0f;
 
 				// calculate send level based on distance if that method is enabled
 				if (SendInfo.SendLevelControlMethod == ESendLevelControlMethod::Manual)
@@ -975,10 +970,6 @@ namespace Audio
 
 				// set the level for this send
 				MixerSourceVoice->SetSubmixSendInfo(SubmixInstance, SendLevel);
-
-				// Make sure we flag that we're using this submix sends since these can be dynamically added from BP
-				// If we don't flag this then these channel maps won't be generated for this channel format
-				ChannelMaps[(int32)SendInfo.SoundSubmix->ChannelFormat].bUsed = true;
 			}
 		}
  	}
@@ -1043,24 +1034,14 @@ namespace Audio
 		const FAudioPlatformDeviceInfo& DeviceInfo = MixerDevice->GetPlatformDeviceInfo();
 
 		// Compute a new speaker map for each possible output channel mapping for the source
-		for (int32 i = 0; i < (int32)ESubmixChannelFormat::Count; ++i)
+		const uint32 NumChannels = Buffer->NumChannels;
+		if (ComputeChannelMap(Buffer->NumChannels, ChannelMap))
 		{
-			FChannelMapInfo& ChannelMapInfo = ChannelMaps[i];
-			if (ChannelMapInfo.bUsed)
-			{
-				ESubmixChannelFormat ChannelType = (ESubmixChannelFormat)i;
-
-				check(Buffer);
-				const uint32 NumChannels = Buffer->NumChannels;
-				if (ComputeChannelMap(ChannelType, Buffer->NumChannels, ChannelMapInfo.ChannelMap))
-				{
-					MixerSourceVoice->SetChannelMap(ChannelType, NumChannels, ChannelMapInfo.ChannelMap, bIs3D, WaveInstance->bCenterChannelOnly);
-				}
-			}
+			MixerSourceVoice->SetChannelMap(NumChannels, ChannelMap, bIs3D, WaveInstance->bCenterChannelOnly);
 		}
 	}
 
-	bool FMixerSource::ComputeMonoChannelMap(const ESubmixChannelFormat SubmixChannelType, Audio::AlignedFloatBuffer& OutChannelMap)
+	bool FMixerSource::ComputeMonoChannelMap(Audio::AlignedFloatBuffer& OutChannelMap)
 	{
 		if (IsUsingObjectBasedSpatialization())
 		{
@@ -1071,20 +1052,20 @@ namespace Audio
 			}
 
 			// Treat the source as if it is a 2D stereo source:
-			return ComputeStereoChannelMap(SubmixChannelType, OutChannelMap);
+			return ComputeStereoChannelMap(OutChannelMap);
 		}
 		else if (WaveInstance->GetUseSpatialization() && (!FMath::IsNearlyEqual(WaveInstance->AbsoluteAzimuth, PreviousAzimuth, 0.01f) || MixerSourceVoice->NeedsSpeakerMap()))
 		{
 			// Don't need to compute the source channel map if the absolute azimuth hasn't changed much
 			PreviousAzimuth = WaveInstance->AbsoluteAzimuth;
 			OutChannelMap.Reset();
-			MixerDevice->Get3DChannelMap(SubmixChannelType, WaveInstance, WaveInstance->AbsoluteAzimuth, SpatializationParams.NormalizedOmniRadius, OutChannelMap);
+			MixerDevice->Get3DChannelMap(MixerDevice->GetNumDeviceChannels(), WaveInstance, WaveInstance->AbsoluteAzimuth, SpatializationParams.NormalizedOmniRadius, OutChannelMap);
 			return true;
 		}
 		else if (!OutChannelMap.Num())
 		{
 			// Only need to compute the 2D channel map once
-			MixerDevice->Get2DChannelMap(bIsVorbis, SubmixChannelType, 1, WaveInstance->bCenterChannelOnly, OutChannelMap);
+			MixerDevice->Get2DChannelMap(bIsVorbis, 1, WaveInstance->bCenterChannelOnly, OutChannelMap);
 			return true;
 		}
 
@@ -1092,7 +1073,7 @@ namespace Audio
 		return false;
 	}
 
-	bool FMixerSource::ComputeStereoChannelMap(const ESubmixChannelFormat InSubmixChannelType, Audio::AlignedFloatBuffer& OutChannelMap)
+	bool FMixerSource::ComputeStereoChannelMap(Audio::AlignedFloatBuffer& OutChannelMap)
 	{
 		// Only recalculate positional data if the source has moved a significant amount:
 		if (WaveInstance->GetUseSpatialization() && (!FMath::IsNearlyEqual(WaveInstance->AbsoluteAzimuth, PreviousAzimuth, 0.01f) || MixerSourceVoice->NeedsSpeakerMap()))
@@ -1134,8 +1115,10 @@ namespace Audio
 				// Reset the channel map, the stereo spatialization channel mapping calls below will append their mappings
 				OutChannelMap.Reset();
 
-				MixerDevice->Get3DChannelMap(InSubmixChannelType, WaveInstance, LeftAzimuth, SpatializationParams.NormalizedOmniRadius, OutChannelMap);
-				MixerDevice->Get3DChannelMap(InSubmixChannelType, WaveInstance, RightAzimuth, SpatializationParams.NormalizedOmniRadius, OutChannelMap);
+				const int32 NumOutputChannels = MixerDevice->GetNumDeviceChannels();
+
+				MixerDevice->Get3DChannelMap(NumOutputChannels, WaveInstance, LeftAzimuth, SpatializationParams.NormalizedOmniRadius, OutChannelMap);
+				MixerDevice->Get3DChannelMap(NumOutputChannels, WaveInstance, RightAzimuth, SpatializationParams.NormalizedOmniRadius, OutChannelMap);
 
 				return true;
 			}
@@ -1143,26 +1126,26 @@ namespace Audio
 
 		if (!OutChannelMap.Num())
 		{
-			MixerDevice->Get2DChannelMap(bIsVorbis, InSubmixChannelType, 2, WaveInstance->bCenterChannelOnly, OutChannelMap);
+			MixerDevice->Get2DChannelMap(bIsVorbis, 2, WaveInstance->bCenterChannelOnly, OutChannelMap);
 			return true;
 		}
 
 		return false;
 	}
 
-	bool FMixerSource::ComputeChannelMap(const ESubmixChannelFormat InSubmixChannelType, const int32 NumSourceChannels, Audio::AlignedFloatBuffer& OutChannelMap)
+	bool FMixerSource::ComputeChannelMap(const int32 NumSourceChannels, Audio::AlignedFloatBuffer& OutChannelMap)
 	{
 		if (NumSourceChannels == 1)
 		{
-			return ComputeMonoChannelMap(InSubmixChannelType, OutChannelMap);
+			return ComputeMonoChannelMap(OutChannelMap);
 		}
 		else if (NumSourceChannels == 2)
 		{
-			return ComputeStereoChannelMap(InSubmixChannelType, OutChannelMap);
+			return ComputeStereoChannelMap(OutChannelMap);
 		}
 		else if (!OutChannelMap.Num())
 		{
-			MixerDevice->Get2DChannelMap(bIsVorbis, InSubmixChannelType, NumSourceChannels, WaveInstance->bCenterChannelOnly, OutChannelMap);
+			MixerDevice->Get2DChannelMap(bIsVorbis, NumSourceChannels, WaveInstance->bCenterChannelOnly, OutChannelMap);
 			return true;
 		}
 		return false;
diff --git a/Engine/Source/Runtime/AudioMixer/Private/AudioMixerSource.h b/Engine/Source/Runtime/AudioMixer/Private/AudioMixerSource.h
--- a/Engine/Source/Runtime/AudioMixer/Private/AudioMixerSource.h
+++ b/Engine/Source/Runtime/AudioMixer/Private/AudioMixerSource.h
@@ -86,13 +86,13 @@ namespace Audio
 		void UpdateChannelMaps();
 
 		/** Computes the mono-channel map. */
-		bool ComputeMonoChannelMap(const ESubmixChannelFormat SubmixChannelType, Audio::AlignedFloatBuffer& OutChannelMap);
+		bool ComputeMonoChannelMap(Audio::AlignedFloatBuffer& OutChannelMap);
 
 		/** Computes the stereo-channel map. */
-		bool ComputeStereoChannelMap(const ESubmixChannelFormat SubmixChannelType, Audio::AlignedFloatBuffer& OutChannelMap);
+		bool ComputeStereoChannelMap(Audio::AlignedFloatBuffer& OutChannelMap);
 
 		/** Compute the channel map based on the number of output and source channels. */
-		bool ComputeChannelMap(const ESubmixChannelFormat SubmixChannelType, const int32 NumSourceChannels, Audio::AlignedFloatBuffer& OutChannelMap);
+		bool ComputeChannelMap(const int32 NumSourceChannels, Audio::AlignedFloatBuffer& OutChannelMap);
 
 		/** Whether or not we should create the source voice with the HRTF spatializer. */
 		bool UseObjectBasedSpatialization() const;
@@ -119,16 +119,6 @@ namespace Audio
 		TSharedPtr<FMixerSourceBuffer> MixerSourceBuffer;
 		FMixerSourceVoice* MixerSourceVoice;
 
-		struct FChannelMapInfo
-		{
-			Audio::AlignedFloatBuffer ChannelMap;
-			bool bUsed;
-
-			FChannelMapInfo()
-				: bUsed(false)
-			{}
-		};
-
 		// This holds data copied from FSoundSourceBusSendInfo when a new sound starts playing
 		// so that distance-based level control can be calculated during rendering
 		struct FDynamicBusSendInfo
@@ -156,7 +146,7 @@ namespace Audio
 		};
 
 		// Mapping of channel map types to channel maps. Determined by what submixes this source sends its audio to.
-		FChannelMapInfo ChannelMaps[(int32) ESubmixChannelFormat::Count];
+		Audio::AlignedFloatBuffer ChannelMap;
 
 		float PreviousAzimuth;
 		mutable float PreviousPlaybackPercent;
diff --git a/Engine/Source/Runtime/AudioMixer/Private/AudioMixerSourceManager.cpp b/Engine/Source/Runtime/AudioMixer/Private/AudioMixerSourceManager.cpp
--- a/Engine/Source/Runtime/AudioMixer/Private/AudioMixerSourceManager.cpp
+++ b/Engine/Source/Runtime/AudioMixer/Private/AudioMixerSourceManager.cpp
@@ -8,6 +8,7 @@
 #include "AudioMixerSubmix.h"
 #include "IAudioExtensionPlugin.h"
 #include "AudioMixer.h"
+#include "SoundFieldRendering.h"
 #include "ProfilingDebugging/CsvProfiler.h"
 
 // Link to "Audio" profiling category
@@ -651,6 +652,7 @@ namespace Audio
 			SourceInfo.bIsExternalSend = InitParams.bIsExternalSend;
 			SourceInfo.bIsVorbis = InitParams.bIsVorbis;
 			SourceInfo.AudioComponentID = InitParams.AudioComponentID;
+			SourceInfo.bIsAmbisonics = InitParams.bIsAmbisonics;
 
 			// Call initialization from the render thread so anything wanting to do any initialization here can do so (e.g. procedural sound waves)
 			SourceInfo.SourceListener = InitParams.SourceListener;
@@ -789,6 +791,11 @@ namespace Audio
 			const int32 SourceInputChannels = (SourceInfo.bUseHRTFSpatializer && !SourceInfo.bIsExternalSend) ? 2 : SourceInfo.NumInputChannels;
 			FSourceDownmixData& DownmixData = InitializeDownmixForSource(SourceId, SourceInputChannels, MixerDevice->GetDeviceOutputChannels(), NumOutputFrames);
 
+			if (SourceInfo.bIsAmbisonics)
+			{
+				DownmixData.AmbisonicsDecoder = CreateDefaultSourceAmbisonicsDecoder(MixerDevice);
+			}
+
 			for (int32 i = 0; i < InitParams.SubmixSends.Num(); ++i)
 			{
 				const FMixerSourceSubmixSend& MixerSubmixSend = InitParams.SubmixSends[i];
@@ -799,12 +806,46 @@ namespace Audio
 					SourceInfo.SubmixSends.Add(MixerSubmixSend);
 					SubmixPtr->AddOrSetSourceVoice(InitParams.SourceVoice, MixerSubmixSend.SendLevel);
 
-					// Prepare output buffers and speaker map entries for every submix channel type
-					const ESubmixChannelFormat SubmixChannelType = SubmixPtr->GetSubmixChannels();
+					
+					if (SubmixPtr->IsSoundfieldSubmix())
+					{
+						FSoundfieldEncodingKey Key = SubmixPtr->GetKeyForSubmixEncoding();
+						FSubmixSoundfieldData& SubmixSoundfieldInfo = GetChannelInfoForFormat(Key, DownmixData);
 
-					FSubmixChannelTypeInfo& SubmixChannelInfo = GetChannelInfoForFormat(SubmixChannelType, DownmixData);
+						// Set up encoder (or transcoder if this is an ambisonics source).
+						ISoundfieldFactory* Factory = SubmixPtr->GetSoundfieldFactory();
+						check(Factory);
 
-					SubmixChannelInfo.bInUse = true;
+						FAudioPluginInitializationParams PluginInitParams = SubmixPtr->GetInitializationParamsForSoundfieldStream();
+						PluginInitParams.NumOutputChannels = SourceInputChannels;
+
+						SubmixSoundfieldInfo.EncoderSettings = SubmixPtr->GetSoundfieldSettings().Duplicate();
+
+						if (SourceInfo.bIsAmbisonics)
+						{
+							// If this soundfield submix can transcode our ambisonics format, set up a transcoder stream.
+							if (Factory->GetSoundfieldFormatName() == GetUnrealAmbisonicsFormatName())
+							{
+								SubmixSoundfieldInfo.bIsUnrealAmbisonicsSubmix = true;
+							}
+							else if (Factory->CanTranscodeFromSoundfieldFormat(GetUnrealAmbisonicsFormatName(), GetAmbisonicsSourceDefaultSettings()))
+							{
+								SubmixSoundfieldInfo.AmbiTranscoder = Factory->CreateTranscoderStream(GetUnrealAmbisonicsFormatName(), GetAmbisonicsSourceDefaultSettings(), Factory->GetSoundfieldFormatName(), *SubmixSoundfieldInfo.EncoderSettings, PluginInitParams);
+							}
+						}
+						else
+						{
+							check(SubmixSoundfieldInfo.EncoderSettings.IsValid());
+
+							SubmixSoundfieldInfo.Encoder = Factory->CreateEncoderStream(PluginInitParams, *SubmixSoundfieldInfo.EncoderSettings);
+						}
+
+						SubmixSoundfieldInfo.EncodedPacket = Factory->CreateEmptyPacket();
+					}
+					else
+					{
+						FSubmixChannelData& SubmixChannelInfo = GetChannelInfoForDevice(DownmixData);
+					}
 				}
 			}
 
@@ -1047,13 +1088,13 @@ namespace Audio
 		});
 	}
 
-	void FMixerSourceManager::SetChannelMap(const int32 SourceId, const ESubmixChannelFormat SubmixChannelType, const uint32 NumInputChannels, const Audio::AlignedFloatBuffer& ChannelMap, const bool bInIs3D, const bool bInIsCenterChannelOnly)
+	void FMixerSourceManager::SetChannelMap(const int32 SourceId, const uint32 NumInputChannels, const Audio::AlignedFloatBuffer& ChannelMap, const bool bInIs3D, const bool bInIsCenterChannelOnly)
 	{
 		AUDIO_MIXER_CHECK(SourceId < NumTotalSources);
 		AUDIO_MIXER_CHECK(GameThreadInfo.bIsBusy[SourceId]);
 		AUDIO_MIXER_CHECK_GAME_THREAD(MixerDevice);
 
-		AudioMixerThreadCommand([this, SourceId, SubmixChannelType, NumInputChannels, ChannelMap, bInIs3D, bInIsCenterChannelOnly]()
+		AudioMixerThreadCommand([this, SourceId, NumInputChannels, ChannelMap, bInIs3D, bInIsCenterChannelOnly]()
 		{
 			AUDIO_MIXER_CHECK_AUDIO_PLAT_THREAD(MixerDevice);
 
@@ -1073,45 +1114,36 @@ namespace Audio
 			SourceInfo.bIs3D = bInIs3D;
 			SourceInfo.bIsCenterChannelOnly = bInIsCenterChannelOnly;
 
-			FSubmixChannelTypeInfo& ChannelTypeInfo = GetChannelInfoForFormat(SubmixChannelType, DownmixData);
-			ChannelTypeInfo.bInUse = true;
+			FSubmixChannelData& ChannelTypeInfo = GetChannelInfoForDevice(DownmixData);
 
 			// Fix up the channel map in case device output count changed
-			if (SubmixChannelType == ESubmixChannelFormat::Device)
+			const uint32 ChannelMapSize = ChannelTypeInfo.ChannelMap.CopySize / sizeof(float);
+
+			// If this is true, then the device changed while the command was in-flight
+			if (ChannelMap.Num() != ChannelMapSize)
 			{
-				const uint32 ChannelMapSize = ChannelTypeInfo.ChannelMap.CopySize / sizeof(float);
+				// todo: investigate turning this into a stack array
+				Audio::AlignedFloatBuffer NewChannelMap;
 
-				// If this is true, then the device changed while the command was in-flight
-				if (ChannelMap.Num() != ChannelMapSize)
+				// If 3d then just zero it out, we'll get another channel map shortly
+				if (bInIs3D)
 				{
-					// todo: investigate turning this into a stack array
-					Audio::AlignedFloatBuffer NewChannelMap;
-
-					// If 3d then just zero it out, we'll get another channel map shortly
-					if (bInIs3D)
-					{
-						NewChannelMap.AddZeroed(ChannelMapSize);
-						GameThreadInfo.bNeedsSpeakerMap[SourceId] = true;
-					}
-					// Otherwise, get an appropriate channel map for the new device configuration
-					else
-					{
-						const uint32 NumOutputChannels = ChannelMapSize / NumInputChannels;
-						FMixerDevice::Get2DChannelMap(SourceInfo.bIsVorbis, NumInputChannels, NumOutputChannels, bInIsCenterChannelOnly, NewChannelMap);
-					}
-
-					// Make sure we've been flagged to be using this submix channel type entry
-					ChannelTypeInfo.ChannelMap.SetChannelMap(NewChannelMap.GetData());
+					NewChannelMap.AddZeroed(ChannelMapSize);
+					GameThreadInfo.bNeedsSpeakerMap[SourceId] = true;
 				}
+				// Otherwise, get an appropriate channel map for the new device configuration
 				else
 				{
-					GameThreadInfo.bNeedsSpeakerMap[SourceId] = false;
-					ChannelTypeInfo.ChannelMap.SetChannelMap(ChannelMap.GetData());
+					const uint32 NumOutputChannels = ChannelMapSize / NumInputChannels;
+					FMixerDevice::Get2DChannelMap(SourceInfo.bIsVorbis, NumInputChannels, NumOutputChannels, bInIsCenterChannelOnly, NewChannelMap);
 				}
+
+				// Make sure we've been flagged to be using this submix channel type entry
+				ChannelTypeInfo.ChannelMap.SetChannelMap(NewChannelMap.GetData());
 			}
 			else
 			{
-				// Since we're artificially mixing to this channel count, we don't need to deal with device reset
+				GameThreadInfo.bNeedsSpeakerMap[SourceId] = false;
 				ChannelTypeInfo.ChannelMap.SetChannelMap(ChannelMap.GetData());
 			}
 		});
@@ -1175,18 +1207,10 @@ namespace Audio
 				if (bIsNew)
 				{
 					SourceInfo.SubmixSends.Add(InSubmixSend);
-
-					FSourceDownmixData& DownmixData = DownmixDataArray[SourceId];
-
-					// Flag that we're now using this submix channel info
-					const ESubmixChannelFormat ChannelType = InSubmixPtr->GetSubmixChannels();
-					FSubmixChannelTypeInfo& ChannelInfo = GetChannelInfoForFormat(ChannelType, DownmixData);
-
-					ChannelInfo.bInUse = true;
 				}
 
-			}
 				InSubmixPtr->AddOrSetSourceVoice(MixerSources[SourceId], InSubmixSend.SendLevel);
+			}
 		});
 	}
 
@@ -1712,11 +1736,36 @@ namespace Audio
 		}
 	}
 
-	void FMixerSourceManager::ComputeDownmix3D(FSourceDownmixData& DownmixData)
+	void FMixerSourceManager::ComputeDownmix3D(FSourceDownmixData& DownmixData, FMixerDevice* MixerDevice)
 	{
 		// This enormous switch statement handles using the correct function for a given number of input and output channels.
 		// For 3D sources, we interpolate from ChannelStartGains to ChannelDestinationGains.
-		if (DownmixData.NumInputChannels == 1)
+		if (DownmixData.AmbisonicsDecoder.IsValid())
+		{
+			FAmbisonicsSoundfieldBuffer AmbiBuffer;
+			AmbiBuffer.AudioBuffer = MoveTemp(*DownmixData.PostEffectBuffers);
+			AmbiBuffer.NumChannels = DownmixData.NumInputChannels;
+			AmbiBuffer.Rotation = DownmixData.PositionalData.Rotation;
+
+			DownmixData.PositionalData.NumChannels = DownmixData.NumDeviceChannels;
+			DownmixData.PositionalData.ChannelPositions = MixerDevice->GetDefaultPositionMap(DownmixData.NumDeviceChannels);
+
+			FSoundfieldDecoderInputData InputData =
+			{
+				AmbiBuffer,
+				DownmixData.PositionalData,
+				static_cast<int32>(DownmixData.PostEffectBuffers->Num() / DownmixData.NumInputChannels),
+				MixerDevice->GetSampleRate()
+			};
+			
+			FSoundfieldDecoderOutputData OutputData = { DownmixData.DeviceSubmixInfo.OutputBuffer };
+
+			DownmixData.AmbisonicsDecoder->Decode(InputData, OutputData);
+
+			// Move the encoded ambisonics source buffer back to PostEffectBuffers to prevent reallocation
+			*DownmixData.PostEffectBuffers = MoveTemp(AmbiBuffer.AudioBuffer);
+		}
+		else if (DownmixData.NumInputChannels == 1)
 		{
 			switch (DownmixData.NumDeviceChannels)
 			{
@@ -1733,26 +1782,6 @@ namespace Audio
 				Audio::MixMonoTo2ChannelsFast(*DownmixData.PostEffectBuffers, DownmixData.DeviceSubmixInfo.OutputBuffer, DownmixData.DeviceSubmixInfo.ChannelMap.ChannelStartGains, DownmixData.DeviceSubmixInfo.ChannelMap.ChannelDestinationGains);
 				break;
 			}
-
-			if (DownmixData.StereoSubmixInfo.bInUse)
-			{
-				Audio::MixMonoTo2ChannelsFast(*DownmixData.PostEffectBuffers, DownmixData.StereoSubmixInfo.OutputBuffer, DownmixData.StereoSubmixInfo.ChannelMap.ChannelStartGains, DownmixData.StereoSubmixInfo.ChannelMap.ChannelDestinationGains);
-			}
-
-			if (DownmixData.QuadSubmixInfo.bInUse)
-			{
-				Audio::MixMonoTo4ChannelsFast(*DownmixData.PostEffectBuffers, DownmixData.QuadSubmixInfo.OutputBuffer, DownmixData.QuadSubmixInfo.ChannelMap.ChannelStartGains, DownmixData.QuadSubmixInfo.ChannelMap.ChannelDestinationGains);
-			}
-
-			if (DownmixData.FiveOneSubmixInfo.bInUse)
-			{
-				Audio::MixMonoTo6ChannelsFast(*DownmixData.PostEffectBuffers, DownmixData.FiveOneSubmixInfo.OutputBuffer, DownmixData.FiveOneSubmixInfo.ChannelMap.ChannelStartGains, DownmixData.FiveOneSubmixInfo.ChannelMap.ChannelDestinationGains);
-			}
-
-			if (DownmixData.SevenOneSubmixInfo.bInUse)
-			{
-				Audio::MixMonoTo8ChannelsFast(*DownmixData.PostEffectBuffers, DownmixData.SevenOneSubmixInfo.OutputBuffer, DownmixData.SevenOneSubmixInfo.ChannelMap.ChannelStartGains, DownmixData.SevenOneSubmixInfo.ChannelMap.ChannelDestinationGains);
-			}
 		}
 		else if (DownmixData.NumInputChannels == 2)
 		{
@@ -1771,61 +1800,46 @@ namespace Audio
 				Audio::Mix2ChannelsTo2ChannelsFast(*DownmixData.PostEffectBuffers, DownmixData.DeviceSubmixInfo.OutputBuffer, DownmixData.DeviceSubmixInfo.ChannelMap.ChannelStartGains, DownmixData.DeviceSubmixInfo.ChannelMap.ChannelDestinationGains);
 				break;
 			}
-
-			if (DownmixData.StereoSubmixInfo.bInUse)
-			{
-				Audio::Mix2ChannelsTo2ChannelsFast(*DownmixData.PostEffectBuffers, DownmixData.StereoSubmixInfo.OutputBuffer, DownmixData.StereoSubmixInfo.ChannelMap.ChannelStartGains, DownmixData.StereoSubmixInfo.ChannelMap.ChannelDestinationGains);
-			}
-
-			if (DownmixData.QuadSubmixInfo.bInUse)
-			{
-				Audio::Mix2ChannelsTo4ChannelsFast(*DownmixData.PostEffectBuffers, DownmixData.QuadSubmixInfo.OutputBuffer, DownmixData.QuadSubmixInfo.ChannelMap.ChannelStartGains, DownmixData.QuadSubmixInfo.ChannelMap.ChannelDestinationGains);
-			}
-
-			if (DownmixData.FiveOneSubmixInfo.bInUse)
-			{
-				Audio::Mix2ChannelsTo6ChannelsFast(*DownmixData.PostEffectBuffers, DownmixData.FiveOneSubmixInfo.OutputBuffer, DownmixData.FiveOneSubmixInfo.ChannelMap.ChannelStartGains, DownmixData.FiveOneSubmixInfo.ChannelMap.ChannelDestinationGains);
-			}
-
-			if (DownmixData.SevenOneSubmixInfo.bInUse)
-			{
-				Audio::Mix2ChannelsTo8ChannelsFast(*DownmixData.PostEffectBuffers, DownmixData.SevenOneSubmixInfo.OutputBuffer, DownmixData.SevenOneSubmixInfo.ChannelMap.ChannelStartGains, DownmixData.SevenOneSubmixInfo.ChannelMap.ChannelDestinationGains);
-			}
 		}
 		else
 		{
 			// Use generic calls:
 			Audio::DownmixBuffer(DownmixData.NumInputChannels, DownmixData.NumDeviceChannels, *DownmixData.PostEffectBuffers, DownmixData.DeviceSubmixInfo.OutputBuffer, DownmixData.DeviceSubmixInfo.ChannelMap.ChannelStartGains, DownmixData.DeviceSubmixInfo.ChannelMap.ChannelDestinationGains);
-
-			if (DownmixData.StereoSubmixInfo.bInUse)
-			{
-				Audio::DownmixBuffer(DownmixData.NumInputChannels, 2, *DownmixData.PostEffectBuffers, DownmixData.StereoSubmixInfo.OutputBuffer, DownmixData.StereoSubmixInfo.ChannelMap.ChannelStartGains, DownmixData.StereoSubmixInfo.ChannelMap.ChannelDestinationGains);
-			}
-
-			if (DownmixData.QuadSubmixInfo.bInUse)
-			{
-				Audio::DownmixBuffer(DownmixData.NumInputChannels, 4, *DownmixData.PostEffectBuffers, DownmixData.QuadSubmixInfo.OutputBuffer, DownmixData.QuadSubmixInfo.ChannelMap.ChannelStartGains, DownmixData.QuadSubmixInfo.ChannelMap.ChannelDestinationGains);
-			}
-
-			if (DownmixData.FiveOneSubmixInfo.bInUse)
-			{
-				Audio::DownmixBuffer(DownmixData.NumInputChannels, 6, *DownmixData.PostEffectBuffers, DownmixData.FiveOneSubmixInfo.OutputBuffer, DownmixData.FiveOneSubmixInfo.ChannelMap.ChannelStartGains, DownmixData.FiveOneSubmixInfo.ChannelMap.ChannelDestinationGains);
-			}
-
-			if (DownmixData.SevenOneSubmixInfo.bInUse)
-			{
-				Audio::DownmixBuffer(DownmixData.NumInputChannels, 8, *DownmixData.PostEffectBuffers, DownmixData.SevenOneSubmixInfo.OutputBuffer, DownmixData.SevenOneSubmixInfo.ChannelMap.ChannelStartGains, DownmixData.SevenOneSubmixInfo.ChannelMap.ChannelDestinationGains);
-			}
 		}
 
 		DownmixData.DeviceSubmixInfo.ChannelMap.CopyDestinationToStart();
 	}
 
-	void FMixerSourceManager::ComputeDownmix2D(FSourceDownmixData& DownmixData)
+	void FMixerSourceManager::ComputeDownmix2D(FSourceDownmixData& DownmixData, FMixerDevice* MixerDevice)
 	{
 		// This enormous switch statement handles using the correct function for a given number of input and output channels.
 		// For 2D sources, we just apply the gain matrix in ChannelDestionationGains with no interpolation.
-		if (DownmixData.NumInputChannels == 1)
+		if (DownmixData.AmbisonicsDecoder.IsValid())
+		{
+			FAmbisonicsSoundfieldBuffer AmbiBuffer;
+			AmbiBuffer.AudioBuffer = MoveTemp(*DownmixData.PostEffectBuffers);
+			AmbiBuffer.NumChannels = DownmixData.NumInputChannels;
+			AmbiBuffer.Rotation = DownmixData.PositionalData.Rotation;
+
+			DownmixData.PositionalData.NumChannels = DownmixData.NumDeviceChannels;
+			DownmixData.PositionalData.ChannelPositions = MixerDevice->GetDefaultPositionMap(DownmixData.NumDeviceChannels);
+
+			FSoundfieldDecoderInputData InputData =
+			{
+				AmbiBuffer,
+				DownmixData.PositionalData,
+				static_cast<int32>(DownmixData.PostEffectBuffers->Num() / DownmixData.NumInputChannels),
+				MixerDevice->GetSampleRate()
+			};
+
+			FSoundfieldDecoderOutputData OutputData = { DownmixData.DeviceSubmixInfo.OutputBuffer };
+
+			DownmixData.AmbisonicsDecoder->Decode(InputData, OutputData);
+
+			// Move the encoded ambisonics source buffer back to PostEffectBuffers to prevent reallocation
+			*DownmixData.PostEffectBuffers = MoveTemp(AmbiBuffer.AudioBuffer);
+		}
+		else if (DownmixData.NumInputChannels == 1)
 		{
 			switch (DownmixData.NumDeviceChannels)
 			{
@@ -1842,26 +1856,6 @@ namespace Audio
 				Audio::MixMonoTo2ChannelsFast(*DownmixData.PostEffectBuffers, DownmixData.DeviceSubmixInfo.OutputBuffer, DownmixData.DeviceSubmixInfo.ChannelMap.ChannelDestinationGains);
 				break;
 			}
-
-			if (DownmixData.StereoSubmixInfo.bInUse)
-			{
-				Audio::MixMonoTo2ChannelsFast(*DownmixData.PostEffectBuffers, DownmixData.StereoSubmixInfo.OutputBuffer, DownmixData.StereoSubmixInfo.ChannelMap.ChannelDestinationGains);
-			}
-
-			if (DownmixData.QuadSubmixInfo.bInUse)
-			{
-				Audio::MixMonoTo4ChannelsFast(*DownmixData.PostEffectBuffers, DownmixData.QuadSubmixInfo.OutputBuffer, DownmixData.QuadSubmixInfo.ChannelMap.ChannelDestinationGains);
-			}
-
-			if (DownmixData.FiveOneSubmixInfo.bInUse)
-			{
-				Audio::MixMonoTo6ChannelsFast(*DownmixData.PostEffectBuffers, DownmixData.FiveOneSubmixInfo.OutputBuffer, DownmixData.FiveOneSubmixInfo.ChannelMap.ChannelDestinationGains);
-			}
-
-			if (DownmixData.SevenOneSubmixInfo.bInUse)
-			{
-				Audio::MixMonoTo8ChannelsFast(*DownmixData.PostEffectBuffers, DownmixData.SevenOneSubmixInfo.OutputBuffer, DownmixData.SevenOneSubmixInfo.ChannelMap.ChannelDestinationGains);
-			}
 		}
 		else if (DownmixData.NumInputChannels == 2)
 		{
@@ -1881,133 +1875,145 @@ namespace Audio
 				break;
 			}
 
-			if (DownmixData.StereoSubmixInfo.bInUse)
-			{
-				Audio::Mix2ChannelsTo2ChannelsFast(*DownmixData.PostEffectBuffers, DownmixData.StereoSubmixInfo.OutputBuffer, DownmixData.StereoSubmixInfo.ChannelMap.ChannelDestinationGains);
-			}
+		}
+		else
+		{
+			// Use generic calls:
+			Audio::DownmixBuffer(DownmixData.NumInputChannels, DownmixData.NumDeviceChannels, *DownmixData.PostEffectBuffers, DownmixData.DeviceSubmixInfo.OutputBuffer, DownmixData.DeviceSubmixInfo.ChannelMap.ChannelDestinationGains);
+		}
+	}
 
-			if (DownmixData.QuadSubmixInfo.bInUse)
-			{
-				Audio::Mix2ChannelsTo4ChannelsFast(*DownmixData.PostEffectBuffers, DownmixData.QuadSubmixInfo.OutputBuffer, DownmixData.QuadSubmixInfo.ChannelMap.ChannelDestinationGains);
-			}
+	void FMixerSourceManager::EncodeToSoundfieldFormats(FSourceDownmixData& DownmixData, const FSourceInfo& InSource, FMixerDevice* InDevice)
+	{
+		const FSpatializationParams& SpatializationParams = InSource.SpatParams;
 
-			if (DownmixData.FiveOneSubmixInfo.bInUse)
+		// First, build our SoundfieldSpeakerPositionalData out of our SpatializationParams.
+		DownmixData.PositionalData.NumChannels = DownmixData.NumInputChannels;
+		DownmixData.PositionalData.Rotation = InSource.SpatParams.EmitterWorldRotation;
+		
+		// Build our input channel positions from the Spatialization params if this is a 3D source, otherwise use the default channel positions.
+		DownmixData.InputChannelPositions.Reset();
+		if (InSource.bIs3D)
+		{
+			if (DownmixData.NumInputChannels == 1)
 			{
-				Audio::Mix2ChannelsTo6ChannelsFast(*DownmixData.PostEffectBuffers, DownmixData.FiveOneSubmixInfo.OutputBuffer, DownmixData.FiveOneSubmixInfo.ChannelMap.ChannelDestinationGains);
+				Audio::FChannelPositionInfo ChannelPosition;
+				ChannelPosition.Channel = EAudioMixerChannel::FrontCenter;
+				ConvertCartesianToSpherical(SpatializationParams.EmitterPosition, ChannelPosition.Azimuth, ChannelPosition.Elevation, ChannelPosition.Radius);
+
+				DownmixData.InputChannelPositions.Add(ChannelPosition);
+
+				DownmixData.PositionalData.ChannelPositions = &DownmixData.InputChannelPositions;
 			}
+			else if (DownmixData.NumInputChannels == 2)
+			{
+				Audio::FChannelPositionInfo LeftChannelPosition;
+				LeftChannelPosition.Channel = EAudioMixerChannel::FrontLeft;
+				ConvertCartesianToSpherical(SpatializationParams.LeftChannelPosition, LeftChannelPosition.Azimuth, LeftChannelPosition.Elevation, LeftChannelPosition.Radius);
+
+				DownmixData.InputChannelPositions.Add(LeftChannelPosition);
+
+				Audio::FChannelPositionInfo RightChannelPosition;
+				LeftChannelPosition.Channel = EAudioMixerChannel::FrontRight;
+				ConvertCartesianToSpherical(SpatializationParams.RightChannelPosition, RightChannelPosition.Azimuth, RightChannelPosition.Elevation, RightChannelPosition.Radius);
+
+				DownmixData.InputChannelPositions.Add(RightChannelPosition);
 
-			if (DownmixData.SevenOneSubmixInfo.bInUse)
+				DownmixData.PositionalData.ChannelPositions = &DownmixData.InputChannelPositions;
+			}
+			else
 			{
-				Audio::Mix2ChannelsTo8ChannelsFast(*DownmixData.PostEffectBuffers, DownmixData.SevenOneSubmixInfo.OutputBuffer, DownmixData.SevenOneSubmixInfo.ChannelMap.ChannelDestinationGains);
+				// Spatialization of multichannel audio beyond stereo is not currently supported in the engine.
+				DownmixData.PositionalData.ChannelPositions = InDevice->GetDefaultPositionMap(DownmixData.NumInputChannels);
 			}
 		}
 		else
 		{
-			// Use generic calls:
-			Audio::DownmixBuffer(DownmixData.NumInputChannels, DownmixData.NumDeviceChannels, *DownmixData.PostEffectBuffers, DownmixData.DeviceSubmixInfo.OutputBuffer, DownmixData.DeviceSubmixInfo.ChannelMap.ChannelDestinationGains);
+			DownmixData.PositionalData.ChannelPositions = InDevice->GetDefaultPositionMap(DownmixData.NumInputChannels);
+		}
 
-			if (DownmixData.StereoSubmixInfo.bInUse)
-			{
-				Audio::DownmixBuffer(DownmixData.NumInputChannels, 2, *DownmixData.PostEffectBuffers, DownmixData.StereoSubmixInfo.OutputBuffer, DownmixData.StereoSubmixInfo.ChannelMap.ChannelDestinationGains);
-			}
+		// Finally, run our encoders.
+		for (auto& Soundfield : DownmixData.EncodedSoundfieldDownmixes)
+		{
+			FSubmixSoundfieldData& SoundfieldData = Soundfield.Value;
 
-			if (DownmixData.QuadSubmixInfo.bInUse)
-			{
-				Audio::DownmixBuffer(DownmixData.NumInputChannels, 4, *DownmixData.PostEffectBuffers, DownmixData.QuadSubmixInfo.OutputBuffer, DownmixData.QuadSubmixInfo.ChannelMap.ChannelDestinationGains);
-			}
+			check(SoundfieldData.EncoderSettings.IsValid());
+			check(SoundfieldData.EncodedPacket.IsValid());
 
-			if (DownmixData.FiveOneSubmixInfo.bInUse)
+			SoundfieldData.EncodedPacket->Reset();
+
+			// If this is an ambisonics source, transcode it. Otherwise, encode the source.
+			if (SoundfieldData.AmbiTranscoder)
 			{
-				Audio::DownmixBuffer(DownmixData.NumInputChannels, 6, *DownmixData.PostEffectBuffers, DownmixData.FiveOneSubmixInfo.OutputBuffer, DownmixData.FiveOneSubmixInfo.ChannelMap.ChannelDestinationGains);
+				FAmbisonicsSoundfieldBuffer AmbiBuffer;
+				AmbiBuffer.AudioBuffer = MoveTemp(*DownmixData.PostEffectBuffers);
+				AmbiBuffer.NumChannels = DownmixData.NumInputChannels;
 
+				SoundfieldData.AmbiTranscoder->Transcode(AmbiBuffer, GetAmbisonicsSourceDefaultSettings(), *SoundfieldData.EncodedPacket, *SoundfieldData.EncoderSettings);
+				*DownmixData.PostEffectBuffers = MoveTemp(AmbiBuffer.AudioBuffer);
 			}
-
-			if (DownmixData.SevenOneSubmixInfo.bInUse)
+			else if (SoundfieldData.Encoder)
 			{
-				Audio::DownmixBuffer(DownmixData.NumInputChannels, 8, *DownmixData.PostEffectBuffers, DownmixData.SevenOneSubmixInfo.OutputBuffer, DownmixData.SevenOneSubmixInfo.ChannelMap.ChannelDestinationGains);
-			}
+				FSoundfieldEncoderInputData InputData =
+				{
+					*DownmixData.PostEffectBuffers,
+					static_cast<int32>(DownmixData.NumInputChannels),
+					*SoundfieldData.EncoderSettings,
+					DownmixData.PositionalData
+				};
 
-			if (DownmixData.AmbisonicsSubmixInfo.bInUse)
+				SoundfieldData.Encoder->Encode(InputData, *SoundfieldData.EncodedPacket);
+			}
+			else if (SoundfieldData.bIsUnrealAmbisonicsSubmix)
 			{
-				Audio::DownmixBuffer(DownmixData.NumInputChannels, 4, *DownmixData.PostEffectBuffers, DownmixData.AmbisonicsSubmixInfo.OutputBuffer, DownmixData.AmbisonicsSubmixInfo.ChannelMap.ChannelDestinationGains);
+				FAmbisonicsSoundfieldBuffer& OutputPacket = DowncastSoundfieldRef<FAmbisonicsSoundfieldBuffer>(*SoundfieldData.EncodedPacket);
+				// Fixme: This is an array copy. Can we serve DownmixData.PostEffectBuffers directly to this soundfield?
+				OutputPacket.AudioBuffer = *DownmixData.PostEffectBuffers;
+				OutputPacket.NumChannels = DownmixData.NumInputChannels;
 			}
 		}
 	}
 
+	void FMixerSourceManager::ConvertCartesianToSpherical(const FVector& InVector, float& OutAzimuth, float& OutElevation, float& OutRadius)
+	{
+		// Convert coordinates from unreal cartesian system to left handed spherical coordinates (zenith is positive elevation, right is positive azimuth)
+		const float InX = -InVector.Z; //InVector.Y;
+		const float InY = InVector.X;// -InVector.Z;
+		const float InZ = -InVector.Y;
+
+
+		OutElevation = FMath::Atan2(InY, InX);
+
+		// Note, rather than using arccos(z / radius) here, we use Atan2 to avoid wrapping issues with negative elevation values.
+		OutAzimuth = FMath::Atan2(FMath::Sqrt(InX * InX + InY * InY), InZ);
+		OutRadius = InVector.Size();
+	}
+
 	FMixerSourceManager::FSourceDownmixData& FMixerSourceManager::InitializeDownmixForSource(const int32 SourceId, const int32 NumInputChannels, const int32 NumOutputChannels, const int32 InNumOutputFrames)
 	{
 		DownmixDataArray[SourceId].ResetData(NumInputChannels, NumOutputChannels);
 		return DownmixDataArray[SourceId];
 	}
 
-	const FMixerSourceManager::FSubmixChannelTypeInfo& FMixerSourceManager::GetChannelInfoForFormat(const ESubmixChannelFormat InFormat, const FSourceDownmixData& InDownmixData) const
+	const Audio::FMixerSourceManager::FSubmixChannelData& FMixerSourceManager::GetChannelInfoForDevice(const FSourceDownmixData& InDownmixData) const
 	{
-		switch (InFormat)
-		{
-		case ESubmixChannelFormat::Device:
-		{
-			return InDownmixData.DeviceSubmixInfo;
-		}
-		case ESubmixChannelFormat::SevenDotOne:
-		{
-			return InDownmixData.SevenOneSubmixInfo;
-		}
-		case ESubmixChannelFormat::FiveDotOne:
-		{
-			return InDownmixData.FiveOneSubmixInfo;
-		}
-		case ESubmixChannelFormat::Quad:
-		{
-			return InDownmixData.QuadSubmixInfo;
-		}
-		case ESubmixChannelFormat::Stereo:
-		{
-			return InDownmixData.StereoSubmixInfo;
-		}
-		case ESubmixChannelFormat::Ambisonics:
-		{
-			return InDownmixData.AmbisonicsSubmixInfo;
-		}
-		default:
-			check(false);
-			return InDownmixData.DeviceSubmixInfo;
-			break;
-		}
+		return InDownmixData.DeviceSubmixInfo;
 	}
 
-	FMixerSourceManager::FSubmixChannelTypeInfo& FMixerSourceManager::GetChannelInfoForFormat(const ESubmixChannelFormat InFormat, FSourceDownmixData& InDownmixData)
+	Audio::FMixerSourceManager::FSubmixChannelData& FMixerSourceManager::GetChannelInfoForDevice(FSourceDownmixData& InDownmixData)
 	{
-		switch (InFormat)
-		{
-		case ESubmixChannelFormat::Device:
-		{
-			return InDownmixData.DeviceSubmixInfo;
-		}
-		case ESubmixChannelFormat::SevenDotOne:
-		{
-			return InDownmixData.SevenOneSubmixInfo;
-		}
-		case ESubmixChannelFormat::FiveDotOne:
-		{
-			return InDownmixData.FiveOneSubmixInfo;
-		}
-		case ESubmixChannelFormat::Quad:
-		{
-			return InDownmixData.QuadSubmixInfo;
-		}
-		case ESubmixChannelFormat::Stereo:
-		{
-			return InDownmixData.StereoSubmixInfo;
-		}
-		case ESubmixChannelFormat::Ambisonics:
-		{
-			return InDownmixData.AmbisonicsSubmixInfo;
-		}
-		default:
-			check(false);
-			return InDownmixData.DeviceSubmixInfo;
-			break;
-		}
+		return InDownmixData.DeviceSubmixInfo;
+	}
+
+	const FMixerSourceManager::FSubmixSoundfieldData& FMixerSourceManager::GetChannelInfoForFormat(const FSoundfieldEncodingKey& InFormat, const FSourceDownmixData& InDownmixData) const
+	{
+		check(InDownmixData.EncodedSoundfieldDownmixes.Contains(InFormat));
+		return InDownmixData.EncodedSoundfieldDownmixes[InFormat];
+	}
+
+	FMixerSourceManager::FSubmixSoundfieldData& FMixerSourceManager::GetChannelInfoForFormat(const FSoundfieldEncodingKey& InFormat, FSourceDownmixData& InDownmixData)
+	{
+		return InDownmixData.EncodedSoundfieldDownmixes.FindOrAdd(InFormat);
 	}
 
 	void FMixerSourceManager::ComputePostSourceEffectBufferForIdRange(bool bGenerateBuses, const int32 SourceIdStart, const int32 SourceIdEnd)
@@ -2234,15 +2240,23 @@ namespace Audio
 				continue;
 			}
 
+			DownmixData.PositionalData.Rotation = SourceInfo.SpatParams.ListenerOrientation;
+
 			if (SourceInfo.bIs3D && !DownmixData.bIsInitialDownmix)
 			{
-				ComputeDownmix3D(DownmixData);
+				ComputeDownmix3D(DownmixData, MixerDevice);
 			}
 			else
 			{
-				ComputeDownmix2D(DownmixData);
+				ComputeDownmix2D(DownmixData, MixerDevice);
 				DownmixData.bIsInitialDownmix = false;
 			}
+
+			if (DownmixData.EncodedSoundfieldDownmixes.Num())
+			{
+				// Perform All Encoding for all sends.
+				EncodeToSoundfieldFormats(DownmixData, SourceInfo, MixerDevice);
+			}
 		}
 	}
 
@@ -2289,7 +2303,7 @@ namespace Audio
 		}
 	}
 
-	void FMixerSourceManager::MixOutputBuffers(const int32 SourceId, const ESubmixChannelFormat InSubmixChannelType, const float SendLevel, AlignedFloatBuffer& OutWetBuffer) const
+	void FMixerSourceManager::MixOutputBuffers(const int32 SourceId, int32 InNumOutputChannels, const float SendLevel, AlignedFloatBuffer& OutWetBuffer) const
 	{
 		if (SendLevel > 0.0f)
 		{
@@ -2299,11 +2313,7 @@ namespace Audio
 			if (!SourceInfo.bIsPaused && !SourceInfo.bIsDone && SourceInfo.bIsPlaying)
 			{
 				const FSourceDownmixData& DownmixData = DownmixDataArray[SourceId];
-				const FSubmixChannelTypeInfo& ChannelInfo = GetChannelInfoForFormat(InSubmixChannelType, DownmixData);
-// 				if (InSubmixChannelType == ESubmixChannelFormat::Ambisonics && SourceInfo.NumInputChannels == 1)
-// 				{
-// 					// TODO: encode 3d sources (non-ambisonics files) to ambisonics
-// 				}
+				const FSubmixChannelData& ChannelInfo = GetChannelInfoForDevice(DownmixData);
 
 				const float* RESTRICT SourceOutputBufferPtr = ChannelInfo.OutputBuffer.GetData();
 
@@ -2317,6 +2327,24 @@ namespace Audio
 		}
 	}
 
+	const ISoundfieldAudioPacket* FMixerSourceManager::GetEncodedOutput(const int32 SourceId, const FSoundfieldEncodingKey& InKey) const
+	{
+		AUDIO_MIXER_CHECK_AUDIO_PLAT_THREAD(MixerDevice);
+
+		const FSourceInfo& SourceInfo = SourceInfos[SourceId];
+
+		// Don't need to mix into submixes if the source is paused
+		if (!SourceInfo.bIsPaused && !SourceInfo.bIsDone && SourceInfo.bIsPlaying)
+		{
+			const FSourceDownmixData& DownmixData = DownmixDataArray[SourceId];
+			const FSubmixSoundfieldData& ChannelInfo = GetChannelInfoForFormat(InKey, DownmixData);
+
+			return ChannelInfo.EncodedPacket.Get();
+		}
+
+		return nullptr;
+	}
+
 	void FMixerSourceManager::UpdateDeviceChannelCount(const int32 InNumOutputChannels)
 	{
 		AudioMixerThreadCommand([this, InNumOutputChannels]()
@@ -2337,29 +2365,26 @@ namespace Audio
 				FSourceDownmixData& DownmixData = DownmixDataArray[SourceId];
 				DownmixData.ResetNumberOfDeviceChannels(InNumOutputChannels);
 
-				FSubmixChannelTypeInfo& ChannelTypeInfo = GetChannelInfoForFormat(ESubmixChannelFormat::Device, DownmixData);
+				FSubmixChannelData& ChannelTypeInfo = GetChannelInfoForDevice(DownmixData);
 
-				if (ChannelTypeInfo.bInUse)
+				SourceInfo.ScratchChannelMap.Reset();
+				const int32 NumSourceChannels = SourceInfo.bUseHRTFSpatializer ? 2 : SourceInfo.NumInputChannels;
+
+				// If this is a 3d source, then just zero out the channel map, it'll cause a temporary blip
+				// but it should reset in the next tick
+				if (SourceInfo.bIs3D)
+				{
+					GameThreadInfo.bNeedsSpeakerMap[SourceId] = true;
+					SourceInfo.ScratchChannelMap.AddZeroed(NumSourceChannels * InNumOutputChannels);
+				}
+				// If it's a 2D sound, then just get a new channel map appropriate for the new device channel count
+				else
 				{
 					SourceInfo.ScratchChannelMap.Reset();
-					const int32 NumSourceChannels = SourceInfo.bUseHRTFSpatializer ? 2 : SourceInfo.NumInputChannels;
-
-					// If this is a 3d source, then just zero out the channel map, it'll cause a temporary blip
-					// but it should reset in the next tick
-					if (SourceInfo.bIs3D)
-					{
-						GameThreadInfo.bNeedsSpeakerMap[SourceId] = true;
-						SourceInfo.ScratchChannelMap.AddZeroed(NumSourceChannels * InNumOutputChannels);
-					}
-					// If it's a 2D sound, then just get a new channel map appropriate for the new device channel count
-					else
-					{
-						SourceInfo.ScratchChannelMap.Reset();
-						MixerDevice->Get2DChannelMap(SourceInfo.bIsVorbis, ESubmixChannelFormat::Device, NumSourceChannels, SourceInfo.bIsCenterChannelOnly, SourceInfo.ScratchChannelMap);
-					}
-			
-					ChannelTypeInfo.ChannelMap.SetChannelMap(SourceInfo.ScratchChannelMap.GetData());
+					MixerDevice->Get2DChannelMap(SourceInfo.bIsVorbis, NumSourceChannels, InNumOutputChannels, SourceInfo.bIsCenterChannelOnly, SourceInfo.ScratchChannelMap);
 				}
+
+				ChannelTypeInfo.ChannelMap.SetChannelMap(SourceInfo.ScratchChannelMap.GetData());
 			}
 		});
 	}
diff --git a/Engine/Source/Runtime/AudioMixer/Private/AudioMixerSourceManager.h b/Engine/Source/Runtime/AudioMixer/Private/AudioMixerSourceManager.h
--- a/Engine/Source/Runtime/AudioMixer/Private/AudioMixerSourceManager.h
+++ b/Engine/Source/Runtime/AudioMixer/Private/AudioMixerSourceManager.h
@@ -17,6 +17,8 @@
 #include "DSP/ParamInterpolator.h"
 #include "DSP/BufferVectorOperations.h"
 #include "IAudioExtensionPlugin.h"
+#include "ISoundfieldFormat.h"
+#include "Containers/Queue.h"
 
 
 namespace Audio
@@ -73,6 +75,10 @@ namespace Audio
 
 		// Whather or not this is the primary send (i.e. first in the send chain)
 		bool bIsMainSend;
+
+		// If this is a soundfield submix, this is a pointer to the submix's Soundfield Factory.
+		// If this is nullptr, the submix is not a soundfield submix.
+		ISoundfieldFactory* SoundfieldFactory;
 	};
 
 	// Struct holding mappings of bus ids (unique ids) to send level
@@ -225,7 +231,7 @@ namespace Audio
 		void SetVolume(const int32 SourceId, const float Volume);
 		void SetDistanceAttenuation(const int32 SourceId, const float DistanceAttenuation);
 		void SetSpatializationParams(const int32 SourceId, const FSpatializationParams& InParams);
-		void SetChannelMap(const int32 SourceId, const ESubmixChannelFormat SubmixChannelType, const uint32 NumInputChannels, const Audio::AlignedFloatBuffer& InChannelMap, const bool bInIs3D, const bool bInIsCenterChannelOnly);
+		void SetChannelMap(const int32 SourceId, const uint32 NumInputChannels, const Audio::AlignedFloatBuffer& InChannelMap, const bool bInIs3D, const bool bInIsCenterChannelOnly);
 		void SetLPFFrequency(const int32 SourceId, const float Frequency);
 		void SetHPFFrequency(const int32 SourceId, const float Frequency);
 
@@ -238,7 +244,13 @@ namespace Audio
 		bool NeedsSpeakerMap(const int32 SourceId) const;
 		void ComputeNextBlockOfSamples();
 		void ClearStoppingSounds();
-		void MixOutputBuffers(const int32 SourceId, const ESubmixChannelFormat InSubmixChannelType, const float SendLevel, AlignedFloatBuffer& OutWetBuffer) const;
+		void MixOutputBuffers(const int32 SourceId, int32 InNumOutputChannels, const float SendLevel, AlignedFloatBuffer& OutWetBuffer) const;
+
+		// Called by a soundfield submix to get encoded audio.
+		// If this source wasn't encoded (possibly because it is paused or finished playing),
+		// this returns nullptr.
+		// Returned nonnull pointers are only guaranteed to be valid on the audio mixer render thread.
+		const ISoundfieldAudioPacket* GetEncodedOutput(const int32 SourceId, const FSoundfieldEncodingKey& InKey) const;
 
 		void SetSubmixSendInfo(const int32 SourceId, const FMixerSourceSubmixSend& SubmixSend);
 		void SetBusSendInfo(const int32 SourceId, EBusSendType InBusSendType, FMixerBusSend& BusSend);
@@ -333,15 +345,14 @@ namespace Audio
 
 		TArray<int32> DebugSoloSources;
 
-		struct FSubmixChannelTypeInfo
+		// All state having to do with a non-soundfield downmix.
+		struct FSubmixChannelData
 		{
 			FSourceChannelMap ChannelMap;
 			Audio::AlignedFloatBuffer OutputBuffer;
-			bool bInUse;
 
-			FSubmixChannelTypeInfo(uint32 InNumInChannels, uint32 InNumOutputChannels, uint32 NumFrames)
+			FSubmixChannelData(uint32 InNumInChannels, uint32 InNumOutputChannels, uint32 NumFrames)
 				: ChannelMap(InNumInChannels, InNumOutputChannels)
-				, bInUse(false)
 			{
 				OutputBuffer.Reset();
 				OutputBuffer.AddUninitialized(NumFrames * InNumOutputChannels);
@@ -356,6 +367,17 @@ namespace Audio
 			}
 		};
 
+		struct FSubmixSoundfieldData
+		{
+			TUniquePtr<ISoundfieldEncoderStream> Encoder;
+			// If this is an ambisonics source, we use a transcoder stream.
+			TUniquePtr<ISoundfieldTranscodeStream> AmbiTranscoder;
+			TUniquePtr<ISoundfieldEncodingSettingsProxy> EncoderSettings;
+			TUniquePtr<ISoundfieldAudioPacket> EncodedPacket;
+			// If this is a unreal ambisonics soundfield buffer, we hand it the submixed buffer directly.
+			bool bIsUnrealAmbisonicsSubmix;
+		};
+
 		struct FSourceDownmixData
 		{
 			// Output data, after computing a block of sample data, this is read back from mixers
@@ -363,26 +385,25 @@ namespace Audio
 			Audio::AlignedFloatBuffer* PostEffectBuffers;
 
 			// Data needed for outputting to submixes for each channel configuration.
-			FSubmixChannelTypeInfo DeviceSubmixInfo;
-			FSubmixChannelTypeInfo StereoSubmixInfo;
-			FSubmixChannelTypeInfo QuadSubmixInfo;
-			FSubmixChannelTypeInfo FiveOneSubmixInfo;
-			FSubmixChannelTypeInfo SevenOneSubmixInfo;
-			FSubmixChannelTypeInfo AmbisonicsSubmixInfo;
+			FSubmixChannelData DeviceSubmixInfo;
+			
+			TMap<FSoundfieldEncodingKey, FSubmixSoundfieldData> EncodedSoundfieldDownmixes;
+			TArray<Audio::FChannelPositionInfo> InputChannelPositions;
 
 			uint32 NumInputChannels;
 			const uint32 NumFrames;
 			uint32 NumDeviceChannels;
 			uint8 bIsInitialDownmix : 1;
 
+			// cached parameters for encoding to a soundfield format.
+			FSoundfieldSpeakerPositionalData PositionalData;
+
+			// If this source is an ambisonics source, we use this to downmix the source to our output.
+			TUniquePtr<ISoundfieldDecoderStream> AmbisonicsDecoder;
+
 			FSourceDownmixData(uint32 SourceNumChannels, uint32 NumDeviceOutputChannels, uint32 InNumFrames)
 				: PostEffectBuffers(nullptr)
-				, DeviceSubmixInfo(FSubmixChannelTypeInfo(SourceNumChannels, NumDeviceOutputChannels, InNumFrames))
-				, StereoSubmixInfo(FSubmixChannelTypeInfo(SourceNumChannels, 2, InNumFrames))
-				, QuadSubmixInfo(FSubmixChannelTypeInfo(SourceNumChannels, 4, InNumFrames))
-				, FiveOneSubmixInfo(FSubmixChannelTypeInfo(SourceNumChannels, 6, InNumFrames))
-				, SevenOneSubmixInfo(FSubmixChannelTypeInfo(SourceNumChannels, 8, InNumFrames))
-				, AmbisonicsSubmixInfo(FSubmixChannelTypeInfo(SourceNumChannels, 4, InNumFrames))
+				, DeviceSubmixInfo(FSubmixChannelData(SourceNumChannels, NumDeviceOutputChannels, InNumFrames))
 				, NumInputChannels(SourceNumChannels)
 				, NumFrames(InNumFrames)
 				, NumDeviceChannels(NumDeviceOutputChannels)
@@ -403,11 +424,8 @@ namespace Audio
 				PostEffectBuffers = nullptr;
 
 				DeviceSubmixInfo.Reset(NumInputChannels, NumDeviceChannels, NumFrames);
-				StereoSubmixInfo.Reset(NumInputChannels, 2, NumFrames);
-				QuadSubmixInfo.Reset(NumInputChannels, 4, NumFrames);
-				FiveOneSubmixInfo.Reset(NumInputChannels, 6, NumFrames);
-				SevenOneSubmixInfo.Reset(NumInputChannels, 8, NumFrames);
-				AmbisonicsSubmixInfo.Reset(NumInputChannels, 4, NumFrames);
+				EncodedSoundfieldDownmixes.Reset();
+				AmbisonicsDecoder.Reset();
 				bIsInitialDownmix = true;
 			}
 		};
@@ -503,9 +521,14 @@ namespace Audio
 			uint8 bIsLastBuffer:1;
 			uint8 bOutputToBusOnly:1;
 			uint8 bIsVorbis:1;
+
+			// TODO: Make this generic to soundfield formats as a whole.
+			// For now, if this is true we just wrap the audio buffer in an FAmbisonicsSoundfieldBuffer.
+			uint8 bIsAmbisonics:1;
 			uint8 bIsBypassingLPF:1;
 			uint8 bIsBypassingHPF:1;
 			uint8 bIsModulationUpdated:1;
+
 			// Source format info
 			int32 NumInputChannels;
 			int32 NumPostEffectChannels;
@@ -523,14 +546,19 @@ namespace Audio
 		static void ApplyDistanceAttenuation(FSourceInfo& InSourceInfo, int32 NumSamples);
 		void ComputePluginAudio(FSourceInfo& InSourceInfo, FSourceDownmixData& DownmixData, int32 SourceId, int32 NumSamples);
 
-		static void ComputeDownmix3D(FSourceDownmixData& DownmixData);
-		static void ComputeDownmix2D(FSourceDownmixData& DownmixData);
+		static void ComputeDownmix3D(FSourceDownmixData& DownmixData, FMixerDevice* MixerDevice);
+		static void ComputeDownmix2D(FSourceDownmixData& DownmixData, FMixerDevice* MixerDevice);
+		static void EncodeToSoundfieldFormats(FSourceDownmixData& DownmixData, const FSourceInfo& InSource, FMixerDevice* InDevice);
+		static void ConvertCartesianToSpherical(const FVector& InVector, float& OutAzimuth, float& OutElevation, float& OutRadius);
 
 		// This function is effectively equivalent to calling DownmixDataArray.EmplaceAt_GetRef(args...), but bypasses that function's intrinsic call to AddUninitialized.
 		FSourceDownmixData& InitializeDownmixForSource(const int32 SourceId, const int32 NumInputChannels, const int32 NumOutputChannels, const int32 InNumOutputFrames);
 
-		const FSubmixChannelTypeInfo& GetChannelInfoForFormat(const ESubmixChannelFormat InFormat, const FSourceDownmixData& InDownmixData) const;
-		FSubmixChannelTypeInfo& GetChannelInfoForFormat(const ESubmixChannelFormat InFormat, FSourceDownmixData& InDownmixData);
+		const FSubmixChannelData& GetChannelInfoForDevice(const FSourceDownmixData& InDownmixData) const;
+		FSubmixChannelData& GetChannelInfoForDevice(FSourceDownmixData& InDownmixData);
+
+		const FSubmixSoundfieldData& GetChannelInfoForFormat(const FSoundfieldEncodingKey& InFormat, const FSourceDownmixData& InDownmixData) const;
+		FSubmixSoundfieldData& GetChannelInfoForFormat(const FSoundfieldEncodingKey& InFormat, FSourceDownmixData& InDownmixData);
 
 		// Array of listener transforms
 		TArray<FTransform> ListenerTransforms;
diff --git a/Engine/Source/Runtime/AudioMixer/Private/AudioMixerSourceVoice.cpp b/Engine/Source/Runtime/AudioMixer/Private/AudioMixerSourceVoice.cpp
--- a/Engine/Source/Runtime/AudioMixer/Private/AudioMixerSourceVoice.cpp
+++ b/Engine/Source/Runtime/AudioMixer/Private/AudioMixerSourceVoice.cpp
@@ -141,11 +141,11 @@ namespace Audio
 		}
 	}
 
-	void FMixerSourceVoice::SetChannelMap(ESubmixChannelFormat InChannelType, const uint32 NumInputChannels, const Audio::AlignedFloatBuffer& InChannelMap, const bool bInIs3D, const bool bInIsCenterChannelOnly)
+	void FMixerSourceVoice::SetChannelMap(const uint32 NumInputChannels, const Audio::AlignedFloatBuffer& InChannelMap, const bool bInIs3D, const bool bInIsCenterChannelOnly)
 	{
 		AUDIO_MIXER_CHECK_GAME_THREAD(MixerDevice);
 
-		SourceManager->SetChannelMap(SourceId, InChannelType, NumInputChannels, InChannelMap, bInIs3D, bInIsCenterChannelOnly);
+		SourceManager->SetChannelMap(SourceId, NumInputChannels, InChannelMap, bInIs3D, bInIsCenterChannelOnly);
 	}
 
 	void FMixerSourceVoice::SetSpatializationParams(const FSpatializationParams& InParams)
@@ -249,13 +249,22 @@ namespace Audio
 		return SourceManager->GetEnvelopeValue(SourceId);
 	}
 
-	void FMixerSourceVoice::MixOutputBuffers(const ESubmixChannelFormat InSubmixChannelType, const float SendLevel, AlignedFloatBuffer& OutWetBuffer) const
+	void FMixerSourceVoice::MixOutputBuffers(int32 InNumOutputChannels, const float SendLevel, AlignedFloatBuffer& OutWetBuffer) const
 	{
 		AUDIO_MIXER_CHECK_AUDIO_PLAT_THREAD(MixerDevice);
 
 		check(!bOutputToBusOnly);
 
-		return SourceManager->MixOutputBuffers(SourceId, InSubmixChannelType, SendLevel, OutWetBuffer);
+		return SourceManager->MixOutputBuffers(SourceId, InNumOutputChannels, SendLevel, OutWetBuffer);
+	}
+
+	const ISoundfieldAudioPacket* FMixerSourceVoice::GetEncodedOutput(const FSoundfieldEncodingKey& InKey) const
+	{
+		AUDIO_MIXER_CHECK_AUDIO_PLAT_THREAD(MixerDevice);
+
+		check(!bOutputToBusOnly);
+
+		return SourceManager->GetEncodedOutput(SourceId, InKey);
 	}
 
 	void FMixerSourceVoice::SetSubmixSendInfo(FMixerSubmixWeakPtr Submix, const float SendLevel)
diff --git a/Engine/Source/Runtime/AudioMixer/Private/AudioMixerSourceVoice.h b/Engine/Source/Runtime/AudioMixer/Private/AudioMixerSourceVoice.h
--- a/Engine/Source/Runtime/AudioMixer/Private/AudioMixerSourceVoice.h
+++ b/Engine/Source/Runtime/AudioMixer/Private/AudioMixerSourceVoice.h
@@ -49,7 +49,7 @@ namespace Audio
 		void SetHPFFrequency(const float InFrequency);
 
 		// Sets the source voice's channel map (2d or 3d).
-		void SetChannelMap(ESubmixChannelFormat InChannelType, const uint32 NumInputChannels, const Audio::AlignedFloatBuffer& InChannelMap, const bool bInIs3D, const bool bInIsCenterChannelOnly);
+		void SetChannelMap(const uint32 NumInputChannels, const Audio::AlignedFloatBuffer& InChannelMap, const bool bInIs3D, const bool bInIsCenterChannelOnly);
 
 		// Sets params used by HRTF spatializer
 		void SetSpatializationParams(const FSpatializationParams& InParams);
@@ -96,7 +96,10 @@ namespace Audio
 		float GetEnvelopeValue() const;
 
 		// Mixes the dry and wet buffer audio into the given buffers.
-		void MixOutputBuffers(const ESubmixChannelFormat InSubmixChannelType, const float SendLevel, AlignedFloatBuffer& OutWetBuffer) const;
+		void MixOutputBuffers(int32 InNumChannels, const float SendLevel, AlignedFloatBuffer& OutWetBuffer) const;
+
+		// For soundfield conversions, get the encoded audio.
+		const ISoundfieldAudioPacket* GetEncodedOutput(const FSoundfieldEncodingKey& InKey) const;
 
 		// Sets the submix send levels
 		void SetSubmixSendInfo(FMixerSubmixWeakPtr Submix, const float SendLevel);
@@ -114,7 +117,7 @@ namespace Audio
 		FMixerSourceManager* SourceManager;
 		TMap<uint32, FMixerSourceSubmixSend> SubmixSends;
 		FMixerDevice* MixerDevice;
-		TMap<ESubmixChannelFormat, TArray<float>> ChannelMaps;
+		TArray<float> DeviceChannelMap;
 		FThreadSafeBool bStopFadedOut;
 		float Pitch;
 		float Volume;
