Editgrate 4.25 audio features from project stream:

diff --git a/Engine/Source/Runtime/AudioMixer/Private/AudioMixerSubmix.cpp b/Engine/Source/Runtime/AudioMixer/Private/AudioMixerSubmix.cpp
--- a/Engine/Source/Runtime/AudioMixer/Private/AudioMixerSubmix.cpp
+++ b/Engine/Source/Runtime/AudioMixer/Private/AudioMixerSubmix.cpp
@@ -19,21 +19,25 @@ FAutoConsoleVariableRef CVarRecoverRecordingOnShutdown(
 	TEXT("0: Disabled, 1: Enabled"),
 	ECVF_Default);
 
+static int32 AllowInitOnAudioRenderThreadCVar = 0;
+FAutoConsoleVariableRef CVarAllowInitOnAudioRenderThread(
+	TEXT("au.AllowInitOnAudioRenderThread"),
+	RecoverRecordingOnShutdownCVar,
+	TEXT("When set to 1, we will initialize the submix on the audio render thread.\n")
+	TEXT("0: Disabled, 1: Enabled"),
+	ECVF_Default);
+
 namespace Audio
 {
 	// Unique IDs for mixer submixes
 	static uint32 GSubmixMixerIDs = 0;
 
 	FMixerSubmix::FMixerSubmix(FMixerDevice* InMixerDevice)
-		: AmbisonicsSettings(nullptr)
-		, Id(GSubmixMixerIDs++)
+		: Id(GSubmixMixerIDs++)
 		, ParentSubmix(nullptr)
 		, MixerDevice(InMixerDevice)
-		, ChannelFormat(ESubmixChannelFormat::Device)
 		, NumChannels(0)
 		, NumSamples(0)
-		, SubmixAmbisonicsEncoderID(INDEX_NONE)
-		, SubmixAmbisonicsDecoderID(INDEX_NONE)
 		, InitializedOutputVolume(1.0f)
 		, OutputVolume(1.0f)
 		, TargetOutputVolume(1.0f)
@@ -48,25 +52,20 @@ namespace Audio
 	FMixerSubmix::~FMixerSubmix()
 	{
 		ClearSoundEffectSubmixes();
-		if (SubmixAmbisonicsEncoderID != INDEX_NONE)
-		{
-			TearDownAmbisonicsEncoder();
-		}
-
-		if (SubmixAmbisonicsDecoderID != INDEX_NONE)
-		{
-			TearDownAmbisonicsDecoder();
-		}
 
 		if (RecoverRecordingOnShutdownCVar && OwningSubmixObject && bIsRecording)
 		{
 			FString InterruptedFileName = TEXT("InterruptedRecording.wav");
 			UE_LOG(LogAudioMixer, Warning, TEXT("Recording of Submix %s was interrupted. Saving interrupted recording as %s."), *(OwningSubmixObject->GetName()), *InterruptedFileName);
-			OwningSubmixObject->StopRecordingOutput(MixerDevice, EAudioRecordingExportType::WavFile, InterruptedFileName, FString());
+			if (const USoundSubmix* SoundSubmix = Cast<const USoundSubmix>(OwningSubmixObject))
+			{
+				USoundSubmix* MutableSubmix = const_cast<USoundSubmix*>(SoundSubmix);
+				MutableSubmix->StopRecordingOutput(MixerDevice, EAudioRecordingExportType::WavFile, InterruptedFileName, FString());
+			}
 		}
 	}
 
-	void FMixerSubmix::Init(USoundSubmix* InSoundSubmix, bool bAllowReInit)
+	void FMixerSubmix::Init(const USoundSubmixBase* InSoundSubmix, bool bAllowReInit)
 	{
 		check(IsInAudioThread());
 		if (InSoundSubmix != nullptr)
@@ -81,98 +80,147 @@ namespace Audio
 			else if (bAllowReInit)
 			{
 				check(OwningSubmixObject == InSoundSubmix);
-				SubmixCommand([this]()
+				if (AllowInitOnAudioRenderThreadCVar)
+				{
+					SubmixCommand([this]()
+					{
+						InitInternal();
+					});
+				}
+				else
 				{
 					InitInternal();
-				});
+				}
 			}
 		}
 	}
 
 	void FMixerSubmix::InitInternal()
 	{
-		// Set the initialized output volume
-		InitializedOutputVolume = FMath::Clamp(OwningSubmixObject->OutputVolume, 0.0f, 1.0f);
-
-		if (!FMath::IsNearlyEqual(InitializedOutputVolume, 1.0f))
-		{
-			bApplyOutputVolumeScale = true;
-		}
-
 		// Loop through the submix's presets and make new instances of effects in the same order as the presets
 		ClearSoundEffectSubmixes();
 
-		for (USoundEffectSubmixPreset* EffectPreset : OwningSubmixObject->SubmixEffectChain)
+
+		if (const USoundSubmix* SoundSubmix = Cast<const USoundSubmix>(OwningSubmixObject))
 		{
-			if (EffectPreset)
+			// Set the initialized output volume
+			InitializedOutputVolume = FMath::Clamp(SoundSubmix->OutputVolume, 0.0f, 1.0f);
+
+			if (!FMath::IsNearlyEqual(InitializedOutputVolume, 1.0f))
+			{
+				bApplyOutputVolumeScale = true;
+			}
+
+			for (USoundEffectSubmixPreset* EffectPreset : SoundSubmix->SubmixEffectChain)
 			{
-				// Create a new effect instance using the preset
-				FSoundEffectSubmix* SubmixEffect = static_cast<FSoundEffectSubmix*>(EffectPreset->CreateNewEffect());
+				if (EffectPreset)
+				{
+					// Create a new effect instance using the preset
+					FSoundEffectSubmix* SubmixEffect = static_cast<FSoundEffectSubmix*>(EffectPreset->CreateNewEffect());
 				
 				FSoundEffectSubmixInitData InitData;
 				InitData.DeviceID = MixerDevice->DeviceID;
 				InitData.SampleRate = MixerDevice->GetSampleRate();
 				InitData.PresetSettings = nullptr;
 
-				// Now set the preset
-				SubmixEffect->Init(InitData);
-				SubmixEffect->SetPreset(EffectPreset);
-				SubmixEffect->SetEnabled(true);
+					// Now set the preset
+					SubmixEffect->Init(InitData);
+					SubmixEffect->SetPreset(EffectPreset);
+					SubmixEffect->SetEnabled(true);
 
-				FSubmixEffectInfo EffectInfo;
-				EffectInfo.PresetId = EffectPreset->GetUniqueID();
-				EffectInfo.EffectInstance = MakeShareable(SubmixEffect);
+					FSubmixEffectInfo EffectInfo;
+					EffectInfo.PresetId = EffectPreset->GetUniqueID();
+					EffectInfo.EffectInstance = MakeShareable(SubmixEffect);
 
-				// Add the effect to this submix's chain
-				EffectSubmixChain.Add(MoveTemp(EffectInfo));
+					// Add the effect to this submix's chain
+					EffectSubmixChain.Add(MoveTemp(EffectInfo));
+				}
 			}
+
+			NumChannels = MixerDevice->GetNumDeviceChannels();
+			const int32 NumOutputFrames = MixerDevice->GetNumOutputFrames();
+			NumSamples = NumChannels * NumOutputFrames;
 		}
+		else if (const USoundfieldSubmix* SoundfieldSubmix = Cast<const USoundfieldSubmix>(OwningSubmixObject))
+		{
+			ISoundfieldFactory* SoundfieldFactory = SoundfieldSubmix->GetSoundfieldFactoryForSubmix();
+			const USoundfieldEncodingSettingsBase* EncodingSettings = SoundfieldSubmix->GetSoundfieldEncodingSettings();
 
-		ChannelFormat = OwningSubmixObject->ChannelFormat;
+			TArray<USoundfieldEffectBase*> Effects = SoundfieldSubmix->GetSoundfieldProcessors();
+			SetupSoundfieldStreams(EncodingSettings, Effects, SoundfieldFactory);
+		}
+		else if (const UEndpointSubmix* EndpointSubmix = Cast<const UEndpointSubmix>(OwningSubmixObject))
+		{
+			NumChannels = MixerDevice->GetNumDeviceChannels();
+			const int32 NumOutputFrames = MixerDevice->GetNumOutputFrames();
+			NumSamples = NumChannels * NumOutputFrames;
+
+			IAudioEndpointFactory* EndpointFactory = EndpointSubmix->GetAudioEndpointForSubmix();
+			const UAudioEndpointSettingsBase* EndpointSettings = EndpointSubmix->GetEndpointSettings();
 
-		if (ChannelFormat == ESubmixChannelFormat::Ambisonics)
+			SetupEndpoint(EndpointFactory, EndpointSettings);
+		}
+		else if (const USoundfieldEndpointSubmix* SoundfieldEndpointSubmix = Cast<const USoundfieldEndpointSubmix>(OwningSubmixObject))
 		{
-			//Get the ambisonics mixer.
-			AmbisonicsMixer = MixerDevice->GetAmbisonicsMixer();
+ 			ISoundfieldEndpointFactory* SoundfieldFactory = SoundfieldEndpointSubmix->GetSoundfieldEndpointForSubmix();
+			const USoundfieldEncodingSettingsBase* EncodingSettings = SoundfieldEndpointSubmix->GetEncodingSettings();
 
-			//If we do have a valid ambisonics decoder, lets use it. Otherwise, treat this submix like a device submix.
-			if (AmbisonicsMixer.IsValid())
+			if (!SoundfieldFactory)
 			{
-				UAmbisonicsSubmixSettingsBase* AmbisonicSubmixSettings = OwningSubmixObject->AmbisonicsPluginSettings;
-				if (!AmbisonicSubmixSettings)
-				{
-					AmbisonicSubmixSettings = AmbisonicsMixer->GetDefaultSettings();
-				}
+				UE_LOG(LogAudio, Display, TEXT("Wasn't able to set up soundfield format for submix %s. Sending to default output."), *OwningSubmixObject->GetName());
+				return;
+			}
+
+			if (!EncodingSettings)
+			{
+				EncodingSettings = SoundfieldFactory->GetDefaultEncodingSettings();
 
-				if (AmbisonicSubmixSettings)
+				if (!ensureMsgf(EncodingSettings, TEXT("Soundfield Endpoint %s did not return default encoding settings! Is ISoundfieldEndpointFactory::GetDefaultEncodingSettings() implemented?"), *SoundfieldFactory->GetEndpointTypeName().ToString()))
 				{
-					OnAmbisonicsSettingsChanged(AmbisonicSubmixSettings);
+					return;
 				}
-				else
+			}
+
+			TArray<USoundfieldEffectBase*> Effects = SoundfieldEndpointSubmix->GetSoundfieldProcessors();
+
+			SetupSoundfieldStreams(EncodingSettings, Effects, SoundfieldFactory);
+
+			if (IsSoundfieldSubmix())
+			{
+				const USoundfieldEndpointSettingsBase* EndpointSettings = SoundfieldEndpointSubmix->GetEndpointSettings();
+
+				if (!EndpointSettings)
 				{
-					//Default to first order ambisonics.
-					NumChannels = 4;
-					const int32 NumOutputFrames = MixerDevice->GetNumOutputFrames();
-					NumSamples = NumChannels * NumOutputFrames;
+					EndpointSettings = SoundfieldFactory->GetDefaultEndpointSettings();
+
+					if (!ensureMsgf(EncodingSettings, TEXT("Soundfield Endpoint %s did not return default encoding settings! Is ISoundfieldEndpointFactory::GetDefaultEndpointSettings() implemented?"), *SoundfieldFactory->GetEndpointTypeName().ToString()))
+					{
+						return;
+					}
 				}
+
+				SetupEndpoint(SoundfieldFactory, EndpointSettings);
 			}
 			else
 			{
-				// There is no valid ambisonics decoder, so fall back to standard downmixing.
-				ChannelFormat = ESubmixChannelFormat::Device;
-				NumChannels = MixerDevice->GetNumChannelsForSubmixFormat(ChannelFormat);
-				const int32 NumOutputFrames = MixerDevice->GetNumOutputFrames();
-				NumSamples = NumChannels * NumOutputFrames;
+				UE_LOG(LogAudio, Display, TEXT("Wasn't able to set up soundfield format for submix %s. Sending to default output."), *OwningSubmixObject->GetName());
+				SoundfieldStreams.Reset();
 			}
 		}
 		else
 		{
-			NumChannels = MixerDevice->GetNumChannelsForSubmixFormat(ChannelFormat);
-			const int32 NumOutputFrames = MixerDevice->GetNumOutputFrames();
-			NumSamples = NumChannels * NumOutputFrames;
+			// If we've arrived here, we couldn't identify the type of the submix we're initializing.
+			checkNoEntry();
 		}
 	}
 
+	void FMixerSubmix::DownmixBuffer(const int32 InChannels, const AlignedFloatBuffer& InBuffer, const int32 OutChannels, AlignedFloatBuffer& OutNewBuffer)
+	{
+		Audio::AlignedFloatBuffer MixdownGainsMap;
+		Audio::FMixerDevice::Get2DChannelMap(false, InChannels, OutChannels, false, MixdownGainsMap);
+		Audio::DownmixBuffer(InChannels, OutChannels, InBuffer, OutNewBuffer, MixdownGainsMap.GetData());
+	}
+
 	void FMixerSubmix::SetParentSubmix(TWeakPtr<FMixerSubmix, ESPMode::ThreadSafe> SubmixWeakPtr)
 	{
 		if (ParentSubmix == SubmixWeakPtr)
@@ -197,9 +245,9 @@ namespace Audio
 			AUDIO_MIXER_CHECK_AUDIO_PLAT_THREAD(MixerDevice);
 
 			ParentSubmix = SubmixWeakPtr;
-			if (ChannelFormat == ESubmixChannelFormat::Ambisonics && AmbisonicsMixer.IsValid())
+			if (IsSoundfieldSubmix())
 			{
-				UpdateAmbisonicsDecoderForParent();
+				SetupSoundfieldStreamForParent();
 			}
 		});
 	}
@@ -210,20 +258,14 @@ namespace Audio
 		{
 			AUDIO_MIXER_CHECK_AUDIO_PLAT_THREAD(MixerDevice);
 
-			FChildSubmixInfo NewChildSubmixInfo;
-			NewChildSubmixInfo.SubmixPtr = SubmixWeakPtr;
-
-			//TODO: switch this conditionally when we are able to route submixes to ambisonics submix.
-			NewChildSubmixInfo.bNeedsAmbisonicsEncoding = false;
-
 			TSharedPtr<Audio::FMixerSubmix, ESPMode::ThreadSafe> SubmixSharedPtr = SubmixWeakPtr.Pin();
 			if (SubmixSharedPtr.IsValid())
 			{
-				ChildSubmixes.Add(SubmixSharedPtr->GetId(), NewChildSubmixInfo);
+				FChildSubmixInfo& ChildSubmixInfo = ChildSubmixes.Emplace(SubmixSharedPtr->GetId(), SubmixWeakPtr);
 
-				if (ChannelFormat == ESubmixChannelFormat::Ambisonics && AmbisonicsMixer.IsValid())
+				if (IsSoundfieldSubmix())
 				{
-					UpdateAmbisonicsEncoderForChildren();
+					SetupSoundfieldEncodingForChild(ChildSubmixInfo);
 				}
 			}
 		});
@@ -243,17 +285,12 @@ namespace Audio
 			AUDIO_MIXER_CHECK_AUDIO_PLAT_THREAD(MixerDevice);
 
 			ChildSubmixes.Remove(OldIdToRemove);
-
-			if (ChannelFormat == ESubmixChannelFormat::Ambisonics && AmbisonicsMixer.IsValid())
-			{
-				UpdateAmbisonicsEncoderForChildren();
-			}
 		});
 	}
 
-	ESubmixChannelFormat FMixerSubmix::GetSubmixChannels() const
+	int32 FMixerSubmix::GetSubmixChannels() const
 	{
-		return ChannelFormat;
+		return NumChannels;
 	}
 
 	TWeakPtr<FMixerSubmix, ESPMode::ThreadSafe> FMixerSubmix::GetParentSubmix()
@@ -277,22 +314,6 @@ namespace Audio
 
 		FSubmixVoiceData NewVoiceData;
 		NewVoiceData.SendLevel = InSendLevel;
-		
-		// If this is an ambisonics submix, set up a new encoder stream.
-		if (ChannelFormat == ESubmixChannelFormat::Ambisonics && AmbisonicsMixer.IsValid())
-		{
-			//TODO: If a souce is not an ambisonics source, we need to set up an encoder for it.
-			const bool bSourceIsAmbisonics = true;
-			if (!bSourceIsAmbisonics)
-			{
-				NewVoiceData.AmbisonicsEncoderId = MixerDevice->GetNewUniqueAmbisonicsStreamID();
-				AmbisonicsMixer->OnOpenEncodingStream(NewVoiceData.AmbisonicsEncoderId, AmbisonicsSettings);
-			}
-			else
-			{
-				NewVoiceData.AmbisonicsEncoderId = INDEX_NONE;
-			}
-		}
 
 		MixerSourceVoices.Add(InSourceVoice, NewVoiceData);
 	}
@@ -304,16 +325,6 @@ namespace Audio
 		// If the source has a corresponding ambisonics encoder, close it out.
 		uint32 SourceEncoderID = INDEX_NONE;
 		const FSubmixVoiceData* MixerSourceVoiceData = MixerSourceVoices.Find(InSourceVoice);
-		if (MixerSourceVoiceData)
-		{
-			SourceEncoderID = MixerSourceVoiceData->AmbisonicsEncoderId;
-		}
-		
-		if (SourceEncoderID != INDEX_NONE)
-		{
-			check(AmbisonicsMixer.IsValid());
-			AmbisonicsMixer->OnCloseEncodingStream(SourceEncoderID);
-		}
 
 		// If we did find a valid corresponding FSubmixVoiceData, remove it from the map.
 		if (MixerSourceVoiceData)
@@ -394,72 +405,6 @@ namespace Audio
 		});
 	}
 
-	void FMixerSubmix::FormatChangeBuffer(const ESubmixChannelFormat InNewChannelType, AlignedFloatBuffer& InBuffer, AlignedFloatBuffer& OutNewBuffer)
-	{
-		// Retrieve ptr to the cached downmix channel map from the mixer device
-		int32 NewChannelCount = MixerDevice->GetNumChannelsForSubmixFormat(InNewChannelType);
-		Audio::AlignedFloatBuffer ChannelMap;
-		MixerDevice->Get2DChannelMap(false, InNewChannelType, NumChannels, false, ChannelMap);
-		float* ChannelMapPtr = ChannelMap.GetData();
-
-		// Input and output frame count is going to be the same
-		const int32 NumFrames = InBuffer.Num() / NumChannels;
-
-		// Reset the passed in downmix scratch buffer
-		OutNewBuffer.Reset();
-		OutNewBuffer.AddZeroed(NumFrames * NewChannelCount);
-
-		if (SubmixAmbisonicsDecoderID == INDEX_NONE)
-		{
-			float* OutNewBufferPtr = OutNewBuffer.GetData();
-
-			// Loop through the down mix map and perform the downmix operation
-			int32 InputSampleIndex = 0;
-			int32 DownMixedBufferIndex = 0;
-			for (; InputSampleIndex < InBuffer.Num();)
-			{
-				for (int32 DownMixChannel = 0; DownMixChannel < NewChannelCount; ++DownMixChannel)
-				{
-					for (int32 InChannel = 0; InChannel < NumChannels; ++InChannel)
-					{
-						const int32 ChannelMapIndex = NewChannelCount * InChannel + DownMixChannel;
-						OutNewBufferPtr[DownMixedBufferIndex + DownMixChannel] += InBuffer[InputSampleIndex + InChannel] * ChannelMapPtr[ChannelMapIndex];
-					}
-				}
-
-				InputSampleIndex += NumChannels;
-				DownMixedBufferIndex += NewChannelCount;
-			}
-		}
-		else
-		{
-			FAmbisonicsDecoderInputData InputData;
-			InputData.AudioBuffer = &InBuffer;
-			InputData.NumChannels = NumChannels;
-
-			FAmbisonicsDecoderOutputData OutputData = { OutNewBuffer };
-
-			if (CachedPositionalData.OutputNumChannels != NewChannelCount)
-			{
-				// re-cache output positions
-				CachedPositionalData.OutputNumChannels = NewChannelCount;
-
-				CachedPositionalData.OutputChannelPositions = AmbisonicsStatics::GetDefaultPositionMap(NewChannelCount);
-			}
-
-			//TODO: Update listener rotation in CachedPositionalData
-			const TArray<FTransform>* ListenerTransforms = MixerDevice->GetListenerTransforms();
-
-			if (ListenerTransforms != nullptr && ListenerTransforms->Num() >= 1)
-			{
-				CachedPositionalData.ListenerRotation = (*ListenerTransforms)[0].GetRotation();
-			}
-
-			// Todo: sum into OutputData rather than decoding directly to it.
-			AmbisonicsMixer->DecodeFromAmbisonics(SubmixAmbisonicsDecoderID, InputData, CachedPositionalData, OutputData);
-		}
-	}
-
 	void FMixerSubmix::MixBufferDownToMono(const AlignedFloatBuffer& InBuffer, int32 NumInputChannels, AlignedFloatBuffer& OutBuffer)
 	{
 		check(NumInputChannels > 0);
@@ -483,162 +428,260 @@ namespace Audio
 		}
 	}
 
-	void FMixerSubmix::SetUpAmbisonicsEncoder()
+	void FMixerSubmix::SetupSoundfieldEncodersForChildren()
 	{
-		check(AmbisonicsMixer.IsValid());
+		check(SoundfieldStreams.Factory);
+		check(SoundfieldStreams.Settings.IsValid());
 
-		// If we haven't already set up the encoder, destroy the old stream.
-		if (SubmixAmbisonicsEncoderID == INDEX_NONE)
+		//Here we scan all child submixes to see which submixes need to be reencoded.
+		for (auto& Iter : ChildSubmixes)
 		{
-			TearDownAmbisonicsEncoder();
+			FChildSubmixInfo& ChildSubmix = Iter.Value;
+			SetupSoundfieldEncodingForChild(ChildSubmix);
 		}
 
-		//Get a new unique stream ID
-		SubmixAmbisonicsEncoderID = MixerDevice->GetNewUniqueAmbisonicsStreamID();
-		AmbisonicsMixer->OnOpenEncodingStream(SubmixAmbisonicsEncoderID, AmbisonicsSettings);
+		if ((ChildSubmixes.Num() > 0) && !SoundfieldStreams.Factory->ShouldEncodeAllStreamsIndependently(*SoundfieldStreams.Settings))
+		{
+			FAudioPluginInitializationParams InitParams = GetInitializationParamsForSoundfieldStream();
+			SoundfieldStreams.DownmixedChildrenEncoder = SoundfieldStreams.Factory->CreateEncoderStream(InitParams, *SoundfieldStreams.Settings);
+		}
 	}
 
-	void FMixerSubmix::SetUpAmbisonicsDecoder()
+	void FMixerSubmix::SetupSoundfieldEncodingForChild(FChildSubmixInfo& InChild)
 	{
-		check(AmbisonicsMixer.IsValid());
+		TSharedPtr<Audio::FMixerSubmix, ESPMode::ThreadSafe> SubmixPtr = InChild.SubmixPtr.Pin();
 
-		// if we have already set up the decoder, destroy the old stream.
-		if (SubmixAmbisonicsDecoderID == INDEX_NONE)
+		if (SubmixPtr.IsValid())
 		{
-			TearDownAmbisonicsDecoder();
-		}
+			check(SoundfieldStreams.Factory && SoundfieldStreams.Settings.IsValid());
 
-		SubmixAmbisonicsDecoderID = MixerDevice->GetNewUniqueAmbisonicsStreamID();
+			// If this child submix is not a soundfield submix and we need to encode every child submix independently, set up an encoder.
+			if (!SubmixPtr->IsSoundfieldSubmix() && SoundfieldStreams.Factory->ShouldEncodeAllStreamsIndependently(*SoundfieldStreams.Settings))
+			{
+				FAudioPluginInitializationParams InitParams = GetInitializationParamsForSoundfieldStream();
+				InChild.Encoder = SoundfieldStreams.Factory->CreateEncoderStream(InitParams, *SoundfieldStreams.Settings);
+			}
+			else if(SubmixPtr->IsSoundfieldSubmix())
+			{
+				// If the child submix is of a soundfield format that needs to be transcoded, set up a transcoder.
+				InChild.Transcoder = GetTranscoderForChildSubmix(SubmixPtr);
+			}
 
-		SetUpAmbisonicsPositionalData();
-		AmbisonicsMixer->OnOpenDecodingStream(SubmixAmbisonicsDecoderID, AmbisonicsSettings, CachedPositionalData);
+			// If neither of these are true, either we are downmixing all child audio and encoding it once, or
+			// this submix can handle the child's soundfield audio packet directly, so no encoder nor transcoder is needed.
+		}
 	}
 
-	void FMixerSubmix::TearDownAmbisonicsEncoder()
+	void FMixerSubmix::SetupSoundfieldStreamForParent()
 	{
-		if (SubmixAmbisonicsEncoderID != INDEX_NONE)
+		TSharedPtr<FMixerSubmix, ESPMode::ThreadSafe> ParentSubmixSharedPtr = ParentSubmix.Pin();
+
+		if (ParentSubmixSharedPtr.IsValid() && !ParentSubmixSharedPtr->IsSoundfieldSubmix())
 		{
-			AmbisonicsMixer->OnCloseEncodingStream(SubmixAmbisonicsEncoderID);
-			SubmixAmbisonicsEncoderID = INDEX_NONE;
+			// If the submix we're plugged into isn't a soundfield submix, we need to decode our soundfield for it.
+			SetUpSoundfieldPositionalData(ParentSubmixSharedPtr);
+
+			FAudioPluginInitializationParams InitParams = GetInitializationParamsForSoundfieldStream();
+			SoundfieldStreams.ParentDecoder = SoundfieldStreams.Factory->CreateDecoderStream(InitParams, *SoundfieldStreams.Settings);
 		}
 	}
 
-	void FMixerSubmix::TearDownAmbisonicsDecoder()
+	void FMixerSubmix::SetUpSoundfieldPositionalData(const TSharedPtr<Audio::FMixerSubmix, ESPMode::ThreadSafe>& InParentSubmix)
 	{
-		if (SubmixAmbisonicsDecoderID != INDEX_NONE)
-		{
-			AmbisonicsMixer->OnCloseDecodingStream(SubmixAmbisonicsDecoderID);
-			SubmixAmbisonicsDecoderID = INDEX_NONE;
-		}
+		// If there is a parent and we are not passing it this submix's ambisonics audio, retrieve that submix's channel format.
+		check(InParentSubmix.IsValid());
+
+		const int32 NumParentChannels = InParentSubmix->GetSubmixChannels();
+		SoundfieldStreams.CachedPositionalData.NumChannels = NumParentChannels;
+		SoundfieldStreams.CachedPositionalData.ChannelPositions = MixerDevice->GetDefaultPositionMap(NumParentChannels);
+		
+		// For now we don't actually do any sort of rotation for decoded audio.
+		SoundfieldStreams.CachedPositionalData.Rotation = FQuat::Identity;
 	}
 
-	void FMixerSubmix::UpdateAmbisonicsEncoderForChildren()
+	void FMixerSubmix::MixInSource(const ISoundfieldAudioPacket& InAudio, const ISoundfieldEncodingSettingsProxy& InSettings, ISoundfieldAudioPacket& PacketToSumTo)
 	{
-		bool bNeedsEncoder = false;
+		check(SoundfieldStreams.Mixer.IsValid());
 
-		//Here we scan all child submixes to see which submixes need to be reencoded.
-		for (auto& Iter : ChildSubmixes)
+		FSoundfieldMixerInputData InputData =
 		{
-			FChildSubmixInfo& ChildSubmix = Iter.Value;
+			InAudio, // InputPacket
+			InSettings, // EncodingSettings
+			1.0f // SendLevel
+		};
 
-			TSharedPtr<Audio::FMixerSubmix, ESPMode::ThreadSafe> SubmixPtr = ChildSubmix.SubmixPtr.Pin();
+		SoundfieldStreams.Mixer->MixTogether(InputData, PacketToSumTo);
+	}
 
-			// Check to see if this child is an ambisonics submix.
-			if (SubmixPtr.IsValid() && SubmixPtr->GetSubmixChannels() == ESubmixChannelFormat::Ambisonics)
+	void FMixerSubmix::MixInChildSubmix(FChildSubmixInfo& Child, ISoundfieldAudioPacket& PacketToSumTo)
+	{
+		check(IsSoundfieldSubmix());
+
+		// We only either encode, transcode input, and never both. If we have both for this child, something went wrong in initialization.
+		check(!(Child.Encoder.IsValid() && Child.Transcoder.IsValid()));
+
+		TSharedPtr<FMixerSubmix, ESPMode::ThreadSafe> ChildSubmixSharedPtr = Child.SubmixPtr.Pin();
+		if (ChildSubmixSharedPtr.IsValid())
+		{
+			if (!ChildSubmixSharedPtr->IsSoundfieldSubmix())
 			{
-				UAmbisonicsSubmixSettingsBase* ChildAmbisonicsSettings = SubmixPtr->AmbisonicsSettings;
+				// Reset the output scratch buffer so that we can call ProcessAudio on the ChildSubmix with it:
+				ScratchBuffer.Reset(NumSamples);
+				ScratchBuffer.AddZeroed(NumSamples);
+
+				// If this is true, the Soundfield Factory explicitly requested that a seperate encoder stream was set up for every
+				// non-soundfield child submix.
+				if (Child.Encoder.IsValid())
+				{
+					
 
-				// Check if this child submix needs to be re-encoded.
-				if (!ChildAmbisonicsSettings || AmbisonicsMixer->ShouldReencodeBetween(ChildAmbisonicsSettings, AmbisonicsSettings))
+					ChildSubmixSharedPtr->ProcessAudio(ScratchBuffer);
+
+					// Encode the resulting audio and mix it in.
+					FSoundfieldEncoderInputData InputData = {
+						ScratchBuffer, /* AudioBuffer */
+						ChildSubmixSharedPtr->NumChannels, /* NumChannels */
+						*SoundfieldStreams.Settings, /** InputSettings */
+						SoundfieldStreams.CachedPositionalData /** PosititonalData */
+					};
+
+					Child.Encoder->EncodeAndMixIn(InputData, PacketToSumTo);
+				}
+				else
+				{
+					// Otherwise, process and mix in the submix's audio to the scratch buffer, and we will encode ScratchBuffer later.
+					ChildSubmixSharedPtr->ProcessAudio(ScratchBuffer);
+				}
+			}
+			else if (Child.Transcoder.IsValid())
+			{
+				// Make sure our packet that we call process on is zeroed out:
+				if (!Child.IncomingPacketToTranscode.IsValid())
 				{
-					ChildSubmix.bNeedsAmbisonicsEncoding = false;
+					Child.IncomingPacketToTranscode = ChildSubmixSharedPtr->SoundfieldStreams.Factory->CreateEmptyPacket();
 				}
 				else
 				{
-					bNeedsEncoder = true;
+					Child.IncomingPacketToTranscode->Reset();
 				}
+
+				check(Child.IncomingPacketToTranscode.IsValid());
+
+				ChildSubmixSharedPtr->ProcessAudio(*Child.IncomingPacketToTranscode);
+
+				Child.Transcoder->TranscodeAndMixIn(*Child.IncomingPacketToTranscode, ChildSubmixSharedPtr->GetSoundfieldSettings(), PacketToSumTo, *SoundfieldStreams.Settings);
 			}
 			else
 			{
-				bNeedsEncoder = true;
+				// No conversion necessary.
+				ChildSubmixSharedPtr->ProcessAudio(PacketToSumTo);
 			}
 		}
+	}
 
-		if (bNeedsEncoder)
-		{
-			SetUpAmbisonicsEncoder();
-		}
-		else
-		{
-			TearDownAmbisonicsEncoder();
-		}
+	bool FMixerSubmix::IsSoundfieldSubmix() const
+	{
+		return SoundfieldStreams.Factory != nullptr;
+	}
+
+	bool FMixerSubmix::IsDefaultEndpointSubmix() const
+	{
+		return !ParentSubmix.IsValid() && !(EndpointData.SoundfieldEndpoint.IsValid() || EndpointData.NonSoundfieldEndpoint.IsValid());
 	}
 
-	void FMixerSubmix::UpdateAmbisonicsDecoderForParent()
+	bool FMixerSubmix::IsExternalEndpointSubmix() const
 	{
-		UAmbisonicsSubmixSettingsBase* ParentAmbisonicsSettings = nullptr;
+		return !ParentSubmix.IsValid() && (EndpointData.SoundfieldEndpoint.IsValid() || EndpointData.NonSoundfieldEndpoint.IsValid());
+	}
 
-		TSharedPtr<FMixerSubmix, ESPMode::ThreadSafe> ParentSubmixSharedPtr = ParentSubmix.Pin();
-		if (ParentSubmixSharedPtr.IsValid())
-		{
-			if (ParentSubmixSharedPtr->GetSubmixChannels() == ESubmixChannelFormat::Ambisonics)
-			{
-				ParentAmbisonicsSettings = ParentSubmixSharedPtr->AmbisonicsSettings;
-			}
-		}
+	bool FMixerSubmix::IsSoundfieldEndpointSubmix() const
+	{
+		return !ParentSubmix.IsValid() && IsSoundfieldSubmix();
+	}
 
-		// If we need to reencode between here and the parent submix, set up the submix decoder.
-		if (!ParentAmbisonicsSettings || AmbisonicsMixer->ShouldReencodeBetween(AmbisonicsSettings, ParentAmbisonicsSettings))
+	FName FMixerSubmix::GetSoundfieldFormat() const
+	{
+		if (IsSoundfieldSubmix())
 		{
-			SetUpAmbisonicsDecoder();
+			return SoundfieldStreams.Factory->GetSoundfieldFormatName();
 		}
 		else
 		{
-			TearDownAmbisonicsDecoder();
+			return ISoundfieldFactory::GetFormatNameForNoEncoding();
 		}
 	}
 
-	void FMixerSubmix::SetUpAmbisonicsPositionalData()
+	ISoundfieldEncodingSettingsProxy& FMixerSubmix::GetSoundfieldSettings()
 	{
-		// If there is a parent and we are not passing it this submix's ambisonics audio, retrieve that submix's channel format.
-		TSharedPtr<Audio::FMixerSubmix, ESPMode::ThreadSafe> ParentSubmixSharedPtr = ParentSubmix.Pin();
-		if (ParentSubmixSharedPtr.IsValid())
-		{
-			const ESubmixChannelFormat ParentSubmixFormat = ParentSubmixSharedPtr->GetSubmixChannels();
+		check(IsSoundfieldSubmix());
+		check(SoundfieldStreams.Settings.IsValid());
 
-			const int32 NumParentChannels = MixerDevice->GetNumChannelsForSubmixFormat(ParentSubmixFormat);
-			CachedPositionalData.OutputNumChannels = NumParentChannels;
-			CachedPositionalData.OutputChannelPositions = AmbisonicsStatics::GetDefaultPositionMap(NumParentChannels);
-		}
-		
-		CachedPositionalData.ListenerRotation = FQuat::Identity;
+		return *SoundfieldStreams.Settings;
 	}
 
-	void FMixerSubmix::EncodeAndMixInSource(AlignedFloatBuffer& InAudioData, FSubmixVoiceData& InVoiceInfo)
+	FAudioPluginInitializationParams FMixerSubmix::GetInitializationParamsForSoundfieldStream()
 	{
-		InVoiceInfo.CachedEncoderInputData.AudioBuffer = &InAudioData;
+		FAudioPluginInitializationParams InitializationParams;
+		InitializationParams.AudioDevicePtr = MixerDevice;
+		InitializationParams.BufferLength = MixerDevice ? MixerDevice->GetNumOutputFrames() : 0;
+		InitializationParams.NumOutputChannels = MixerDevice ? MixerDevice->GetNumDeviceChannels() : 0;
+		InitializationParams.SampleRate = MixerDevice ? MixerDevice->SampleRate : 0.0f;
 
-		FAmbisonicsEncoderOutputData OutputData = { InputAmbisonicsBuffer };
-		
-		// Encode voice to ambisonics:
-		check(AmbisonicsMixer.IsValid());
-		AmbisonicsMixer->EncodeToAmbisonics(InVoiceInfo.AmbisonicsEncoderId, InVoiceInfo.CachedEncoderInputData, OutputData, AmbisonicsSettings);
-		
-		//Sum output to ambisonics bed:
-		float* DestinationBuffer = InputBuffer.GetData();
-		float* SrcBuffer = InputAmbisonicsBuffer.GetData();
-		for (int32 Index = 0; Index < InputBuffer.Num(); Index++)
-		{
-			DestinationBuffer[Index] += SrcBuffer[Index];
-		}
+		// We only use one soundfield stream per source.
+		InitializationParams.NumSources = 1;
+
+		return InitializationParams;
 	}
 
-	void FMixerSubmix::EncodeAndMixInChildSubmix(FChildSubmixInfo& Child)
+	FSoundfieldSpeakerPositionalData FMixerSubmix::GetDefaultPositionalDataForAudioDevice()
 	{
-		check(AmbisonicsMixer.IsValid());
-		check(SubmixAmbisonicsEncoderID != INDEX_NONE);
+		FSoundfieldSpeakerPositionalData PositionalData;
+		PositionalData.NumChannels = MixerDevice->GetNumDeviceChannels();
+		PositionalData.ChannelPositions = MixerDevice->GetDefaultPositionMap(PositionalData.NumChannels);
+
+		// For now we don't actually do any sort of rotation for decoded audio.
+		PositionalData.Rotation = FQuat::Identity;
 
-		//TODO: Implement generic mixdowns to ambisonics. Set up input encoder channel
+		return PositionalData;
+	}
+
+	TUniquePtr<ISoundfieldTranscodeStream> FMixerSubmix::GetTranscoderForChildSubmix(const TSharedPtr<Audio::FMixerSubmix, ESPMode::ThreadSafe>& InChildSubmix)
+	{
+		check(InChildSubmix.IsValid());
+		check(IsSoundfieldSubmix() && InChildSubmix->IsSoundfieldSubmix());
+		check(SoundfieldStreams.Settings.IsValid() && InChildSubmix->SoundfieldStreams.Settings.IsValid());
+
+		if (GetSoundfieldFormat() != InChildSubmix->GetSoundfieldFormat())
+		{
+			ISoundfieldFactory* ChildFactory = InChildSubmix->GetSoundfieldFactory();
+
+			if (SoundfieldStreams.Factory->CanTranscodeFromSoundfieldFormat(InChildSubmix->GetSoundfieldFormat(), *InChildSubmix->SoundfieldStreams.Settings))
+			{
+				FAudioPluginInitializationParams InitParams = GetInitializationParamsForSoundfieldStream();
+				return SoundfieldStreams.Factory->CreateTranscoderStream(InChildSubmix->GetSoundfieldFormat(), InChildSubmix->GetSoundfieldSettings(), SoundfieldStreams.Factory->GetSoundfieldFormatName(), GetSoundfieldSettings(), InitParams);
+			}
+			else if (ChildFactory->CanTranscodeToSoundfieldFormat(GetSoundfieldFormat(), GetSoundfieldSettings()))
+			{
+				FAudioPluginInitializationParams InitParams = GetInitializationParamsForSoundfieldStream();
+				return ChildFactory->CreateTranscoderStream(InChildSubmix->GetSoundfieldFormat(), InChildSubmix->GetSoundfieldSettings(), SoundfieldStreams.Factory->GetSoundfieldFormatName(), GetSoundfieldSettings(), InitParams);
+			}
+			else
+			{
+				return nullptr;
+			}
+		}
+		else
+		{
+			if (SoundfieldStreams.Factory->IsTranscodeRequiredBetweenSettings(*InChildSubmix->SoundfieldStreams.Settings, *SoundfieldStreams.Settings))
+			{
+				FAudioPluginInitializationParams InitParams = GetInitializationParamsForSoundfieldStream();
+				return SoundfieldStreams.Factory->CreateTranscoderStream(InChildSubmix->GetSoundfieldFormat(), InChildSubmix->GetSoundfieldSettings(), SoundfieldStreams.Factory->GetSoundfieldFormatName(), GetSoundfieldSettings(), InitParams);
+			}
+			else
+			{
+				return nullptr;
+			}
+		}
 	}
 
 	void FMixerSubmix::PumpCommandQueue()
@@ -655,21 +698,59 @@ namespace Audio
 		CommandQueue.Enqueue(MoveTemp(Command));
 	}
 
-	void FMixerSubmix::ProcessAudio(const ESubmixChannelFormat ParentChannelType, AlignedFloatBuffer& OutAudioBuffer)
+	void FMixerSubmix::ProcessAudio(AlignedFloatBuffer& OutAudioBuffer)
 	{
 		AUDIO_MIXER_CHECK_AUDIO_PLAT_THREAD(MixerDevice);
 
-		// Pump pending command queues
-		PumpCommandQueue();
+		// If this is a Soundfield Submix, process our soundfield and decode it to a OutAudioBuffer.
+		if (IsSoundfieldSubmix())
+		{
+			FScopeLock ScopeLock(&SoundfieldStreams.StreamsLock);
 
-		// Device format may change channels if device is hot swapped
-		if (ChannelFormat == ESubmixChannelFormat::Device)
+			// Initialize or clear the mixed down audio packet.
+			if (!SoundfieldStreams.MixedDownAudio.IsValid())
+			{
+				SoundfieldStreams.MixedDownAudio = SoundfieldStreams.Factory->CreateEmptyPacket();
+			}
+			else
+			{
+				SoundfieldStreams.MixedDownAudio->Reset();
+			}
+
+			check(SoundfieldStreams.MixedDownAudio.IsValid());
+
+			ProcessAudio(*SoundfieldStreams.MixedDownAudio);
+
+			if (!SoundfieldStreams.ParentDecoder.IsValid())
+			{
+				return;
+			}
+			
+			//Decode soundfield to interleaved float audio.
+			FSoundfieldDecoderInputData DecoderInput =
+			{
+				*SoundfieldStreams.MixedDownAudio, /* SoundfieldBuffer */
+				SoundfieldStreams.CachedPositionalData, /* PositionalData */
+				MixerDevice ? MixerDevice->GetNumOutputFrames() : 0, /* NumFrames */
+				MixerDevice ? MixerDevice->GetSampleRate() : 0.0f /* SampleRate */
+			};
+
+			FSoundfieldDecoderOutputData DecoderOutput = { OutAudioBuffer };
+
+			SoundfieldStreams.ParentDecoder->DecodeAndMixIn(DecoderInput, DecoderOutput);
+			return;
+		}
+		else
 		{
-			NumChannels = MixerDevice->GetNumChannelsForSubmixFormat(ChannelFormat);
-			const int32 NumOutputFrames = MixerDevice->GetNumOutputFrames();
-			NumSamples = NumChannels * NumOutputFrames;
+			// Pump pending command queues. For Soundfield Submixes this occurs in ProcessAudio(ISoundfieldAudioPacket&).
+			PumpCommandQueue();
 		}
 
+		// Device format may change channels if device is hot swapped
+		NumChannels = MixerDevice->GetNumDeviceChannels();
+		const int32 NumOutputFrames = OutAudioBuffer.Num() / NumChannels;
+		NumSamples = NumChannels * NumOutputFrames;
+
 		InputBuffer.Reset(NumSamples);
 		InputBuffer.AddZeroed(NumSamples);
 
@@ -686,22 +767,7 @@ namespace Audio
 				TSharedPtr<Audio::FMixerSubmix, ESPMode::ThreadSafe> ChildSubmix = ChildSubmixEntry.Value.SubmixPtr.Pin();
 				if (ChildSubmix.IsValid())
 				{
-					ScratchBuffer.Reset(NumSamples);
-					ScratchBuffer.AddZeroed(NumSamples);
-
-					ChildSubmix->ProcessAudio(ChannelFormat, ScratchBuffer);
-
-					float* ScratchBufferPtr = ScratchBuffer.GetData();
-
-					if (ChildSubmixEntry.Value.bNeedsAmbisonicsEncoding)
-					{
-						// Encode into ambisonics. TODO: Implement.
-						EncodeAndMixInChildSubmix(ChildSubmixEntry.Value);
-					}
-					else
-					{
-						Audio::MixInBufferFast(ScratchBufferPtr, BufferPtr, NumSamples);
-					}
+					ChildSubmix->ProcessAudio(InputBuffer);
 				}
 				else
 				{
@@ -719,13 +785,13 @@ namespace Audio
 			CSV_SCOPED_TIMING_STAT(Audio, SubmixSource);
 
 			// Loop through this submix's sound sources
-			for (const auto MixerSourceVoiceIter : MixerSourceVoices)
+			for (const auto& MixerSourceVoiceIter : MixerSourceVoices)
 			{
 				const FMixerSourceVoice* MixerSourceVoice = MixerSourceVoiceIter.Key;
 				const float SendLevel = MixerSourceVoiceIter.Value.SendLevel;
 
 				
-				MixerSourceVoice->MixOutputBuffers(ChannelFormat, SendLevel, InputBuffer);
+				MixerSourceVoice->MixOutputBuffers(NumChannels, SendLevel, InputBuffer);
 			}
 		}
 
@@ -765,15 +831,17 @@ namespace Audio
 				// Check to see if we need to down-mix our audio before sending to the submix effect
 				const uint32 ChannelCountOverride = SubmixEffect->GetDesiredInputChannelCountOverride();
 
-				// Only support downmixing to stereo. TODO: change GetDesiredInputChannelCountOverride() API to be "DownmixToStereo"
-				if (ChannelCountOverride < (uint32)NumChannels && ChannelCountOverride == 2)
+				if (ChannelCountOverride != INDEX_NONE && ChannelCountOverride != NumChannels)
 				{
 					// Perform the down-mix operation with the down-mixed scratch buffer
-					FormatChangeBuffer(ESubmixChannelFormat::Stereo, InputBuffer, DownmixedBuffer);
+					DownmixedBuffer.Reset();
+					DownmixedBuffer.AddUninitialized(NumOutputFrames * ChannelCountOverride);
+					DownmixBuffer(NumChannels, InputBuffer, ChannelCountOverride, DownmixedBuffer);
 
 					InputData.NumChannels = ChannelCountOverride;
 					InputData.AudioBuffer = &DownmixedBuffer;
 					SubmixEffect->ProcessAudio(InputData, OutputData);
+					DownmixBuffer(ChannelCountOverride, DownmixedBuffer, NumChannels, InputBuffer);
 				}
 				else
 				{
@@ -820,14 +888,7 @@ namespace Audio
 		}
 
 		// If the channel types match, just do a copy
-		if (ChannelFormat != ParentChannelType || SubmixAmbisonicsDecoderID != INDEX_NONE)
-		{
-			FormatChangeBuffer(ParentChannelType, InputBuffer, OutAudioBuffer);
-		}
-		else
-		{
-			FMemory::Memcpy((void*)OutAudioBuffer.GetData(), (void*)InputBuffer.GetData(), sizeof(float)*NumSamples);
-		}
+		Audio::MixInBufferFast(InputBuffer, OutAudioBuffer);
 
 		// Perform any envelope following if we're told to do so
 		if (bIsEnvelopeFollowing)
@@ -884,6 +945,7 @@ namespace Audio
 		}
 
 		// Now loop through any buffer listeners and feed the listeners the result of this audio callback
+		if(const USoundSubmix* SoundSubmix = Cast<const USoundSubmix>(OwningSubmixObject))
 		{
 			double AudioClock = MixerDevice->GetAudioTime();
 			float SampleRate = MixerDevice->GetSampleRate();
@@ -891,7 +953,191 @@ namespace Audio
 			for (ISubmixBufferListener* BufferListener : BufferListeners)
 			{
 				check(BufferListener);
-				BufferListener->OnNewSubmixBuffer(OwningSubmixObject, OutAudioBuffer.GetData(), OutAudioBuffer.Num(), NumChannels, SampleRate, AudioClock);
+				BufferListener->OnNewSubmixBuffer(SoundSubmix, OutAudioBuffer.GetData(), OutAudioBuffer.Num(), NumChannels, SampleRate, AudioClock);
+			}
+		}
+	}
+
+	void FMixerSubmix::ProcessAudio(ISoundfieldAudioPacket& OutputAudio)
+	{
+		check(IsSoundfieldSubmix());
+		PumpCommandQueue();
+
+		// Mix all submix audio into OutputAudio.
+		{
+			CSV_SCOPED_TIMING_STAT(Audio, SubmixSoundfieldChildren);
+
+			// If we are mixing down all non-soundfield child submixes,
+			// Set up the scratch buffer so that we can sum all non-soundfield child submixes to it.
+			if (SoundfieldStreams.DownmixedChildrenEncoder.IsValid())
+			{
+				ScratchBuffer.Reset();
+				ScratchBuffer.AddZeroed(MixerDevice->GetNumOutputFrames() * MixerDevice->GetNumDeviceChannels());
+			}
+
+			// First loop this submix's child submixes that are soundfields mixing in their output into this submix's dry/wet buffers.
+			for (auto& ChildSubmixEntry : ChildSubmixes)
+			{
+				MixInChildSubmix(ChildSubmixEntry.Value, OutputAudio);
+			}
+
+			// If we mixed down all non-soundfield child submixes,
+			// We encode and mix in here.
+			if (ChildSubmixes.Num() && SoundfieldStreams.DownmixedChildrenEncoder.IsValid())
+			{
+				FSoundfieldSpeakerPositionalData PositionalData = GetDefaultPositionalDataForAudioDevice();
+
+				FSoundfieldEncoderInputData InputData = {
+						ScratchBuffer, /* AudioBuffer */
+						MixerDevice->GetNumDeviceChannels(), /* NumChannels */
+						*SoundfieldStreams.Settings, /** InputSettings */
+						PositionalData /** PosititonalData */
+				};
+
+				SoundfieldStreams.DownmixedChildrenEncoder->EncodeAndMixIn(InputData, OutputAudio);
+			}
+		}
+
+		// Mix all source sends into OutputAudio.
+		{
+			CSV_SCOPED_TIMING_STAT(Audio, SubmixSoundfieldSources);
+
+			check(SoundfieldStreams.Mixer.IsValid());
+
+			// Loop through this submix's sound sources
+			for (const auto MixerSourceVoiceIter : MixerSourceVoices)
+			{
+				const FMixerSourceVoice* MixerSourceVoice = MixerSourceVoiceIter.Key;
+				const float SendLevel = MixerSourceVoiceIter.Value.SendLevel;
+
+				// if this voice has a valid encoded packet, mix it in.
+				const ISoundfieldAudioPacket* Packet = MixerSourceVoice->GetEncodedOutput(GetKeyForSubmixEncoding());
+				if (Packet)
+				{
+					FSoundfieldMixerInputData InputData =
+					{
+						*Packet,
+						*SoundfieldStreams.Settings,
+						SendLevel
+					};
+
+					SoundfieldStreams.Mixer->MixTogether(InputData, OutputAudio);
+				}
+			}
+		}
+
+		// Run soundfield processors.
+		{
+			CSV_SCOPED_TIMING_STAT(Audio, SubmixSoundfieldProcessors);
+
+			for (auto& EffectData : SoundfieldStreams.EffectProcessors)
+			{
+				check(EffectData.Processor.IsValid());
+				check(EffectData.Settings.IsValid());
+
+				EffectData.Processor->ProcessAudio(OutputAudio, *SoundfieldStreams.Settings, *EffectData.Settings);
+			}
+		}
+	}
+
+	void FMixerSubmix::ProcessAudioAndSendToEndpoint()
+	{
+		check(MixerDevice);
+
+		if (IsSoundfieldSubmix())
+		{
+			if (!EndpointData.AudioPacket.IsValid())
+			{
+				EndpointData.AudioPacket = SoundfieldStreams.Factory->CreateEmptyPacket();
+			}
+			else
+			{
+				EndpointData.AudioPacket->Reset();
+			}
+
+			// First, process the audio chain for this submix.
+			check(EndpointData.AudioPacket);
+			ProcessAudio(*EndpointData.AudioPacket);
+
+			if (EndpointData.SoundfieldEndpoint->GetRemainderInPacketBuffer() > 0)
+			{
+				EndpointData.SoundfieldEndpoint->PushAudio(MoveTemp(EndpointData.AudioPacket));
+			}
+			else
+			{
+				ensureMsgf(false, TEXT("Buffer overrun in Soundfield endpoint! %s may need to override ISoundfieldEndpoint::EndpointRequiresCallback() to return true."), *SoundfieldStreams.Factory->GetSoundfieldFormatName().ToString());
+			}
+
+			EndpointData.SoundfieldEndpoint->ProcessAudioIfNecessary();
+		}
+		else
+		{
+			// First, process the chain for this submix.
+			EndpointData.AudioBuffer.Reset();
+			EndpointData.AudioBuffer.AddZeroed(MixerDevice->GetNumOutputFrames() * MixerDevice->GetNumDeviceChannels());
+			ProcessAudio(EndpointData.AudioBuffer);
+
+			if (!EndpointData.Input.IsOutputStillActive())
+			{
+				// Either this is our first time pushing audio or we were disconnected.
+				const float DurationPerCallback = MixerDevice->GetNumOutputFrames() / MixerDevice->GetSampleRate();
+
+				EndpointData.Input = EndpointData.NonSoundfieldEndpoint->PatchNewInput(DurationPerCallback, EndpointData.SampleRate, EndpointData.NumChannels);
+
+				if (!FMath::IsNearlyEqual(EndpointData.SampleRate, MixerDevice->GetSampleRate()))
+				{
+					// Initialize the resampler.
+					float SampleRateRatio = EndpointData.SampleRate / MixerDevice->GetSampleRate();
+
+					EndpointData.Resampler.Init(EResamplingMethod::Linear, SampleRateRatio, NumChannels);
+					EndpointData.bShouldResample = true;
+
+					// Add a little slack at the end of the resampled audio buffer in case we have roundoff jitter.
+					EndpointData.ResampledAudioBuffer.Reset();
+					EndpointData.ResampledAudioBuffer.AddUninitialized(EndpointData.AudioBuffer.Num() * SampleRateRatio + 16);
+				}
+			}
+
+			// Resample if necessary.
+			int32 NumResampledFrames = EndpointData.AudioBuffer.Num() / NumChannels;
+			if (EndpointData.bShouldResample)
+			{
+				EndpointData.Resampler.ProcessAudio(EndpointData.AudioBuffer.GetData(), EndpointData.AudioBuffer.Num(), false, EndpointData.ResampledAudioBuffer.GetData(), EndpointData.ResampledAudioBuffer.Num(), NumResampledFrames);
+			}
+			else
+			{
+				EndpointData.ResampledAudioBuffer = MoveTemp(EndpointData.AudioBuffer);
+			}
+
+			// Downmix if necessary.
+			const bool bShouldDownmix = EndpointData.NumChannels != NumChannels;
+			if (bShouldDownmix)
+			{
+				EndpointData.DownmixedResampledAudioBuffer.Reset();
+				EndpointData.DownmixedResampledAudioBuffer.AddUninitialized(NumResampledFrames * EndpointData.NumChannels);
+
+				EndpointData.DownmixChannelMap.Reset();
+				FMixerDevice::Get2DChannelMap(false, NumChannels, EndpointData.NumChannels, false, EndpointData.DownmixChannelMap);
+				DownmixBuffer(NumChannels, EndpointData.ResampledAudioBuffer, EndpointData.NumChannels, EndpointData.DownmixedResampledAudioBuffer);
+			}
+			else
+			{
+				EndpointData.DownmixedResampledAudioBuffer = MoveTemp(EndpointData.ResampledAudioBuffer);
+			}
+
+			EndpointData.Input.PushAudio(EndpointData.DownmixedResampledAudioBuffer.GetData(), EndpointData.DownmixedResampledAudioBuffer.Num());
+			EndpointData.NonSoundfieldEndpoint->ProcessAudioIfNeccessary();
+
+			// If we did any pointer passing to bypass downmixing or resampling, pass the pointer back to avoid reallocating ResampledAudioBuffer or AudioBuffer.
+
+			if (!bShouldDownmix)
+			{
+				EndpointData.ResampledAudioBuffer = MoveTemp(EndpointData.DownmixedResampledAudioBuffer);
+			}
+
+			if (!EndpointData.bShouldResample)
+			{
+				EndpointData.AudioBuffer = MoveTemp(EndpointData.ResampledAudioBuffer);
 			}
 		}
 	}
@@ -920,24 +1166,173 @@ namespace Audio
 		return nullptr;
 	}
 
-	void FMixerSubmix::OnAmbisonicsSettingsChanged(UAmbisonicsSubmixSettingsBase* InAmbisonicsSettings)
+	void FMixerSubmix::SetSoundfieldFactory(ISoundfieldFactory* InSoundfieldFactory)
+	{
+		SoundfieldStreams.Factory = InSoundfieldFactory;
+	}
+
+	void FMixerSubmix::SetupSoundfieldStreams(const USoundfieldEncodingSettingsBase* InAmbisonicsSettings, TArray<USoundfieldEffectBase*>& Processors, ISoundfieldFactory* InSoundfieldFactory)
 	{
+		FScopeLock ScopeLock(&SoundfieldStreams.StreamsLock);
+
+		// SoundfieldStreams.Factory should have already been set by our first pass around the submix graph, since 
+		// we use the soundfield factory
+		check(SoundfieldStreams.Factory == InSoundfieldFactory);
+
+		if (!InSoundfieldFactory)
+		{
+			return;
+		}
+
 		check(InAmbisonicsSettings != nullptr);
-		if (!AmbisonicsMixer.IsValid())
+		// As a santity check, we ensure the passed in soundfield stream is what was used for the initial SetSoundfieldFactory call.
+		check(InSoundfieldFactory == SoundfieldStreams.Factory);
+
+
+		SoundfieldStreams.Reset();
+		SoundfieldStreams.Factory = InSoundfieldFactory;
+
+		// If this submix is encoded to a soundfield, 
+		// Explicitly set NumChannels and NumSamples to 0 since they are technically irrelevant.
+		NumChannels = 0;
+		NumSamples = 0;
+
+		SoundfieldStreams.Settings = InAmbisonicsSettings->GetProxy();
+
+		// Check to see if this implementation of GetProxy failed.
+		if (!ensureAlwaysMsgf(SoundfieldStreams.Settings.IsValid(), TEXT("Soundfield Format %s failed to create a settings proxy for settings asset %s."), *InSoundfieldFactory->GetSoundfieldFormatName().ToString(), *InAmbisonicsSettings->GetName()))
+		{
+			
+			TeardownSoundfieldStreams();
+			return;
+		}
+		
+		SoundfieldStreams.Mixer = InSoundfieldFactory->CreateMixerStream(*SoundfieldStreams.Settings);
+
+		if (!ensureAlwaysMsgf(SoundfieldStreams.Mixer.IsValid(), TEXT("Soundfield Format %s failed to create a settings proxy for settings asset %s."), *InSoundfieldFactory->GetSoundfieldFormatName().ToString(), *InAmbisonicsSettings->GetName()))
 		{
-			AmbisonicsMixer = MixerDevice->GetAmbisonicsMixer();
-			if (!AmbisonicsMixer.IsValid())
+			TeardownSoundfieldStreams();
+			return;
+		}
+
+		// Create new processor proxies.
+		for (USoundfieldEffectBase* Processor : Processors)
+		{
+			if (Processor != nullptr)
 			{
-				return;
+				SoundfieldStreams.EffectProcessors.Emplace(SoundfieldStreams.Factory, *SoundfieldStreams.Settings, Processor);
 			}
 		}
 
-		AmbisonicsSettings = InAmbisonicsSettings;
-		NumChannels = AmbisonicsMixer->GetNumChannelsForAmbisonicsFormat(AmbisonicsSettings);
-		NumSamples = NumChannels * MixerDevice->GetNumOutputFrames();
+		SetupSoundfieldEncodersForChildren();
+		SetupSoundfieldStreamForParent();
+	}
+
+	void FMixerSubmix::TeardownSoundfieldStreams()
+	{
+		SoundfieldStreams.Reset();
 
-		UpdateAmbisonicsEncoderForChildren();
-		UpdateAmbisonicsDecoderForParent();
+		for (auto& ChildSubmix : ChildSubmixes)
+		{
+			ChildSubmix.Value.Encoder.Reset();
+			ChildSubmix.Value.Transcoder.Reset();
+		}
+	}
+
+	void FMixerSubmix::SetupEndpoint(IAudioEndpointFactory* InFactory, const UAudioEndpointSettingsBase* InSettings)
+	{
+		checkf(!IsSoundfieldSubmix(), TEXT("Soundfield Endpoint Submixes called with non-soundfield arguments."));
+		check(!ParentSubmix.IsValid());
+		EndpointData.Reset();
+
+		if (!InFactory)
+		{
+			return;
+		}
+
+		TUniquePtr<IAudioEndpointSettingsProxy> SettingsProxy;
+		if (InSettings)
+		{
+			SettingsProxy = InSettings->GetProxy();
+		}
+		else
+		{
+			InSettings = InFactory->GetDefaultSettings();
+			ensureMsgf(InSettings, TEXT("The audio endpoint factory %s failed to generate default settings!"), *InFactory->GetEndpointTypeName().ToString());
+
+			if (InSettings)
+			{
+				SettingsProxy = InSettings->GetProxy();
+			}
+		}
+
+		if (SettingsProxy)
+		{
+			FAudioPluginInitializationParams InitParams = GetInitializationParamsForSoundfieldStream();
+			EndpointData.NonSoundfieldEndpoint = InFactory->CreateNewEndpointInstance(InitParams, *SettingsProxy);
+		}
+		else
+		{
+			ensureMsgf(false, TEXT("Settings object %s failed to create a proxy object. Likely an error in the implementation of %s::GetProxy()."), *InSettings->GetName(), *InSettings->GetClass()->GetName());
+		}
+	}
+
+	void FMixerSubmix::SetupEndpoint(ISoundfieldEndpointFactory* InFactory, const USoundfieldEndpointSettingsBase* InSettings)
+	{
+		checkf(IsSoundfieldSubmix(), TEXT("Non-Soundfield Endpoint Submixes called with soundfield arguments."));
+		check(SoundfieldStreams.Factory == InFactory);
+		check(!ParentSubmix.IsValid());
+
+		EndpointData.Reset();
+
+		if (!InFactory)
+		{
+			return;
+		}
+
+		TUniquePtr<ISoundfieldEndpointSettingsProxy> SettingsProxy;
+		if (InSettings)
+		{
+			SettingsProxy = InSettings->GetProxy();
+		}
+		else
+		{
+			InSettings = InFactory->GetDefaultEndpointSettings();
+			ensureMsgf(InSettings, TEXT("The audio endpoint factory %s failed to generate default settings!"), *InFactory->GetEndpointTypeName().ToString());
+
+			if (InSettings)
+			{
+				SettingsProxy = InSettings->GetProxy();
+			}
+		}
+
+		if (SettingsProxy)
+		{
+			FAudioPluginInitializationParams InitParams = GetInitializationParamsForSoundfieldStream();
+			EndpointData.SoundfieldEndpoint = InFactory->CreateNewEndpointInstance(InitParams, *SettingsProxy);
+		}
+		else
+		{
+			ensureMsgf(false, TEXT("Settings object %s failed to create a proxy object. Likely an error in the implementation of %s::GetProxy()."), *InSettings->GetName(), *InSettings->GetClass()->GetName());
+		}
+	}
+
+	void FMixerSubmix::UpdateEndpointSettings(TUniquePtr<IAudioEndpointSettingsProxy>&& InSettings)
+	{
+		checkf(!IsSoundfieldSubmix(), TEXT("UpdateEndpointSettings for a soundfield submix was called with an IAudioEndpointSettingsProxy rather than an ISoundfieldEndpointSettingsProxy."));
+		if (ensureMsgf(EndpointData.NonSoundfieldEndpoint.IsValid(), TEXT("UpdateEndpointSettings called on an object that is not an endpoint submix.")))
+		{
+			EndpointData.NonSoundfieldEndpoint->SetNewSettings(MoveTemp(InSettings));
+		}
+	}
+
+	void FMixerSubmix::UpdateEndpointSettings(TUniquePtr<ISoundfieldEndpointSettingsProxy>&& InSettings)
+	{
+		checkf(IsSoundfieldSubmix(), TEXT("UpdateEndpointSettings for a non-soundfield submix was called with an ISoundfieldEndpointSettingsProxy rather than an IAudioEndpointSettingsProxy."));
+		if (ensureMsgf(EndpointData.SoundfieldEndpoint.IsValid(), TEXT("UpdateEndpointSettings called on an object that is not an endpoint submix.")))
+		{
+			EndpointData.SoundfieldEndpoint->SetNewSettings(MoveTemp(InSettings));
+		}
 	}
 
 	void FMixerSubmix::OnStartRecordingOutput(float ExpectedDuration)
@@ -1109,4 +1504,15 @@ namespace Audio
 		}
 	}
 
+	FSoundfieldEncodingKey FMixerSubmix::GetKeyForSubmixEncoding()
+	{
+		check(IsSoundfieldSubmix() && SoundfieldStreams.Settings.IsValid());
+		return FSoundfieldEncodingKey(SoundfieldStreams.Factory, *SoundfieldStreams.Settings);
+	}
+
+	ISoundfieldFactory* FMixerSubmix::GetSoundfieldFactory()
+	{
+		return SoundfieldStreams.Factory;
+	}
+
 }
diff --git a/Engine/Source/Runtime/AudioMixer/Private/Components/SynthComponent.cpp b/Engine/Source/Runtime/AudioMixer/Private/Components/SynthComponent.cpp
--- a/Engine/Source/Runtime/AudioMixer/Private/Components/SynthComponent.cpp
+++ b/Engine/Source/Runtime/AudioMixer/Private/Components/SynthComponent.cpp
@@ -515,7 +515,7 @@ void USynthComponent::SetVolumeMultiplier(float VolumeMultiplier)
 	}
 }
 
-void USynthComponent::SetSubmixSend(USoundSubmix* Submix, float SendLevel)
+void USynthComponent::SetSubmixSend(USoundSubmixBase* Submix, float SendLevel)
 {
 	if (AudioComponent)
 	{
diff --git a/Engine/Source/Runtime/AudioMixer/Public/AudioMixerDevice.h b/Engine/Source/Runtime/AudioMixer/Public/AudioMixerDevice.h
--- a/Engine/Source/Runtime/AudioMixer/Public/AudioMixerDevice.h
+++ b/Engine/Source/Runtime/AudioMixer/Public/AudioMixerDevice.h
@@ -23,22 +23,6 @@ namespace Audio
 	typedef TSharedPtr<FMixerSubmix, ESPMode::ThreadSafe> FMixerSubmixPtr;
 	typedef TWeakPtr<FMixerSubmix, ESPMode::ThreadSafe> FMixerSubmixWeakPtr;
 
-	struct FChannelPositionInfo
-	{
-		EAudioMixerChannel::Type Channel;
-		int32 Azimuth;
-
-		FChannelPositionInfo()
-			: Channel(EAudioMixerChannel::Unknown)
-			, Azimuth(0)
-		{}
-
-		FChannelPositionInfo(EAudioMixerChannel::Type InChannel, int32 InAzimuth)
-			: Channel(InChannel)
-			, Azimuth(InAzimuth)
-		{}
-	};
-
 	/** Data used to schedule events automatically in the audio renderer in audio mixer. */
 	struct FAudioThreadTimingData
 	{
@@ -70,7 +54,6 @@ namespace Audio
 			Master,
 			Reverb,
 			EQ,
-			Ambisonics,
 			Count,
 		};
 	}
@@ -108,8 +91,8 @@ namespace Audio
 		virtual void SuspendContext() override;
 		virtual void EnableDebugAudioOutput() override;
 		virtual FAudioPlatformSettings GetPlatformSettings() const override;
-		virtual void RegisterSoundSubmix(USoundSubmix* SoundSubmix, bool bInit = true) override;
-		virtual void UnregisterSoundSubmix(USoundSubmix* SoundSubmix) override;
+		virtual void RegisterSoundSubmix(const USoundSubmixBase* SoundSubmix, bool bInit = true) override;
+		virtual void UnregisterSoundSubmix(const USoundSubmixBase* SoundSubmix) override;
 
 		virtual void InitSoundEffectPresets() override;
 		virtual int32 GetNumActiveSources() const override;
@@ -122,7 +105,7 @@ namespace Audio
 		virtual bool GetCurrentSourceEffectChain(const uint32 SourceEffectChainId, TArray<FSourceEffectChainEntry>& OutCurrentSourceEffectChainEntries) override;
 
 		// Updates submix instances with new properties
-		virtual void UpdateSubmixProperties(USoundSubmix* InSubmix) override;
+		virtual void UpdateSubmixProperties(USoundSubmixBase* InSubmix) override;
 		
 		// Sets the submix output volume dynamically
 		virtual void SetSubmixOutputVolume(USoundSubmix* InSubmix, float NewVolume) override;
@@ -157,7 +140,13 @@ namespace Audio
 		virtual void OnAudioStreamShutdown() override;
 		//~ End IAudioMixer
 
-		FMixerSubmixWeakPtr GetSubmixInstance(USoundSubmix* InSoundSubmix);
+		FMixerSubmixWeakPtr GetSubmixInstance(const USoundSubmixBase* SoundSubmix);
+
+		// If SoundSubmix is a soundfield submix, this will return the factory used to encode 
+		// source audio to it's soundfield format.
+		// Otherwise, returns nullptr.
+		ISoundfieldFactory* GetFactoryForSubmixInstance(USoundSubmix* SoundSubmix);
+		ISoundfieldFactory* GetFactoryForSubmixInstance(FMixerSubmixWeakPtr& SoundSubmixPtr);
 
 		// Functions which check the thread it's called on and helps make sure functions are called from correct threads
 		void CheckAudioThread() const;
@@ -179,10 +168,10 @@ namespace Audio
 		IAudioMixerPlatformInterface* GetAudioMixerPlatform() const { return AudioMixerPlatform; }
 
 		// Builds a 3D channel map for a spatialized source.
-		void Get3DChannelMap(const ESubmixChannelFormat InSubmixChannelType, const FWaveInstance* InWaveInstance, const float EmitterAzimuth, const float NormalizedOmniRadius, Audio::AlignedFloatBuffer& OutChannelMap);
+		void Get3DChannelMap(const int32 InSubmixNumChannels, const FWaveInstance* InWaveInstance, const float EmitterAzimuth, const float NormalizedOmniRadius, Audio::AlignedFloatBuffer& OutChannelMap);
 
 		// Builds a channel gain matrix for a non-spatialized source. The non-static variation of this function queries AudioMixerDevice->NumOutputChannels directly which may not be thread safe.
-		void Get2DChannelMap(bool bIsVorbis, const ESubmixChannelFormat InSubmixChannelType, const int32 NumSourceChannels, const bool bIsCenterChannelOnly, Audio::AlignedFloatBuffer& OutChannelMap) const;
+		void Get2DChannelMap(bool bIsVorbis, const int32 NumSourceChannels, const bool bIsCenterChannelOnly, Audio::AlignedFloatBuffer& OutChannelMap) const;
 		static void Get2DChannelMap(bool bIsVorbis, const int32 NumSourceChannels, const int32 NumOutputChannels, const bool bIsCenterChannelOnly, Audio::AlignedFloatBuffer& OutChannelMap);
 
 		int32 GetDeviceSampleRate() const;
@@ -193,7 +182,6 @@ namespace Audio
 		FMixerSubmixWeakPtr GetMasterSubmix(); 
 		FMixerSubmixWeakPtr GetMasterReverbSubmix();
 		FMixerSubmixWeakPtr GetMasterEQSubmix();
-		FMixerSubmixWeakPtr GetMasterAmbisonicsSubmix();
 
 		// Add submix effect to master submix
 		void AddMasterSubmixEffect(uint32 SubmixEffectId, FSoundEffectSubmixPtr SoundEffect);
@@ -204,23 +192,24 @@ namespace Audio
 		// Clear all submix effects from master submix
 		void ClearMasterSubmixEffects();
 
-		// Returns the number of channels for a given submix channel type
-		int32 GetNumChannelsForSubmixFormat(const ESubmixChannelFormat InSubmixChannelType) const;
-		ESubmixChannelFormat GetSubmixChannelFormatForNumChannels(const int32 InNumChannels) const;
-
-		uint32 GetNewUniqueAmbisonicsStreamID();
-
 		// Returns the channel array for the given submix channel type
-		const TArray<EAudioMixerChannel::Type>& GetChannelArrayForSubmixChannelFormat(const ESubmixChannelFormat InSubmixChannelFormat) const;
+		const TArray<EAudioMixerChannel::Type>& GetChannelArray() const;
 
 		// Retrieves the listener transforms
 		const TArray<FTransform>* GetListenerTransforms();
 
+		// Retrieves spherical locations of channels for a given submix format
+		const FChannelPositionInfo* GetDefaultChannelPositions() const;
+
 		// Audio thread tick timing relative to audio render thread timing
 		double GetAudioThreadTime() const { return AudioThreadTimingData.AudioThreadTime; }
 		double GetAudioRenderThreadTime() const { return AudioThreadTimingData.AudioRenderThreadTime; }
 		double GetAudioClockDelta() const { return AudioClockDelta; }
 
+		TArray<Audio::FChannelPositionInfo>* GetDefaultPositionMap(int32 NumChannels);
+
+		static bool IsEndpointSubmix(const USoundSubmixBase* InSubmix);
+
 	protected:
 
 		virtual void InitSoundSubmixes() override;
@@ -233,7 +222,7 @@ namespace Audio
 		// Resets the thread ID used for audio rendering
 		void ResetAudioRenderingThreadId();
 
-		void RebuildSubmixLinks(USoundSubmix& SoundSubmix, FMixerSubmixPtr& SubmixInstance);
+		void RebuildSubmixLinks(const USoundSubmixBase& SoundSubmix, FMixerSubmixPtr& SubmixInstance);
 
 		void Get2DChannelMapInternal(const int32 NumSourceChannels, const int32 NumOutputChannels, const bool bIsCenterChannelOnly, TArray<float>& OutChannelMap) const;
 		void InitializeChannelMaps();
@@ -248,13 +237,16 @@ namespace Audio
 
 		void LoadMasterSoundSubmix(EMasterSubmixType::Type InType, const FString& InDefaultName, bool bInDefaultMuteWhenBackgrounded, FSoftObjectPath& InOutObjectPath);
 		void LoadPluginSoundSubmixes();
-		void LoadSoundSubmix(USoundSubmix& SoundSubmix);
-		void UnloadSoundSubmix(USoundSubmix& SoundSubmix);
+		void LoadSoundSubmix(const USoundSubmixBase& SoundSubmix);
+
+		void InitSoundfieldAndEndpointDataForSubmix(const USoundSubmixBase& InSoundSubmix, FMixerSubmixPtr MixerSubmix, bool bAllowReInit);
+
+		void UnloadSoundSubmix(const USoundSubmixBase& SoundSubmix);
 
 	private:
 
-		bool IsMasterSubmixType(USoundSubmix* InSubmix) const;
-		FMixerSubmixPtr GetMasterSubmixInstance(USoundSubmix* InSubmix);
+		bool IsMasterSubmixType(const USoundSubmixBase* InSubmix) const;
+		FMixerSubmixPtr GetMasterSubmixInstance(const USoundSubmixBase* InSubmix);
 
 		// Pushes the command to a audio render thread command queue to be executed on render thread
 		void AudioRenderThreadCommand(TFunction<void()> Command);
@@ -269,12 +261,15 @@ namespace Audio
 		IAudioMixerPlatformInterface* AudioMixerPlatform;
 		
 		/** Contains a map of channel/speaker azimuth positions. */
-		FChannelPositionInfo DefaultChannelAzimuthPosition[EAudioMixerChannel::MaxSupportedChannel];
+		FChannelPositionInfo DefaultChannelAzimuthPositions[EAudioMixerChannel::MaxSupportedChannel];
 
 		/** The azimuth positions for submix channel types. */
-		TMap<ESubmixChannelFormat, TArray<FChannelPositionInfo>> ChannelAzimuthPositions;
+		TArray<FChannelPositionInfo> DeviceChannelAzimuthPositions;
 
-		int32 OutputChannels[(int32)ESubmixChannelFormat::Count];
+		int32 DeviceOutputChannels;
+
+		/** Channel type arrays for submix channel types. */
+		TArray<EAudioMixerChannel::Type> DeviceChannelArray;
 
 		/** What upmix method to use for mono channel upmixing. */
 		EMonoChannelUpmixMethod MonoChannelUpmixMethod;
@@ -301,7 +296,18 @@ namespace Audio
 		FAudioPlatformDeviceInfo PlatformInfo;
 
 		/** Map of USoundSubmix static data objects to the dynamic audio mixer submix. */
-		TMap<USoundSubmix*, FMixerSubmixPtr> Submixes;
+		TMap<const USoundSubmixBase*, FMixerSubmixPtr> Submixes;
+
+		// Submixes that will sum their audio and send it directly to AudioMixerPlatform.
+		// Submixes are added to this list in RegisterSoundSubmix, and removed in UnregisterSoundSubmix.
+		TArray<FMixerSubmixPtr> DefaultEndpointSubmixes;
+
+		// Submixes that need to be processed, but will be sending their audio to external sends.
+		// Submixes are added to this list in RegisterSoundSubmix and removed in UnregisterSoundSubmix.
+		TArray<FMixerSubmixPtr> ExternalEndpointSubmixes;
+
+		// Contended between RegisterSoundSubmix/UnregisterSoundSubmix on the audio thread and OnProcessAudioStream on the audio mixer thread.
+		FCriticalSection EndpointSubmixesMutationLock;
 
 		/** Which submixes have been told to envelope follow with this audio device. */
 		TArray<USoundSubmix*> EnvelopeFollowingSubmixes;
diff --git a/Engine/Source/Runtime/AudioMixer/Public/AudioMixerSubmix.h b/Engine/Source/Runtime/AudioMixer/Public/AudioMixerSubmix.h
--- a/Engine/Source/Runtime/AudioMixer/Public/AudioMixerSubmix.h
+++ b/Engine/Source/Runtime/AudioMixer/Public/AudioMixerSubmix.h
@@ -4,6 +4,8 @@
 
 #include "AudioMixer.h"
 #include "SampleBuffer.h"
+#include "IAudioEndpoint.h"
+#include "ISoundfieldEndpoint.h"
 #include "Sound/SoundSubmix.h"
 #include "DSP/EnvelopeFollower.h"
 #include "DSP/SpectrumAnalyzer.h"
@@ -25,12 +27,9 @@ namespace Audio
 	struct FSubmixVoiceData
 	{
 		float SendLevel;
-		uint32 AmbisonicsEncoderId;
-		FAmbisonicsEncoderInputData CachedEncoderInputData;
 
 		FSubmixVoiceData()
 			: SendLevel(1.0f)
-			, AmbisonicsEncoderId(INDEX_NONE)
 		{
 		}
 	};
@@ -40,22 +39,40 @@ namespace Audio
 	struct FChildSubmixInfo
 	{
 		TWeakPtr<FMixerSubmix, ESPMode::ThreadSafe> SubmixPtr;
-		bool bNeedsAmbisonicsEncoding;
+
+		// If the child submix is not a soundfield submix, we may need to encode its audio output in ProcessAudio.
+		TUniquePtr<ISoundfieldEncoderStream> Encoder;
+
+		// If this child submix is a soundfield submix that we can read the output of, we may need to transcode it's audio output.
+		TUniquePtr<ISoundfieldTranscodeStream> Transcoder;
+
+		// This is filled by either the Encoder or the Transcoder, and passed to this submix' mixer.
+		TUniquePtr<ISoundfieldAudioPacket> IncomingPacketToTranscode;
 
 		FChildSubmixInfo()
-			: bNeedsAmbisonicsEncoding(true)
+		{}
+
+		// TMap doesn't compile using non-copyable types as value types, though you can build and use a TMap without using the copy constructor.
+		FChildSubmixInfo(const FChildSubmixInfo& InInfo)
+			: FChildSubmixInfo()
+		{
+			checkf(false, TEXT("FChildSubmixInfo is not copyable. If you are using FChildSubmixInfo, consider using Emplace or Add(FChildSubmixInfo&&)."));
+		}
+
+		FChildSubmixInfo(TWeakPtr<FMixerSubmix, ESPMode::ThreadSafe> SubmixWeakPtr)
+			: SubmixPtr(SubmixWeakPtr)
 		{
 		}
 	};
 
-	class FMixerSubmix
+	class AUDIOMIXER_API FMixerSubmix
 	{
 	public:
 		FMixerSubmix(FMixerDevice* InMixerDevice);
 		virtual ~FMixerSubmix();
 
 		// Initialize the submix object with the USoundSubmix ptr. Sets up child and parent connects.
-		void Init(USoundSubmix* InSoundSubmix, bool bAllowReInit = true);
+		void Init(const USoundSubmixBase* InSoundSubmix, bool bAllowReInit = true);
 
 		// Returns the mixer submix Id
 		uint32 GetId() const { return Id; }
@@ -76,7 +93,7 @@ namespace Audio
 		void SetDynamicOutputVolume(float InVolume);
 
 		// Gets the submix channels channels
-		ESubmixChannelFormat GetSubmixChannels() const;
+		int32 GetSubmixChannels() const;
 
 		// Gets this submix's parent submix
 		TWeakPtr<FMixerSubmix, ESPMode::ThreadSafe> GetParentSubmix();
@@ -109,7 +126,12 @@ namespace Audio
 		void SetBackgroundMuted(bool bInMuted);
 
 		// Function which processes audio.
-		void ProcessAudio(const ESubmixChannelFormat ParentInputChannels, AlignedFloatBuffer& OutAudio);
+		void ProcessAudio(AlignedFloatBuffer& OutAudio);
+		void ProcessAudio(ISoundfieldAudioPacket& OutputAudio);
+
+		// This should be called if this submix doesn't send it's audio to a parent submix,
+		// but rather an external endpoint.
+		void ProcessAudioAndSendToEndpoint();
 
 		// Returns the device sample rate this submix is rendering to
 		int32 GetSampleRate() const;
@@ -123,8 +145,18 @@ namespace Audio
 		// Returns the submix effect at the given effect chain index
 		FSoundEffectSubmixPtr GetSubmixEffect(const int32 InIndex);
 
-		// updates settings, potentially creating or removing ambisonics streams based on
-		void OnAmbisonicsSettingsChanged(UAmbisonicsSubmixSettingsBase* AmbisonicsSettings);
+		// This must be called on the entire submix graph before calling SetupSoundfieldStreams.
+		void SetSoundfieldFactory(ISoundfieldFactory* InSoundfieldFactory);
+
+		// updates settings, potentially creating or removing ambisonics streams based on what types of submixes this submix is connected to.
+		void SetupSoundfieldStreams(const USoundfieldEncodingSettingsBase* SoundfieldSettings, TArray<USoundfieldEffectBase*>& Processors, ISoundfieldFactory* InSoundfieldFactory);
+		void TeardownSoundfieldStreams();
+
+		void SetupEndpoint(IAudioEndpointFactory* InFactory, const UAudioEndpointSettingsBase* InSettings);
+		void SetupEndpoint(ISoundfieldEndpointFactory* InFactory, const USoundfieldEndpointSettingsBase* InSettings);
+
+		void UpdateEndpointSettings(TUniquePtr<IAudioEndpointSettingsProxy>&& InSettings);
+		void UpdateEndpointSettings(TUniquePtr<ISoundfieldEndpointSettingsProxy>&& InSettings);
 
 		// This is called by the corresponding USoundSubmix when StartRecordingOutput is called.
 		AUDIOMIXER_API void OnStartRecordingOutput(float ExpectedDuration);
@@ -170,45 +202,57 @@ namespace Audio
 		// Broadcast the envelope value on the game thread
 		void BroadcastEnvelope();
 
+		// returns true if this submix is encoded to a soundfield.
+		bool IsSoundfieldSubmix() const;
+
+		// returns true if this submix sends it's audio to the default endpoint.
+		bool IsDefaultEndpointSubmix() const;
+
+		// Returns true if this submix sends its audio to an IAudioEndpoint.
+		bool IsExternalEndpointSubmix() const;
+
+		// returns true if this submix sends its audio to an ISoundfieldEndpoint.
+		bool IsSoundfieldEndpointSubmix() const;
+
+		// Get a unique key for this submix's format and settings.
+		// If another submix has an identical format and settings it will have an equivalent key.
+		FSoundfieldEncodingKey GetKeyForSubmixEncoding();
+
+		ISoundfieldFactory* GetSoundfieldFactory();
+
+		ISoundfieldEncodingSettingsProxy& GetSoundfieldSettings();
+
+		FAudioPluginInitializationParams GetInitializationParamsForSoundfieldStream();
+
+		FSoundfieldSpeakerPositionalData GetDefaultPositionalDataForAudioDevice();
+
 	protected:
 		// Initialize the submix internal 
 		void InitInternal();
 
 		// Down mix the given buffer to the desired down mix channel count
-		void FormatChangeBuffer(const ESubmixChannelFormat NewChannelType, AlignedFloatBuffer& InBuffer, AlignedFloatBuffer& OutNewBuffer);
+		static void DownmixBuffer(const int32 InChannels, const AlignedFloatBuffer& InBuffer, const int32 OutChannels, AlignedFloatBuffer& OutNewBuffer);
 
 		void MixBufferDownToMono(const AlignedFloatBuffer& InBuffer, int32 NumInputChannels, AlignedFloatBuffer& OutBuffer);
 
-		// Set up ambisonics encoder. Called when ambisonics settings are changed.
-		void SetUpAmbisonicsEncoder();
-
-		// Set up ambisonics decoder. Called when ambisonics settings are changed.
-		void SetUpAmbisonicsDecoder();
-
-		// Clean up ambisonics encoder.
-		void TearDownAmbisonicsEncoder();
-
-		// Clean up ambisonics decoder.
-		void TearDownAmbisonicsDecoder();
-
-		// Check if we need to encode for ambisonics for childen (TODO)
-		void UpdateAmbisonicsEncoderForChildren();
+		void SetupSoundfieldEncodersForChildren();
+		void SetupSoundfieldEncodingForChild(FChildSubmixInfo& InChild);
 
 		// Check to see if we need to decode from ambisonics for parent
-		void UpdateAmbisonicsDecoderForParent();
+		void SetupSoundfieldStreamForParent();
 
 		// This sets up the ambisonics positional data for speakers, based on what new format we need to convert to.
-		void SetUpAmbisonicsPositionalData();
+		void SetUpSoundfieldPositionalData(const TSharedPtr<Audio::FMixerSubmix, ESPMode::ThreadSafe>& InParentSubmix);
 
-		// Encode a source and sum it into the AmbisonicsBuffer.
-		void EncodeAndMixInSource(AlignedFloatBuffer& InAudioData, FSubmixVoiceData& InVoiceInfo);
+		// Encode a source and sum it into the mixed soundfield.
+		void MixInSource(const ISoundfieldAudioPacket& InAudio, const ISoundfieldEncodingSettingsProxy& InSettings, ISoundfieldAudioPacket& PacketToSumTo);
 
-		// Encodes child submix into ambisonics (TODO)
-		void EncodeAndMixInChildSubmix(FChildSubmixInfo& Child);
+		// Calls ProcessAudio on the child submix, performs all neccessary conversions and mixes in it's resulting audio.
+		void MixInChildSubmix(FChildSubmixInfo& Child, ISoundfieldAudioPacket& PacketToSumTo);
 
-	public:
-		// Cached pointer to ambisonics settings.
-		UAmbisonicsSubmixSettingsBase* AmbisonicsSettings;
+		FName GetSoundfieldFormat() const;
+
+		TUniquePtr<ISoundfieldTranscodeStream> GetTranscoderForChildSubmix(const TSharedPtr<Audio::FMixerSubmix, ESPMode::ThreadSafe>& InChildSubmix);
 
 	protected:
 
@@ -256,18 +300,127 @@ namespace Audio
 		AlignedFloatBuffer DownmixedBuffer;
 		AlignedFloatBuffer SourceInputBuffer;
 
-		ESubmixChannelFormat ChannelFormat;
 		int32 NumChannels;
 		int32 NumSamples;
 
-		// Cached ambisonics mixer
-		TAmbisonicsMixerPtr AmbisonicsMixer;
+		/**
+		 * Individual processor in our 
+		 */
+		struct FSoundfieldEffectProcessorData
+		{
+			TUniquePtr<ISoundfieldEffectSettingsProxy> Settings;
+			TUniquePtr<ISoundfieldEffectInstance> Processor;
+
+			FSoundfieldEffectProcessorData(ISoundfieldFactory* InFactory, ISoundfieldEncodingSettingsProxy& InSettings, USoundfieldEffectBase* InProcessorBase)
+			{
+				check(InFactory);
+
+				// As a sanity check, make sure if we've gotten to this point, this DSP processor supports this submix's format.
+				check(InProcessorBase->SupportsFormat(InFactory->GetSoundfieldFormatName()));
+
+				Processor = InProcessorBase->PrivateGetNewProcessor(InSettings);
+				
+				// If the processor doesn't have any settings, get the default settings for a processor of this type.
+				const USoundfieldEffectSettingsBase* ProcessorSettings = InProcessorBase->Settings;
+				if (!ProcessorSettings)
+				{
+					ProcessorSettings = InProcessorBase->PrivateGetDefaultSettings();
+				}
+
+				Settings = ProcessorSettings->PrivateGetProxy();
+			}
+		};
+
+		struct FSoundfieldStreams
+		{
+			ISoundfieldFactory* Factory;
+
+			// This encoder is used for the mixed down audio from all non-soundfield submixes plugged into
+			// this submix. Will not be set up if ISoundfieldFactory::ShouldEncodeAllStreamsIndependently 
+			// returns true.
+			TUniquePtr<ISoundfieldEncoderStream> DownmixedChildrenEncoder;
+			
+			// Encoder used if a normal submix outputs to this submix.
+			TUniquePtr<ISoundfieldDecoderStream> ParentDecoder;
+
+			// This is the positional data we are decoding 
+			FSoundfieldSpeakerPositionalData CachedPositionalData;
+
+			// Mixes all encoded child submix inputs.
+			TUniquePtr<ISoundfieldMixerStream> Mixer;
+
+			// This is the packet we mix all input sources and child submixes to.
+			TUniquePtr<ISoundfieldAudioPacket> MixedDownAudio;
+
+			// Current settings for this submix.
+			TUniquePtr<ISoundfieldEncodingSettingsProxy> Settings;
+
+			// All soundfield processors attached to this submix.  
+			TArray<FSoundfieldEffectProcessorData> EffectProcessors;
 
-		// Encoder ID set up with Ambisonics Mixer. Set to INDEX_NONE if there is no encoder stream open.
-		uint32 SubmixAmbisonicsEncoderID;
+			// This critical section is contended by the soundfield overload of ProcessAudio and SetupSoundfieldStreams.
+			FCriticalSection StreamsLock;
+
+			FSoundfieldStreams()
+				: Factory(nullptr)
+			{}
+
+			void Reset()
+			{
+				Factory = nullptr;
+				ParentDecoder.Reset();
+				Mixer.Reset();
+				Settings.Reset();
+			}
+		};
+
+		FSoundfieldStreams SoundfieldStreams;
+
+		struct FEndpointData
+		{
+			// For endpoint submixes,
+			// this is the primary method of pushing audio to the endpoint.
+			Audio::FPatchInput Input;
 
-		// Decoder ID set up with Ambisonics Mixer. Set to INDEX_NONE if there is no decoder stream open.
-		uint32 SubmixAmbisonicsDecoderID;
+			TUniquePtr<IAudioEndpoint> NonSoundfieldEndpoint;
+			TUniquePtr<ISoundfieldEndpoint> SoundfieldEndpoint;
+
+			// for non-soundfield endpoints, we use these buffers for processing.
+			AlignedFloatBuffer AudioBuffer;
+			AlignedFloatBuffer ResampledAudioBuffer;
+			AlignedFloatBuffer DownmixedResampledAudioBuffer;
+			AlignedFloatBuffer DownmixChannelMap;
+
+			// Number of channels and sample rate for the external endpoint.
+			int32 NumChannels;
+			float SampleRate;
+
+			// This is used if the endpoint has a different sample rate than our audio engine.
+			Audio::FResampler Resampler;
+			bool bShouldResample;
+
+			// for soundfield endpoints, this is the buffer we use to send audio to the endpoint.
+			TUniquePtr<ISoundfieldAudioPacket> AudioPacket;
+
+			FEndpointData()
+				: NumChannels(0)
+				, SampleRate(0.0f)
+				, bShouldResample(false)
+			{}
+
+			void Reset()
+			{
+				AudioBuffer.Reset();
+				ResampledAudioBuffer.Reset();
+				DownmixedResampledAudioBuffer.Reset();
+				DownmixChannelMap.Reset();
+				NonSoundfieldEndpoint.Reset();
+				SoundfieldEndpoint.Reset();
+			}
+		};
+
+		FEndpointData EndpointData;
+		
 
 		// The output volume of the submix set via the USoundSubmix property. Can be set in the editor.
 		float InitializedOutputVolume;
@@ -290,12 +443,6 @@ namespace Audio
 		// This buffer is used to downmix the submix output to mono before submitting it to the SpectrumAnalyzer.
 		AlignedFloatBuffer MonoMixBuffer;
 
-		// This buffer is encoded into for each source, then summed into the ambisonics buffer.
-		AlignedFloatBuffer InputAmbisonicsBuffer;
-
-		// Cached positional data for Ambisonics decoder.
-		FAmbisonicsDecoderPositionalData CachedPositionalData;
-
 		// Submix command queue to shuffle commands from audio thread to audio render thread.
 		TQueue<TFunction<void()>> CommandQueue;
 
@@ -327,11 +474,8 @@ namespace Audio
 		FCriticalSection RecordingCriticalSection;
 
 		// Handle back to the owning USoundSubmix. Used when the device is shutdown to prematurely end a recording.
-		USoundSubmix* OwningSubmixObject;
+		const USoundSubmixBase* OwningSubmixObject;
 
 		friend class FMixerDevice;
 	};
-
-
-
 }
diff --git a/Engine/Source/Runtime/AudioMixer/Public/Components/SynthComponent.h b/Engine/Source/Runtime/AudioMixer/Public/Components/SynthComponent.h
--- a/Engine/Source/Runtime/AudioMixer/Public/Components/SynthComponent.h
+++ b/Engine/Source/Runtime/AudioMixer/Public/Components/SynthComponent.h
@@ -110,7 +110,7 @@ public:
 
 	/** Sets how much audio the sound should send to the given submix. */
 	UFUNCTION(BlueprintCallable, Category = "Audio|Components|Audio")
-	void SetSubmixSend(USoundSubmix* Submix, float SendLevel);
+	void SetSubmixSend(USoundSubmixBase* Submix, float SendLevel);
 
 	/** Auto destroy this component on completion */
 	UPROPERTY()
@@ -158,7 +158,7 @@ public:
 
 	/** Submix this sound belongs to */
 	UPROPERTY(EditAnywhere, Category = Effects)
-	USoundSubmix* SoundSubmix;
+	USoundSubmixBase* SoundSubmix;
 
 	/** An array of submix sends. Audio from this sound will send a portion of its audio to these effects.  */
 	UPROPERTY(EditAnywhere, BlueprintReadWrite, Category = Effects)
diff --git a/Engine/Source/Runtime/AudioMixerCore/Public/AudioMixer.h b/Engine/Source/Runtime/AudioMixerCore/Public/AudioMixer.h
--- a/Engine/Source/Runtime/AudioMixerCore/Public/AudioMixer.h
+++ b/Engine/Source/Runtime/AudioMixerCore/Public/AudioMixer.h
@@ -68,12 +68,7 @@ namespace EAudioMixerChannel
 		TopBackRight,
 		Unknown,
 		ChannelTypeCount,
-
-		// Alias Values
-		Ambisonics_W = FrontLeft,
-		Ambisonics_X = FrontRight,
-		Ambisonics_Y = FrontCenter,
-		Ambisonics_Z = LowFrequency,
+		DefaultChannel = FrontLeft
 	};
 
 	static const int32 MaxSupportedChannel = EAudioMixerChannel::TopCenter;
@@ -275,7 +270,7 @@ namespace Audio
 	};
 
 	/** Struct used to store render time analysis data. */
-	struct FAudioRenderTimeAnalysis
+	struct AUDIOMIXERCORE_API FAudioRenderTimeAnalysis
 	{
 		double AvgRenderTime;
 		double MaxRenderTime;
diff --git a/Engine/Source/Runtime/AudioMixerCore/Public/AudioMixerNullDevice.h b/Engine/Source/Runtime/AudioMixerCore/Public/AudioMixerNullDevice.h
--- a/Engine/Source/Runtime/AudioMixerCore/Public/AudioMixerNullDevice.h
+++ b/Engine/Source/Runtime/AudioMixerCore/Public/AudioMixerNullDevice.h
@@ -13,7 +13,7 @@ namespace Audio
 	 * This class, when started, spawns a new high priority thread that exists to query an FAudioMixerPlatformInterface
 	 * and immediately throw out whatever buffers it receives.
 	 */
-	class FMixerNullCallback : protected FRunnable
+	class AUDIOMIXERCORE_API FMixerNullCallback : protected FRunnable
 	{
 	public:
 
diff --git a/Engine/Source/Runtime/AudioMixerCore/Public/AudioMixerTypes.h b/Engine/Source/Runtime/AudioMixerCore/Public/AudioMixerTypes.h
--- a/Engine/Source/Runtime/AudioMixerCore/Public/AudioMixerTypes.h
+++ b/Engine/Source/Runtime/AudioMixerCore/Public/AudioMixerTypes.h
@@ -17,6 +17,7 @@ namespace Audio {
 			SDL2,		// Linux
 			OpenSLES, 	// Android
 			Switch, 	// Switch
+			Other,      // Generic output type.
 			Null		// Unknown/not Supported
 		};
 	}
diff --git a/Engine/Source/Runtime/Engine/Classes/Components/AudioComponent.h b/Engine/Source/Runtime/Engine/Classes/Components/AudioComponent.h
--- a/Engine/Source/Runtime/Engine/Classes/Components/AudioComponent.h
+++ b/Engine/Source/Runtime/Engine/Classes/Components/AudioComponent.h
@@ -489,7 +489,7 @@ public:
 
 	/** Sets how much audio the sound should send to the given submix. */
 	UFUNCTION(BlueprintCallable, Category = "Audio|Components|Audio")
-	void SetSubmixSend(USoundSubmix* Submix, float SendLevel);
+	void SetSubmixSend(USoundSubmixBase* Submix, float SendLevel);
 
 	/** Sets how much audio the sound should send to the given Source Bus (PRE Source Effects).
 		if the Bus Send doesn't already exist, it will be added to the overrides on the active sound */
diff --git a/Engine/Source/Runtime/Engine/Classes/Sound/AudioSettings.h b/Engine/Source/Runtime/Engine/Classes/Sound/AudioSettings.h
--- a/Engine/Source/Runtime/Engine/Classes/Sound/AudioSettings.h
+++ b/Engine/Source/Runtime/Engine/Classes/Sound/AudioSettings.h
@@ -112,10 +112,6 @@ class ENGINE_API UAudioSettings : public UDeveloperSettings
 	UPROPERTY(config, EditAnywhere, Category="Mix", meta=(AllowedClasses="SoundSubmix", DisplayName = "EQ Submix (Legacy)"), AdvancedDisplay)
 	FSoftObjectPath EQSubmix;
 
-	/** The submix through which sounds spatialized using ambisonics are routed */
-	UPROPERTY(config, EditAnywhere, Category="Mix", meta=(AllowedClasses="SoundSubmix"))
-	FSoftObjectPath AmbisonicSubmix;
-
 	/** Sample rate used for voice over IP. VOIP audio is resampled to the application's sample rate on the receiver side. */
 	UPROPERTY(config, EditAnywhere, Category = "Audio", meta = (DisplayName = "VOIP Sample Rate"))
 	EVoiceSampleRate VoiPSampleRate;
diff --git a/Engine/Source/Runtime/Engine/Classes/Sound/SoundBase.h b/Engine/Source/Runtime/Engine/Classes/Sound/SoundBase.h
--- a/Engine/Source/Runtime/Engine/Classes/Sound/SoundBase.h
+++ b/Engine/Source/Runtime/Engine/Classes/Sound/SoundBase.h
@@ -152,7 +152,7 @@ public:
 	/** Submix to route sound output to. If unset, falls back to referenced SoundClass submix.
 	  * If SoundClass submix is unset, sends to the 'Master Submix' as set in the 'Audio' category of Project Settings'. */
 	UPROPERTY(EditAnywhere, BlueprintReadWrite, Category = "Effects|Submix", meta = (DisplayName = "Submix"))
-	USoundSubmix* SoundSubmixObject;
+	USoundSubmixBase* SoundSubmixObject;
 
 	/** Array of submix sends to which a prescribed amount (see 'Send Level') of this sound is sent. */
 	UPROPERTY(EditAnywhere, BlueprintReadWrite, Category = "Effects|Submix", meta = (DisplayName = "Submix Sends"))
@@ -238,7 +238,7 @@ public:
 	virtual USoundClass* GetSoundClass() const;
 
 	/** Returns the SoundSubmix used for this sound. */
-	virtual USoundSubmix* GetSoundSubmix() const;
+	virtual USoundSubmixBase* GetSoundSubmix() const;
 
 	/** Returns the sound submix sends for this sound. */
 	void GetSoundSubmixSends(TArray<FSoundSubmixSendInfo>& OutSends) const;
diff --git a/Engine/Source/Runtime/Engine/Classes/Sound/SoundCue.h b/Engine/Source/Runtime/Engine/Classes/Sound/SoundCue.h
--- a/Engine/Source/Runtime/Engine/Classes/Sound/SoundCue.h
+++ b/Engine/Source/Runtime/Engine/Classes/Sound/SoundCue.h
@@ -7,8 +7,15 @@
 #include "Sound/SoundAttenuation.h"
 #include "Sound/SoundBase.h"
 #include "Sound/SoundNode.h"
+
+#if WITH_EDITOR
+#include "EdGraph/EdGraph.h"
+#endif
+
 #include "SoundCue.generated.h"
 
+
+
 class USoundCue;
 class USoundNode;
 struct FActiveSound;
diff --git a/Engine/Source/Runtime/Engine/Classes/Sound/SoundSubmix.h b/Engine/Source/Runtime/Engine/Classes/Sound/SoundSubmix.h
--- a/Engine/Source/Runtime/Engine/Classes/Sound/SoundSubmix.h
+++ b/Engine/Source/Runtime/Engine/Classes/Sound/SoundSubmix.h
@@ -3,7 +3,9 @@
 
 #include "CoreMinimal.h"
 
-#include "IAmbisonicsMixer.h"
+#include "ISoundfieldFormat.h"
+#include "IAudioEndpoint.h"
+#include "ISoundfieldEndpoint.h"
 #include "SampleBufferIO.h"
 #include "SoundEffectSubmix.h"
 #include "SoundSubmixSend.h"
@@ -20,36 +22,6 @@ class USoundSubmix;
 class ISubmixBufferListener;
 
 
-/* Submix channel format.
- Allows submixes to have sources mix to a particular channel configuration for potential effect chain requirements.
- Master submix will always render at the device channel count. All child submixes will be down-mixed (or up-mixed) to
- the device channel count. This feature exists to allow specific submix effects to do their work on multi-channel mixes
- of audio.
-*/
-UENUM(BlueprintType)
-enum class ESubmixChannelFormat : uint8
-{
-	// Sets the submix channels to the output device channel count
-	Device UMETA(DisplayName = "Device"),
-
-	// Sets the submix mix to stereo (FL, FR)
-	Stereo UMETA(DisplayName = "Stereo"),
-
-	// Sets the submix to mix to quad (FL, FR, SL, SR)
-	Quad UMETA(DisplayName = "Quad"),
-
-	// Sets the submix to mix 5.1 (FL, FR, FC, LF, SL, SR)
-	FiveDotOne UMETA(DisplayName = "5.1"),
-
-	// Sets the submix to mix audio to 7.1 (FL, FR, FC, LF, BL, BR, SL, SR)
-	SevenDotOne UMETA(DisplayName = "7.1"),
-
-	// Sets the submix to render audio to first-order ambisonics.
-	Ambisonics UMETA(DisplayName = "Ambisonics"),
-
-	Count UMETA(Hidden)
-};
-
 
 /**
 * Called when a recorded file has finished writing to disk.
@@ -76,26 +48,94 @@ public:
 };
 #endif
 
-UCLASS(config = Engine, hidecategories = Object, editinlinenew, BlueprintType)
-class ENGINE_API USoundSubmix : public UObject
+UCLASS(config = Engine, abstract, hidecategories = Object, editinlinenew, BlueprintType)
+class ENGINE_API USoundSubmixBase : public UObject
 {
 	GENERATED_UCLASS_BODY()
 
-	// Child submixes to this sound mix
-	UPROPERTY(VisibleAnywhere, BlueprintReadOnly, Category = SoundSubmix)
-	TArray<USoundSubmix*> ChildSubmixes;
-
-	UPROPERTY()
-	USoundSubmix* ParentSubmix;
-
+public:
 #if WITH_EDITORONLY_DATA
 	/** EdGraph based representation of the SoundSubmix */
 	UEdGraph* SoundSubmixGraph;
 #endif
 
-	// Experimental! Specifies the channel format for the submix. Sources will be mixed at the specified format. Useful for specific effects that need to operate on a specific format.
+	// Child submixes to this sound mix
+	UPROPERTY(VisibleAnywhere, BlueprintReadOnly, Category = SoundSubmix)
+	TArray<USoundSubmixBase*> ChildSubmixes;
+
+protected:
+	//~ Begin UObject Interface.
+	virtual FString GetDesc() override;
+	virtual void BeginDestroy() override;
+	virtual void PostLoad() override;
+
+public:
+	// Sound Submix Editor functionality
+#if WITH_EDITOR
+
+	/**
+	* @return true if the child sound class exists in the tree
+	*/
+	bool RecurseCheckChild(const USoundSubmixBase* ChildSoundSubmix) const;
+
+	/**
+	* Add Referenced objects
+	*
+	* @param	InThis SoundSubmix we are adding references from.
+	* @param	Collector Reference Collector
+	*/
+	static void AddReferencedObjects(UObject* InThis, FReferenceCollector& Collector);
+
+protected:
+
+#if WITH_EDITOR
+	virtual void PostDuplicate(EDuplicateMode::Type DuplicateMode) override;
+	virtual void PreEditChange(UProperty* PropertyAboutToChange) override;
+	virtual void PostEditChangeProperty(struct FPropertyChangedEvent& PropertyChangedEvent) override;
+#endif // WITH_EDITOR
+	//~ End UObject Interface.
+
+private:
+	static TArray<USoundSubmixBase*> BackupChildSubmixes;
+#endif // WITH_EDITOR
+};
+
+/**
+ * This submix class can be derived from for submixes that output to a parent submix.
+ */
+UCLASS(config = Engine, abstract, hidecategories = Object, editinlinenew, BlueprintType)
+class ENGINE_API USoundSubmixWithParentBase : public USoundSubmixBase
+{
+	GENERATED_UCLASS_BODY()
+public:
+
 	UPROPERTY(VisibleAnywhere, BlueprintReadOnly, Category = SoundSubmix)
-	ESubmixChannelFormat ChannelFormat;
+	USoundSubmixBase* ParentSubmix;
+
+	/**
+	* Set the parent submix of this SoundSubmix, removing it as a child from its previous owner
+	*
+	* @param	InParentSubmix	The New Parent Submix of this
+	*/
+	void SetParentSubmix(USoundSubmixBase* InParentSubmix);
+
+protected:
+
+#if WITH_EDITOR
+	virtual void PostEditChangeProperty(struct FPropertyChangedEvent& PropertyChangedEvent) override;
+	virtual void PostDuplicate(EDuplicateMode::Type DuplicateMode) override;
+#endif 
+};
+
+/**
+ * Sound Submix class meant for applying an effect to the downmixed sum of multiple audio sources.
+ */
+UCLASS(config = Engine, hidecategories = Object, editinlinenew, BlueprintType, Meta=(DisplayName="Effect Submix"))
+class ENGINE_API USoundSubmix : public USoundSubmixWithParentBase
+{
+	GENERATED_UCLASS_BODY()
+
+public:
 
 	/** Mute this submix when the application is muted or in the background. Used to prevent submix effect tails from continuing when tabbing out of application or if application is muted. */
 	UPROPERTY(EditAnywhere, BlueprintReadOnly, Category = SoundSubmix)
@@ -106,7 +146,7 @@ class ENGINE_API USoundSubmix : public UObject
 
 	/** Optional settings used by plugins which support ambisonics file playback. */
 	UPROPERTY(EditAnywhere, BlueprintReadWrite, Category = SoundSubmix)
-	UAmbisonicsSubmixSettingsBase* AmbisonicsPluginSettings;
+	USoundfieldEncodingSettingsBase* AmbisonicsPluginSettings;
 
 	/** The attack time in milliseconds for the envelope follower. Delegate callbacks can be registered to get the envelope value of sounds played with this submix. */
 	UPROPERTY(EditAnywhere, BlueprintReadWrite, Category = EnvelopeFollower, meta = (ClampMin = "0", UIMin = "0"))
@@ -161,46 +201,117 @@ class ENGINE_API USoundSubmix : public UObject
 
 protected:
 
-	//~ Begin UObject Interface.
-	virtual FString GetDesc() override;
-	virtual void BeginDestroy() override;
-	virtual void PostLoad() override;
-
 #if WITH_EDITOR
-	virtual void PostDuplicate(EDuplicateMode::Type DuplicateMode) override;
-	virtual void PreEditChange(FProperty* PropertyAboutToChange) override;
 	virtual void PostEditChangeProperty(struct FPropertyChangedEvent& PropertyChangedEvent) override;
-#endif // WITH_EDITOR
-	//~ End UObject Interface.
+#endif
 
 	// State handling for bouncing output.
 	TUniquePtr<Audio::FAudioRecordingData> RecordingData;
+};
+	
+
+/**
+ * Sound Submix class meant for use with soundfield formats, such as Ambisonics.
+ */
+UCLASS(config = Engine, hidecategories = Object, editinlinenew, BlueprintType, Meta = (DisplayName = "Soundfield Submix"))
+class ENGINE_API USoundfieldSubmix : public USoundSubmixWithParentBase
+{
+	GENERATED_UCLASS_BODY()
 
 public:
-	// Sound Submix Editor functionality
-#if WITH_EDITOR
+	ISoundfieldFactory* GetSoundfieldFactoryForSubmix() const;
+	const USoundfieldEncodingSettingsBase* GetSoundfieldEncodingSettings() const;
+	TArray<USoundfieldEffectBase *> GetSoundfieldProcessors() const;
+
+public:
+	/** Currently used format. */
+	UPROPERTY(EditAnywhere, AssetRegistrySearchable, Category = Soundfield)
+	FName SoundfieldEncodingFormat;
+
+	//TODO: Make this editable only if SoundfieldEncodingFormat is non-default,
+	// and filter classes based on ISoundfieldFactory::GetCustomSettingsClass().
+	UPROPERTY(EditAnywhere, Category = Soundfield)
+	USoundfieldEncodingSettingsBase* EncodingSettings;
+
+	// TODO: make this editable only if SoundfieldEncodingFormat is non-default
+	// and filter classes based on USoundfieldProcessorBase::SupportsFormat.
+	UPROPERTY(EditAnywhere, Category = Soundfield)
+	TArray<USoundfieldEffectBase*> SoundfieldEffectChain;
+
+	// Traverses parent submixes until we find a submix that doesn't inherit it's soundfield format.
+	FName GetSubmixFormat() const;
+
+	UPROPERTY()
+	TSubclassOf<USoundfieldEncodingSettingsBase> EncodingSettingsClass;
+
+	// Traverses parent submixes until we find a submix that specifies encoding settings.
+	const USoundfieldEncodingSettingsBase* GetEncodingSettings() const;
+};
+
+/**
+ * Sound Submix class meant for sending audio to an external endpoint, such as controller haptics or an additional audio device.
+ */
+UCLASS(config = Engine, hidecategories = Object, editinlinenew, BlueprintType, Meta = (DisplayName = "Audio Endpoint Submix"))
+class ENGINE_API UEndpointSubmix : public USoundSubmixBase
+{
+	GENERATED_UCLASS_BODY()
+
+public:
+	IAudioEndpointFactory* GetAudioEndpointForSubmix() const;
+	const UAudioEndpointSettingsBase* GetEndpointSettings() const;
+
+public:
+	/** Currently used format. */
+	UPROPERTY(EditAnywhere, AssetRegistrySearchable, Category = Endpoint)
+	FName EndpointType;
+
+	UPROPERTY()
+	TSubclassOf<UAudioEndpointSettingsBase> EndpointSettingsClass;
+
+	//TODO: Make this editable only if EndpointType is non-default,
+	// and filter classes based on ISoundfieldFactory::GetCustomSettingsClass().
+	UPROPERTY(EditAnywhere, Category = Endpoint)
+	UAudioEndpointSettingsBase* EndpointSettings;
+
+};
+
+/**
+ * Sound Submix class meant for sending soundfield-encoded audio to an external endpoint, such as a hardware binaural renderer that supports ambisonics.
+ */
+UCLASS(config = Engine, hidecategories = Object, editinlinenew, BlueprintType, Meta = (DisplayName = "Soundfield Endpoint Submix"))
+class ENGINE_API USoundfieldEndpointSubmix : public USoundSubmixBase
+{
+	GENERATED_UCLASS_BODY()
+
+public:
+	ISoundfieldEndpointFactory* GetSoundfieldEndpointForSubmix() const;
+	const USoundfieldEndpointSettingsBase* GetEndpointSettings() const;
+	const USoundfieldEncodingSettingsBase* GetEncodingSettings() const;
+	TArray<USoundfieldEffectBase*> GetSoundfieldProcessors() const;
+public:
+	/** Currently used format. */
+	UPROPERTY(EditAnywhere, Category = Endpoint, AssetRegistrySearchable)
+	FName SoundfieldEndpointType;
+
+	UPROPERTY()
+	TSubclassOf<UAudioEndpointSettingsBase> EndpointSettingsClass;
 
 	/**
 	* @return true if the child sound class exists in the tree
 	*/
 	bool RecurseCheckChild(const USoundSubmix* ChildSoundSubmix) const;
 
-	/**
-	* Set the parent submix of this SoundSubmix, removing it as a child from its previous owner
-	*
-	* @param	InParentSubmix	The New Parent Submix of this
-	*/
-	void SetParentSubmix(USoundSubmix* InParentSubmix);
+	//TODO: Make this editable only if EndpointType is non-default,
+	// and filter classes based on ISoundfieldFactory::GetCustomSettingsClass().
+	UPROPERTY(EditAnywhere, Category = Endpoint)
+	USoundfieldEndpointSettingsBase* EndpointSettings;
 
-	/**
-	* Add Referenced objects
-	*
-	* @param	InThis SoundSubmix we are adding references from.
-	* @param	Collector Reference Collector
-	*/
-	static void AddReferencedObjects(UObject* InThis, FReferenceCollector& Collector);
+	UPROPERTY()
+	TSubclassOf<USoundfieldEncodingSettingsBase> EncodingSettingsClass;
 
-private:
-	TArray<USoundSubmix*> BackupChildSubmixes;
-#endif // WITH_EDITOR
+	UPROPERTY(EditAnywhere, Category = Soundfield)
+	USoundfieldEncodingSettingsBase* EncodingSettings;
+
+	UPROPERTY(EditAnywhere, Category = Soundfield)
+	TArray<USoundfieldEffectBase*> SoundfieldEffectChain;
 };
diff --git a/Engine/Source/Runtime/Engine/Classes/Sound/SoundSubmixSend.h b/Engine/Source/Runtime/Engine/Classes/Sound/SoundSubmixSend.h
--- a/Engine/Source/Runtime/Engine/Classes/Sound/SoundSubmixSend.h
+++ b/Engine/Source/Runtime/Engine/Classes/Sound/SoundSubmixSend.h
@@ -11,7 +11,7 @@
 
 
 // Forward Declarations
-class USoundSubmix;
+class USoundSubmixBase;
 
 
 DECLARE_DYNAMIC_DELEGATE_OneParam(FOnSubmixEnvelopeBP, const TArray<float>&, Envelope);
@@ -57,7 +57,7 @@ struct ENGINE_API FSoundSubmixSendInfo
 
 	// The submix to send the audio to
 	UPROPERTY(EditAnywhere, Category = SubmixSend)
-	USoundSubmix* SoundSubmix;
+	USoundSubmixBase* SoundSubmix;
 
 	// The amount of audio to send
 	UPROPERTY(EditAnywhere, Category = SubmixSend)
diff --git a/Engine/Source/Runtime/Engine/Engine.Build.cs b/Engine/Source/Runtime/Engine/Engine.Build.cs
--- a/Engine/Source/Runtime/Engine/Engine.Build.cs
+++ b/Engine/Source/Runtime/Engine/Engine.Build.cs
@@ -90,8 +90,9 @@ public class Engine : ModuleRules
 				"PakFile",
 				"NetworkReplayStreaming",
 				"PhysicsCore",
-                "SignalProcessing"
-			}
+                "SignalProcessing",
+                "AudioExtensions"
+            }
 		);
 
 		PrivateDependencyModuleNames.AddRange(
diff --git a/Engine/Source/Runtime/Engine/Private/ActiveSound.cpp b/Engine/Source/Runtime/Engine/Private/ActiveSound.cpp
--- a/Engine/Source/Runtime/Engine/Private/ActiveSound.cpp
+++ b/Engine/Source/Runtime/Engine/Private/ActiveSound.cpp
@@ -322,7 +322,7 @@ USoundClass* FActiveSound::GetSoundClass() const
 	return nullptr;
 }
 
-USoundSubmix* FActiveSound::GetSoundSubmix() const
+USoundSubmixBase* FActiveSound::GetSoundSubmix() const
 {
 	return Sound ? Sound->GetSoundSubmix() : nullptr;
 }
@@ -386,6 +386,8 @@ void FActiveSound::GetSoundSubmixSends(TArray<FSoundSubmixSendInfo>& OutSends) c
 					bOverridden = true;
 					break;
 				}
+
+				ensure(OutSendInfo.SendLevel > 0.0f);
 			}
 
 			if (!bOverridden)
diff --git a/Engine/Source/Runtime/Engine/Private/AudioDevice.cpp b/Engine/Source/Runtime/Engine/Private/AudioDevice.cpp
--- a/Engine/Source/Runtime/Engine/Private/AudioDevice.cpp
+++ b/Engine/Source/Runtime/Engine/Private/AudioDevice.cpp
@@ -255,7 +255,6 @@ FAudioDevice::FAudioDevice()
 	, SpatializationPluginInterface(nullptr)
 	, ReverbPluginInterface(nullptr)
 	, OcclusionInterface(nullptr)
-	, AmbisonicsMixer(nullptr)
 	, MaxSources(0)
 	, MaxChannels(0)
 	, MaxChannels_GameThread(0)
diff --git a/Engine/Source/Runtime/Engine/Private/AudioDeviceManager.cpp b/Engine/Source/Runtime/Engine/Private/AudioDeviceManager.cpp
--- a/Engine/Source/Runtime/Engine/Private/AudioDeviceManager.cpp
+++ b/Engine/Source/Runtime/Engine/Private/AudioDeviceManager.cpp
@@ -711,7 +711,7 @@ void FAudioDeviceManager::InitSoundClasses()
 	}
 }
 
-void FAudioDeviceManager::RegisterSoundSubmix(USoundSubmix* SoundSubmix)
+void FAudioDeviceManager::RegisterSoundSubmix(const USoundSubmixBase* SoundSubmix)
 {
 	FScopeLock ScopeLock(&DeviceMapCriticalSection);
 	for (auto& DeviceContainer : Devices)
@@ -720,7 +720,7 @@ void FAudioDeviceManager::RegisterSoundSubmix(USoundSubmix* SoundSubmix)
 	}
 }
 
-void FAudioDeviceManager::UnregisterSoundSubmix(USoundSubmix* SoundSubmix)
+void FAudioDeviceManager::UnregisterSoundSubmix(const USoundSubmixBase* SoundSubmix)
 {
 	FScopeLock ScopeLock(&DeviceMapCriticalSection);
 	for (auto& DeviceContainer : Devices)
@@ -756,7 +756,7 @@ void FAudioDeviceManager::UpdateSourceEffectChain(const uint32 SourceEffectChain
 	}
 }
 
-void FAudioDeviceManager::UpdateSubmix(USoundSubmix* SoundSubmix)
+void FAudioDeviceManager::UpdateSubmix(USoundSubmixBase* SoundSubmix)
 {
 	FScopeLock ScopeLock(&DeviceMapCriticalSection);
 	for (auto& DeviceContainer : Devices)
diff --git a/Engine/Source/Runtime/Engine/Private/AudioSettings.cpp b/Engine/Source/Runtime/Engine/Private/AudioSettings.cpp
--- a/Engine/Source/Runtime/Engine/Private/AudioSettings.cpp
+++ b/Engine/Source/Runtime/Engine/Private/AudioSettings.cpp
@@ -47,9 +47,6 @@ void UAudioSettings::AddDefaultSettings()
 #if WITH_EDITOR
 void UAudioSettings::PreEditChange(FProperty* PropertyAboutToChange)
 {
-	// Cache ambisonic submix in case user tries to set to submix that isn't set to ambisonics
-	CachedAmbisonicSubmix = AmbisonicSubmix;
-
 	// Cache master submix in case user tries to set to submix that isn't a top-level submix
 	CachedMasterSubmix = MasterSubmix;
 
@@ -90,27 +87,6 @@ void UAudioSettings::PostEditChangeChainProperty(FPropertyChangedChainEvent& Pro
 		{
 			bPromptRestartRequired = true;
 		}
-		else if (PropertyName == GET_MEMBER_NAME_CHECKED(UAudioSettings, AmbisonicSubmix))
-		{
-			if (USoundSubmix* NewSubmix = Cast<USoundSubmix>(AmbisonicSubmix.TryLoad()))
-			{
-				if (NewSubmix->ChannelFormat != ESubmixChannelFormat::Ambisonics)
-				{
-					FNotificationInfo Info(LOCTEXT("AudioSettings_InvalidAmbisonicSubmixFormat",
-						"Ambisonic Submix format must be set to 'Ambisonics' in order to be set as 'Master Ambisonics Submix'."));
-					Info.bFireAndForget = true;
-					Info.ExpireDuration = 2.0f;
-					Info.bUseThrobber = true;
-					FSlateNotificationManager::Get().AddNotification(Info);
-
-					AmbisonicSubmix = CachedAmbisonicSubmix;
-				}
-			}
-			else
-			{
-				bPromptRestartRequired = true;
-			}
-		}
 		else if (PropertyName == GET_MEMBER_NAME_CHECKED(UAudioSettings, QualityLevels))
 		{
 			if (QualityLevels.Num() == 0)
diff --git a/Engine/Source/Runtime/Engine/Private/Components/AudioComponent.cpp b/Engine/Source/Runtime/Engine/Private/Components/AudioComponent.cpp
--- a/Engine/Source/Runtime/Engine/Private/Components/AudioComponent.cpp
+++ b/Engine/Source/Runtime/Engine/Private/Components/AudioComponent.cpp
@@ -1296,7 +1296,7 @@ void UAudioComponent::AdjustAttenuation(const FSoundAttenuationSettings& InAtten
 	}
 }
 
-void UAudioComponent::SetSubmixSend(USoundSubmix* Submix, float SendLevel)
+void UAudioComponent::SetSubmixSend(USoundSubmixBase* Submix, float SendLevel)
 {
 	if (FAudioDevice* AudioDevice = GetAudioDevice())
 	{
