Fix build errors from merge resolutions:
1) get rid of DLL export inside AudioMixerSubmix even though Audio::FMixerSubmix is already declared as a DLL export
2) Fix downcast in SoundSubmixEditor to use USoundSubmixBase instead

diff --git a/Engine/Source/Editor/AudioEditor/Private/SoundSubmixEditor.cpp b/Engine/Source/Editor/AudioEditor/Private/SoundSubmixEditor.cpp
--- a/Engine/Source/Editor/AudioEditor/Private/SoundSubmixEditor.cpp
+++ b/Engine/Source/Editor/AudioEditor/Private/SoundSubmixEditor.cpp
@@ -473,7 +473,7 @@ void FSoundSubmixEditor::AddMissingEditableSubmixes()
 		for (UEdGraphNode* Node : Graph->Nodes)
 		{
 			USoundSubmixGraphNode* GraphNode = CastChecked<USoundSubmixGraphNode>(Node);
-			USoundSubmix* UntrackedSubmix = GraphNode->SoundSubmix;
+			USoundSubmixBase* UntrackedSubmix = GraphNode->SoundSubmix;
 			if (UntrackedSubmix && !GetEditingObjects().Contains(UntrackedSubmix))
 			{
 				bChanged = true;
diff --git a/Engine/Source/Runtime/AudioMixer/Public/AudioMixerSubmix.h b/Engine/Source/Runtime/AudioMixer/Public/AudioMixerSubmix.h
--- a/Engine/Source/Runtime/AudioMixer/Public/AudioMixerSubmix.h
+++ b/Engine/Source/Runtime/AudioMixer/Public/AudioMixerSubmix.h
@@ -159,10 +159,10 @@ namespace Audio
 		void UpdateEndpointSettings(TUniquePtr<ISoundfieldEndpointSettingsProxy>&& InSettings);
 
 		// This is called by the corresponding USoundSubmix when StartRecordingOutput is called.
-		AUDIOMIXER_API void OnStartRecordingOutput(float ExpectedDuration);
+		void OnStartRecordingOutput(float ExpectedDuration);
 
 		// This is called by the corresponding USoundSubmix when StopRecordingOutput is called.
-		AUDIOMIXER_API AlignedFloatBuffer& OnStopRecordingOutput(float& OutNumChannels, float& OutSampleRate);
+		AlignedFloatBuffer& OnStopRecordingOutput(float& OutNumChannels, float& OutSampleRate);
 
 		// This is called by the corresponding USoundSubmix when PauseRecording is called.
 		void PauseRecordingOutput();
