Fix race condition on rebuilding submix graph, fix issue with updating params on submixes.

diff --git a/Engine/Source/Runtime/AudioMixer/Private/AudioMixerDevice.cpp b/Engine/Source/Runtime/AudioMixer/Private/AudioMixerDevice.cpp
--- a/Engine/Source/Runtime/AudioMixer/Private/AudioMixerDevice.cpp
+++ b/Engine/Source/Runtime/AudioMixer/Private/AudioMixerDevice.cpp
@@ -985,27 +985,92 @@ namespace Audio
 #endif // WITH_EDITOR
 	}
 
-	void FMixerDevice::SetSubmixOutputVolume(USoundSubmix* InSoundSubmix, float NewVolume)
+	void FMixerDevice::SetSubmixWetDryLevel(USoundSubmix* InSoundSubmix, float InOutputVolume, float InWetLevel, float InDryLevel)
 	{
 		if (!IsInAudioThread())
 		{
-			DECLARE_CYCLE_STAT(TEXT("FAudioThreadTask.SetSubmixOutputVolume"), STAT_AudioSetSubmixOutputVolume, STATGROUP_AudioThreadCommands);
+			FMixerDevice* MixerDevice = this;
+			FAudioThread::RunCommandOnAudioThread([MixerDevice, InSoundSubmix, InOutputVolume, InWetLevel, InDryLevel]()
+			{
+				MixerDevice->SetSubmixWetDryLevel(InSoundSubmix, InOutputVolume, InWetLevel, InDryLevel);
+			});
+			return;
+		}
+
+		FMixerSubmixPtr MixerSubmixPtr = GetSubmixInstance(InSoundSubmix).Pin();
+		if (MixerSubmixPtr.IsValid())
+		{
+			AudioRenderThreadCommand([MixerSubmixPtr, InOutputVolume, InWetLevel, InDryLevel]()
+			{
+				MixerSubmixPtr->SetOutputVolume(InOutputVolume);
+				MixerSubmixPtr->SetWetLevel(InWetLevel);
+				MixerSubmixPtr->SetDryLevel(InDryLevel);
+			});
+		}
+	}
 
+	void FMixerDevice::SetSubmixOutputVolume(USoundSubmix* InSoundSubmix, float InOutputVolume)
+	{
+		if (!IsInAudioThread())
+		{
 			FMixerDevice* MixerDevice = this;
-			FAudioThread::RunCommandOnAudioThread([MixerDevice, InSoundSubmix, NewVolume]()
+			FAudioThread::RunCommandOnAudioThread([MixerDevice, InSoundSubmix, InOutputVolume]()
 			{
-				CSV_SCOPED_TIMING_STAT(Audio, SetSubmixOutputVolume);
-				MixerDevice->SetSubmixOutputVolume(InSoundSubmix, NewVolume);
-			}, GET_STATID(STAT_AudioSetSubmixOutputVolume));
+				MixerDevice->SetSubmixOutputVolume(InSoundSubmix, InOutputVolume);
+			});
 			return;
 		}
 
-		FMixerSubmixPtr MixerSubmix = GetSubmixInstance(InSoundSubmix).Pin();
-		if (MixerSubmix.IsValid())
+		FMixerSubmixPtr MixerSubmixPtr = GetSubmixInstance(InSoundSubmix).Pin();
+		if (MixerSubmixPtr.IsValid())
 		{
-			AudioRenderThreadCommand([MixerSubmix, NewVolume]()
+			AudioRenderThreadCommand([MixerSubmixPtr, InOutputVolume]()
+			{
+				MixerSubmixPtr->SetOutputVolume(InOutputVolume);
+			});
+		}
+	}
+
+	void FMixerDevice::SetSubmixWetLevel(USoundSubmix* InSoundSubmix, float InWetLevel)
+	{
+		if (!IsInAudioThread())
+		{
+			FMixerDevice* MixerDevice = this;
+			FAudioThread::RunCommandOnAudioThread([MixerDevice, InSoundSubmix, InWetLevel]()
+			{
+				MixerDevice->SetSubmixWetLevel(InSoundSubmix, InWetLevel);
+			});
+			return;
+		}
+
+		FMixerSubmixPtr MixerSubmixPtr = GetSubmixInstance(InSoundSubmix).Pin();
+		if (MixerSubmixPtr.IsValid())
+		{
+			AudioRenderThreadCommand([MixerSubmixPtr, InWetLevel]()
+			{
+				MixerSubmixPtr->SetWetLevel(InWetLevel);
+			});
+		}
+	}
+
+	void FMixerDevice::SetSubmixDryLevel(USoundSubmix* InSoundSubmix, float InDryLevel)
+	{
+		if (!IsInAudioThread())
+		{
+			FMixerDevice* MixerDevice = this;
+			FAudioThread::RunCommandOnAudioThread([MixerDevice, InSoundSubmix, InDryLevel]()
 			{
-				MixerSubmix->SetDynamicOutputVolume(NewVolume);
+				MixerDevice->SetSubmixDryLevel(InSoundSubmix, InDryLevel);
+			});
+			return;
+		}
+
+		FMixerSubmixPtr MixerSubmixPtr = GetSubmixInstance(InSoundSubmix).Pin();
+		if (MixerSubmixPtr.IsValid())
+		{
+			AudioRenderThreadCommand([MixerSubmixPtr, InDryLevel]()
+			{
+				MixerSubmixPtr->SetDryLevel(InDryLevel);
 			});
 		}
 	}
@@ -1124,14 +1189,14 @@ namespace Audio
 
 		FMixerSubmixPtr SubmixPtr = GetSubmixInstance(InSoundSubmix).Pin();
 
-		if (!bIsMasterSubmix)
+		if (bInit)
 		{
-			RebuildSubmixLinks(*InSoundSubmix, SubmixPtr);
+			InitSoundfieldAndEndpointDataForSubmix(*InSoundSubmix, SubmixPtr, true);
 		}
 
-		if (bInit)
+		if (!bIsMasterSubmix)
 		{
-			InitSoundfieldAndEndpointDataForSubmix(*InSoundSubmix, SubmixPtr, true);
+			RebuildSubmixLinks(*InSoundSubmix, SubmixPtr);
 		}
 	}
 
diff --git a/Engine/Source/Runtime/AudioMixer/Private/AudioMixerSubmix.cpp b/Engine/Source/Runtime/AudioMixer/Private/AudioMixerSubmix.cpp
--- a/Engine/Source/Runtime/AudioMixer/Private/AudioMixerSubmix.cpp
+++ b/Engine/Source/Runtime/AudioMixer/Private/AudioMixerSubmix.cpp
@@ -19,14 +19,6 @@ FAutoConsoleVariableRef CVarRecoverRecordingOnShutdown(
 	TEXT("0: Disabled, 1: Enabled"),
 	ECVF_Default);
 
-static int32 AllowInitOnAudioRenderThreadCVar = 0;
-FAutoConsoleVariableRef CVarAllowInitOnAudioRenderThread(
-	TEXT("au.AllowInitOnAudioRenderThread"),
-	RecoverRecordingOnShutdownCVar,
-	TEXT("When set to 1, we will initialize the submix on the audio render thread.\n")
-	TEXT("0: Disabled, 1: Enabled"),
-	ECVF_Default);
-
 namespace Audio
 {
 	// Unique IDs for mixer submixes
@@ -38,13 +30,15 @@ namespace Audio
 		, MixerDevice(InMixerDevice)
 		, NumChannels(0)
 		, NumSamples(0)
-		, InitializedOutputVolume(1.0f)
-		, OutputVolume(1.0f)
+		, CurrentOutputVolume(1.0f)
 		, TargetOutputVolume(1.0f)
+		, CurrentWetLevel(1.0f)
+		, TargetWetLevel(1.0f)
+		, CurrentDryLevel(0.0f)
+		, TargetDryLevel(0.0f)
 		, EnvelopeNumChannels(0)
 		, bIsRecording(false)
 		, bIsBackgroundMuted(false)
-		, bApplyOutputVolumeScale(false)
 		, OwningSubmixObject(nullptr)
 	{
 	}
@@ -80,17 +74,10 @@ namespace Audio
 			else if (bAllowReInit)
 			{
 				check(OwningSubmixObject == InSoundSubmix);
-				if (AllowInitOnAudioRenderThreadCVar)
-				{
-					SubmixCommand([this]()
-					{
-						InitInternal();
-					});
-				}
-				else
+				SubmixCommand([this]()
 				{
 					InitInternal();
-				}
+				});
 			}
 		}
 	}
@@ -103,13 +90,15 @@ namespace Audio
 
 		if (const USoundSubmix* SoundSubmix = Cast<const USoundSubmix>(OwningSubmixObject))
 		{
+			CurrentOutputVolume = FMath::Clamp(SoundSubmix->OutputVolume, 0.0f, 1.0f);
+			TargetOutputVolume = CurrentOutputVolume;
+
 			// Set the initialized output volume
-			InitializedOutputVolume = FMath::Clamp(SoundSubmix->OutputVolume, 0.0f, 1.0f);
+			CurrentWetLevel = FMath::Clamp(SoundSubmix->WetLevel, 0.0f, 1.0f);
+			TargetWetLevel = CurrentWetLevel;
 
-			if (!FMath::IsNearlyEqual(InitializedOutputVolume, 1.0f))
-			{
-				bApplyOutputVolumeScale = true;
-			}
+			CurrentDryLevel = FMath::Clamp(SoundSubmix->DryLevel, 0.0f, 1.0f);
+			TargetDryLevel = CurrentDryLevel;
 
 			for (USoundEffectSubmixPreset* EffectPreset : SoundSubmix->SubmixEffectChain)
 			{
@@ -118,10 +107,10 @@ namespace Audio
 					// Create a new effect instance using the preset
 					FSoundEffectSubmix* SubmixEffect = static_cast<FSoundEffectSubmix*>(EffectPreset->CreateNewEffect());
 				
-				FSoundEffectSubmixInitData InitData;
-				InitData.DeviceID = MixerDevice->DeviceID;
-				InitData.SampleRate = MixerDevice->GetSampleRate();
-				InitData.PresetSettings = nullptr;
+					FSoundEffectSubmixInitData InitData;
+					InitData.DeviceID = MixerDevice->DeviceID;
+					InitData.SampleRate = MixerDevice->GetSampleRate();
+					InitData.PresetSettings = nullptr;
 
 					// Now set the preset
 					SubmixEffect->Init(InitData);
@@ -725,7 +714,7 @@ namespace Audio
 			{
 				return;
 			}
-			
+
 			//Decode soundfield to interleaved float audio.
 			FSoundfieldDecoderInputData DecoderInput =
 			{
@@ -790,11 +779,32 @@ namespace Audio
 				const FMixerSourceVoice* MixerSourceVoice = MixerSourceVoiceIter.Key;
 				const float SendLevel = MixerSourceVoiceIter.Value.SendLevel;
 
-				
+
 				MixerSourceVoice->MixOutputBuffers(NumChannels, SendLevel, InputBuffer);
 			}
 		}
 
+		DryChannelBuffer.Reset();
+
+		// Check if we need to allocate a dry buffer
+		if (!FMath::IsNearlyEqual(CurrentDryLevel, TargetDryLevel) || !FMath::IsNearlyZero(CurrentDryLevel))
+		{
+			DryChannelBuffer.Append(InputBuffer);
+
+			// If we've already set the volume, only need to multiply by constant
+			if (FMath::IsNearlyEqual(TargetDryLevel, CurrentDryLevel))
+			{
+
+				Audio::MultiplyBufferByConstantInPlace(DryChannelBuffer, TargetDryLevel);
+			}
+			else
+			{
+				// To avoid popping, we do a fade on the buffer to the target volume
+				Audio::FadeBufferFast(DryChannelBuffer, CurrentDryLevel, TargetDryLevel);
+				CurrentDryLevel = TargetDryLevel;
+			}
+		}
+
 		if (EffectSubmixChain.Num() > 0)
 		{
 			CSV_SCOPED_TIMING_STAT(Audio, SubmixEffectProcessing);
@@ -858,14 +868,14 @@ namespace Audio
 					Audio::MixInBufferFast(InputBuffer, ScratchBuffer, DryLevel);
 				}
 
-				FMemory::Memcpy((void*)BufferPtr, (void*)ScratchBuffer.GetData(), sizeof(float)*NumSamples);
+				FMemory::Memcpy((void*)BufferPtr, (void*)ScratchBuffer.GetData(), sizeof(float) * NumSamples);
 			}
 		}
 
 		// If we're muted, memzero the buffer. Note we are still doing all the work to maintain buffer state between mutings.
 		if (bIsBackgroundMuted)
 		{
-			FMemory::Memzero((void*)BufferPtr, sizeof(float)*NumSamples);
+			FMemory::Memzero((void*)BufferPtr, sizeof(float) * NumSamples);
 		}
 
 		// If we are recording, Add out buffer to the RecordingData buffer:
@@ -899,7 +909,7 @@ namespace Audio
 
 			// Perform envelope following per channel
 			FScopeLock EnvelopeScopeLock(&EnvelopeCriticalSection);
-			FMemory::Memset(EnvelopeValues, sizeof(float)*AUDIO_MIXER_MAX_OUTPUT_CHANNELS);
+			FMemory::Memset(EnvelopeValues, sizeof(float) * AUDIO_MIXER_MAX_OUTPUT_CHANNELS);
 
 			for (int32 ChannelIndex = 0; ChannelIndex < NumChannels; ++ChannelIndex)
 			{
@@ -920,27 +930,40 @@ namespace Audio
 		}
 
 		// Don't necessarily need to do this if the user isn't using this feature
-		if (bApplyOutputVolumeScale)
+		if (!FMath::IsNearlyEqual(TargetWetLevel, CurrentWetLevel) || !FMath::IsNearlyEqual(CurrentWetLevel, 1.0f))
 		{
-			const float TargetVolumeProduct = TargetOutputVolume * InitializedOutputVolume;
-			const float OutputVolumeProduct = OutputVolume * InitializedOutputVolume;
-
 			// If we've already set the volume, only need to multiply by constant
-			if (FMath::IsNearlyEqual(TargetVolumeProduct, OutputVolumeProduct))
+			if (FMath::IsNearlyEqual(TargetWetLevel, CurrentWetLevel))
 			{
-				Audio::MultiplyBufferByConstantInPlace(OutAudioBuffer, OutputVolumeProduct);
+				Audio::MultiplyBufferByConstantInPlace(OutAudioBuffer, TargetWetLevel);
 			}
 			else
 			{
 				// To avoid popping, we do a fade on the buffer to the target volume
-				Audio::FadeBufferFast(OutAudioBuffer, OutputVolumeProduct, TargetVolumeProduct);
-				OutputVolume = TargetOutputVolume;
+				Audio::FadeBufferFast(OutAudioBuffer, CurrentWetLevel, TargetWetLevel);
+				CurrentWetLevel = TargetWetLevel;
+			}
+		}
 
-				// No longer need to multiply the output buffer if we're now at 1.0
-				if (FMath::IsNearlyEqual(OutputVolume * InitializedOutputVolume, 1.0f))
-				{
-					bApplyOutputVolumeScale = false;
-				}
+		// Check to see if need to mix together the dry and wet buffers
+		if (DryChannelBuffer.Num())
+		{
+			Audio::MixInBufferFast(DryChannelBuffer, OutAudioBuffer);
+		}
+
+		// Now apply the output volume
+		if (!FMath::IsNearlyEqual(TargetOutputVolume, CurrentOutputVolume) || !FMath::IsNearlyEqual(CurrentOutputVolume, 1.0f))
+		{
+			// If we've already set the output volume, only need to multiply by constant
+			if (FMath::IsNearlyEqual(TargetOutputVolume, CurrentOutputVolume))
+			{
+				Audio::MultiplyBufferByConstantInPlace(OutAudioBuffer, TargetOutputVolume);
+			}
+			else
+			{
+				// To avoid popping, we do a fade on the buffer to the target volume
+				Audio::FadeBufferFast(OutAudioBuffer, CurrentOutputVolume, TargetOutputVolume);
+				CurrentOutputVolume = TargetOutputVolume;
 			}
 		}
 
@@ -1460,24 +1483,19 @@ namespace Audio
 		}
 	}
 
-	void FMixerSubmix::SetDynamicOutputVolume(float InVolume)
+	void FMixerSubmix::SetOutputVolume(float InOutputVolume)
 	{
-		InVolume = FMath::Clamp(InVolume, 0.0f, 1.0f);
-		if (!FMath::IsNearlyEqual(InVolume, TargetOutputVolume))
-		{
-			TargetOutputVolume = InVolume;
-			bApplyOutputVolumeScale = true;
-		}
+		TargetOutputVolume = FMath::Clamp(InOutputVolume, 0.0f, 1.0f);
 	}
 
-	void FMixerSubmix::SetOutputVolume(float InVolume)
+	void FMixerSubmix::SetDryLevel(float InDryLevel)
 	{
-		InVolume = FMath::Clamp(InVolume, 0.0f, 1.0f);
-		if (!FMath::IsNearlyEqual(InitializedOutputVolume, InVolume))
-		{
-			InitializedOutputVolume = InVolume;
-			bApplyOutputVolumeScale = true;
-		}
+		TargetDryLevel = FMath::Clamp(InDryLevel, 0.0f, 1.0f);
+	}
+
+	void FMixerSubmix::SetWetLevel(float InWetLevel)
+	{
+		TargetWetLevel = FMath::Clamp(InWetLevel, 0.0f, 1.0f);
 	}
 
 	void FMixerSubmix::BroadcastEnvelope()
diff --git a/Engine/Source/Runtime/AudioMixer/Private/Effects/AudioMixerSubmixEffectReverbFast.cpp b/Engine/Source/Runtime/AudioMixer/Private/Effects/AudioMixerSubmixEffectReverbFast.cpp
--- a/Engine/Source/Runtime/AudioMixer/Private/Effects/AudioMixerSubmixEffectReverbFast.cpp
+++ b/Engine/Source/Runtime/AudioMixer/Private/Effects/AudioMixerSubmixEffectReverbFast.cpp
@@ -167,13 +167,13 @@ bool FSubmixEffectReverbFast::SetParameters(const FAudioEffectParameters& InPara
 	Audio::FPlateReverbFastSettings NewSettings;
 
 	// Early Reflections
-	NewSettings.EarlyReflections.Gain = FMath::GetMappedRangeValueClamped({ 0.0f, 3.16f }, { 0.0f, 1.0f }, ReverbEffectParams.ReflectionsGain);
+	NewSettings.EarlyReflections.Gain = FMath::GetMappedRangeValueClamped({ 0.0f, 3.16f }, { 0.0f, 1.0f }, ReverbEffectParams.ReflectionsGain * ReverbEffectParams.Volume);
 	NewSettings.EarlyReflections.PreDelayMsec = FMath::GetMappedRangeValueClamped({ 0.0f, 0.3f }, { 0.0f, 300.0f }, ReverbEffectParams.ReflectionsDelay);
 	NewSettings.EarlyReflections.Bandwidth = FMath::GetMappedRangeValueClamped({ 0.0f, 1.0f }, { 0.0f, 1.0f }, ReverbEffectParams.GainHF);
 
 	// LateReflections
 	NewSettings.LateReflections.LateDelayMsec = FMath::GetMappedRangeValueClamped({ 0.0f, 0.1f }, { 0.0f, 100.0f }, ReverbEffectParams.LateDelay);
-	NewSettings.LateReflections.LateGainDB = FMath::GetMappedRangeValueClamped({ 0.0f, 1.0f }, { 0.0f, 1.0f }, ReverbEffectParams.Gain);
+	NewSettings.LateReflections.LateGainDB = FMath::GetMappedRangeValueClamped({ 0.0f, 1.0f }, { 0.0f, 1.0f }, ReverbEffectParams.Gain * ReverbEffectParams.Volume);
 	NewSettings.LateReflections.Bandwidth = FMath::GetMappedRangeValueClamped({ 0.0f, 1.0f }, { 0.1f, 0.6f }, ReverbEffectParams.AirAbsorptionGainHF);
 	NewSettings.LateReflections.Diffusion = FMath::GetMappedRangeValueClamped({ 0.05f, 1.0f }, { 0.0f, 0.95f }, ReverbEffectParams.Diffusion);
 	NewSettings.LateReflections.Dampening = FMath::GetMappedRangeValueClamped({ 0.05f, 1.95f }, { 0.0f, 0.999f }, ReverbEffectParams.DecayHFRatio);
diff --git a/Engine/Source/Runtime/AudioMixer/Public/AudioMixerDevice.h b/Engine/Source/Runtime/AudioMixer/Public/AudioMixerDevice.h
--- a/Engine/Source/Runtime/AudioMixer/Public/AudioMixerDevice.h
+++ b/Engine/Source/Runtime/AudioMixer/Public/AudioMixerDevice.h
@@ -107,8 +107,11 @@ namespace Audio
 		// Updates submix instances with new properties
 		virtual void UpdateSubmixProperties(USoundSubmixBase* InSubmix) override;
 		
-		// Sets the submix output volume dynamically
-		virtual void SetSubmixOutputVolume(USoundSubmix* InSubmix, float NewVolume) override;
+		// Submix wet/dry settings
+		void SetSubmixWetDryLevel(USoundSubmix* InSoundSubmix, float InOutputVolume, float InWetLevel, float InDryLevel) override;
+		void SetSubmixOutputVolume(USoundSubmix* InSoundSubmix, float InOutputVolume) override;
+		void SetSubmixWetLevel(USoundSubmix* InSoundSubmix, float InWetLevel) override;
+		void SetSubmixDryLevel(USoundSubmix* InSoundSubmix, float InDryLevel) override;
 
 		// Submix recording callbacks:
 		virtual void StartRecording(USoundSubmix* InSubmix, float ExpectedRecordingDuration) override;
@@ -117,12 +120,12 @@ namespace Audio
 		virtual void PauseRecording(USoundSubmix* InSubmix);
 		virtual void ResumeRecording(USoundSubmix* InSubmix);
 
-		/** Submix envelope following */
+		// Submix envelope following
 		virtual void StartEnvelopeFollowing(USoundSubmix* InSubmix) override;
 		virtual void StopEnvelopeFollowing(USoundSubmix* InSubmix) override;
 		virtual void AddEnvelopeFollowerDelegate(USoundSubmix* InSubmix, const FOnSubmixEnvelopeBP& OnSubmixEnvelopeBP) override;
 
-		/** Submix Spectrum Analysis */
+		// Submix Spectrum Analysis
 		virtual void StartSpectrumAnalysis(USoundSubmix* InSubmix, const Audio::FSpectrumAnalyzerSettings& InSettings) override;
 		virtual void StopSpectrumAnalysis(USoundSubmix* InSubmix) override;
 		virtual void GetMagnitudesForFrequencies(USoundSubmix* InSubmix, const TArray<float>& InFrequencies, TArray<float>& OutMagnitudes);
diff --git a/Engine/Source/Runtime/AudioMixer/Public/AudioMixerSubmix.h b/Engine/Source/Runtime/AudioMixer/Public/AudioMixerSubmix.h
--- a/Engine/Source/Runtime/AudioMixer/Public/AudioMixerSubmix.h
+++ b/Engine/Source/Runtime/AudioMixer/Public/AudioMixerSubmix.h
@@ -10,7 +10,7 @@
 #include "DSP/EnvelopeFollower.h"
 #include "DSP/SpectrumAnalyzer.h"
 #include "Templates/SharedPointer.h"
-
+#include "AudioDynamicParameter.h"
 
 // Forward Declarations
 class FOnSubmixEnvelopeBP;
@@ -86,11 +86,14 @@ namespace Audio
 		// Removes the given submix from this submix's children
 		void RemoveChildSubmix(TWeakPtr<FMixerSubmix, ESPMode::ThreadSafe> SubmixWeakPtr);
 
+		// Sets the output level of the submix
+		void SetOutputVolume(float InOutputLevel);
+
 		// Sets the static output volume of the submix
-		void SetOutputVolume(float InVolume);
+		void SetDryLevel(float InDryLevel);
 
-		// Sets the dynamic output volume
-		void SetDynamicOutputVolume(float InVolume);
+		// Sets the wet level of the submix
+		void SetWetLevel(float InWetLevel);
 
 		// Gets the submix channels channels
 		int32 GetSubmixChannels() const;
@@ -421,15 +424,12 @@ namespace Audio
 
 		FEndpointData EndpointData;
 		
-
-		// The output volume of the submix set via the USoundSubmix property. Can be set in the editor.
-		float InitializedOutputVolume;
-		
-		// The current dynamic output volume
-		float OutputVolume;
-		
-		// The target dynamic output volume
+		float CurrentOutputVolume;
 		float TargetOutputVolume;
+		float CurrentWetLevel;
+		float TargetWetLevel;
+		float CurrentDryLevel;
+		float TargetDryLevel;
 
 		// Envelope following data
 		float EnvelopeValues[AUDIO_MIXER_MAX_OUTPUT_CHANNELS];
@@ -443,6 +443,9 @@ namespace Audio
 		// This buffer is used to downmix the submix output to mono before submitting it to the SpectrumAnalyzer.
 		AlignedFloatBuffer MonoMixBuffer;
 
+		// The dry channel buffer
+		AlignedFloatBuffer DryChannelBuffer;
+
 		// Submix command queue to shuffle commands from audio thread to audio render thread.
 		TQueue<TFunction<void()>> CommandQueue;
 
@@ -461,9 +464,6 @@ namespace Audio
 		// Whether or not this submix is muted.
 		uint8 bIsBackgroundMuted : 1;
 
-		// Whether or not to apply a volume scale to output
-		uint8 bApplyOutputVolumeScale : 1;
-
 		// Bool set to true when envelope following is enabled
 		FThreadSafeBool bIsEnvelopeFollowing;
 
diff --git a/Engine/Source/Runtime/Engine/Classes/Sound/SoundSubmix.h b/Engine/Source/Runtime/Engine/Classes/Sound/SoundSubmix.h
--- a/Engine/Source/Runtime/Engine/Classes/Sound/SoundSubmix.h
+++ b/Engine/Source/Runtime/Engine/Classes/Sound/SoundSubmix.h
@@ -127,6 +127,14 @@ protected:
 #endif 
 };
 
+// Whether to use linear or decibel values for audio gains
+UENUM(BlueprintType)
+enum class EGainParamMode : uint8
+{
+	Linear = 0,
+	Decibels,
+};
+
 /**
  * Sound Submix class meant for applying an effect to the downmixed sum of multiple audio sources.
  */
@@ -156,10 +164,36 @@ public:
 	UPROPERTY(EditAnywhere, BlueprintReadWrite, Category = EnvelopeFollower, meta = (ClampMin = "0", UIMin = "0"))
 	int32 EnvelopeFollowerReleaseTime;
 
-	/** The output volume of the submix. Applied after submix effects and analysis are performed. */
-	UPROPERTY(EditAnywhere, BlueprintReadWrite, Category = SoundSubmix, meta = (ClampMin = "0.0", ClampMax = "1.0", UIMin = "0.0", UIMax = "1.0"))
+	/** Whether to treat submix gain levels as linear or decibel values. */
+	UPROPERTY(EditAnywhere, Category = SubmixLevel, meta = (InlineCategoryProperty))
+	EGainParamMode GainMode;
+
+	/** The output volume of the submix. Applied after submix effects and analysis are performed.*/
+	UPROPERTY(EditAnywhere, BlueprintReadWrite, Category = SubmixLevel, meta = (ClampMin = "0.0", ClampMax = "1.0", UIMin = "0.0", UIMax = "1.0", EditCondition = "GainMode == EGainParamMode::Linear", DisplayName = "Output Volume", EditConditionHides))
 	float OutputVolume;
 
+	/** The wet level of the submix. Applied after submix effects and analysis are performed. */
+	UPROPERTY(EditAnywhere, BlueprintReadWrite, Category = SubmixLevel, meta = (ClampMin = "0.0", ClampMax = "1.0", UIMin = "0.0", UIMax = "1.0", EditCondition = "GainMode == EGainParamMode::Linear", DisplayName = "Wet Level", EditConditionHides))
+	float WetLevel;
+
+	/** The dry level of the submix. Applied before submix effects and analysis are performed. */
+	UPROPERTY(EditAnywhere, BlueprintReadWrite, Category = SubmixLevel, meta = (ClampMin = "0.0", ClampMax = "1.0", UIMin = "0.0", UIMax = "1.0", EditCondition = "GainMode == EGainParamMode::Linear", DisplayName = "Dry Level", EditConditionHides))
+	float DryLevel;
+
+#if WITH_EDITORONLY_DATA
+	/** The output volume of the submix (in dB). */
+	UPROPERTY(EditAnywhere, BlueprintReadWrite, Category = SubmixLevel, meta = (ClampMin = "-120.0", ClampMax = "0.0", UIMin = "-60.0", UIMax = "0.0", EditCondition = "GainMode == EGainParamMode::Decibels", DisplayName = "Output Volume (dB)", EditConditionHides))
+	float OutputVolumeDB;
+
+	/** The wet level of the submix  (in dB). Applied after submix effects and analysis are performed. */
+	UPROPERTY(EditAnywhere, BlueprintReadWrite, Category = SubmixLevel, meta = (ClampMin = "-120.0", ClampMax = "0.0", UIMin = "-60.0", UIMax = "0.0", EditCondition = "GainMode == EGainParamMode::Decibels", DisplayName = "Wet Level (dB)", EditConditionHides))
+	float WetLevelDB;
+
+	/** The dry level of the submix  (in dB)s. Applied before submix effects and analysis are performed. */
+	UPROPERTY(EditAnywhere, BlueprintReadWrite, Category = SubmixLevel, meta = (ClampMin = "-120.0", ClampMax = "0.0", UIMin = "-60.0", UIMax = "0.0", EditCondition = "GainMode == EGainParamMode::Decibels", DisplayName = "Dry Level (dB)", EditConditionHides))
+	float DryLevelDB;
+#endif
+
 	// Blueprint delegate for when a recorded file is finished exporting.
 	UPROPERTY(BlueprintAssignable)
 	FOnSubmixRecordedFileDone OnSubmixRecordedFileDone;
@@ -201,6 +235,7 @@ public:
 
 protected:
 
+	virtual void PostLoad() override;
 #if WITH_EDITOR
 	virtual void PostEditChangeProperty(struct FPropertyChangedEvent& PropertyChangedEvent) override;
 #endif
diff --git a/Engine/Source/Runtime/Engine/Private/SoundSubmix.cpp b/Engine/Source/Runtime/Engine/Private/SoundSubmix.cpp
--- a/Engine/Source/Runtime/Engine/Private/SoundSubmix.cpp
+++ b/Engine/Source/Runtime/Engine/Private/SoundSubmix.cpp
@@ -7,6 +7,7 @@
 #include "EngineGlobals.h"
 #include "Sound/SoundSubmixSend.h"
 #include "UObject/UObjectIterator.h"
+#include "DSP/Dsp.h"
 
 #if WITH_EDITOR
 #include "Framework/Notifications/NotificationManager.h"
@@ -38,10 +39,29 @@ USoundSubmix::USoundSubmix(const FObjectInitializer& ObjectInitializer)
 	, AmbisonicsPluginSettings(nullptr)
 	, EnvelopeFollowerAttackTime(10)
 	, EnvelopeFollowerReleaseTime(500)
+	, GainMode(EGainParamMode::Linear)
 	, OutputVolume(1.0f)
+	, WetLevel(1.0f)
+	, DryLevel(0.0f)
+#if WITH_EDITOR
+	, OutputVolumeDB(0.0f)
+	, WetLevelDB(0.0f)
+	, DryLevelDB(-120.0f)
+#endif
 {
 }
 
+void USoundSubmix::PostLoad()
+{
+	Super::PostLoad();
+
+#if WITH_EDITOR
+	OutputVolumeDB = Audio::ConvertToDecibels(OutputVolume);
+	WetLevelDB = Audio::ConvertToDecibels(WetLevel);
+	DryLevelDB = Audio::ConvertToDecibels(DryLevel);
+#endif
+
+}
 
 UEndpointSubmix::UEndpointSubmix(const FObjectInitializer& ObjectInitializer)
 	: Super(ObjectInitializer)
@@ -236,19 +256,90 @@ void USoundSubmix::SetSubmixOutputVolume(const UObject* WorldContextObject, floa
 		AudioDevice->SetSubmixOutputVolume(this, InOutputVolume);
 	}
 }
+
 #if WITH_EDITOR
 void USoundSubmix::PostEditChangeProperty(struct FPropertyChangedEvent& PropertyChangedEvent)
 {
 	if (PropertyChangedEvent.Property != nullptr)
 	{
 		static const FName NAME_OutputVolume(TEXT("OutputVolume"));
+		static const FName NAME_WetLevel(TEXT("WetLevel"));
+		static const FName NAME_DryLevel(TEXT("DryLevel"));
+		static const FName NAME_OutputVolumeDB(TEXT("OutputVolumeDB"));
+		static const FName NAME_WetLevelDB(TEXT("WetLevelDB"));
+		static const FName NAME_DryLevelDB(TEXT("DryLevelDB"));
 
-		if (PropertyChangedEvent.Property->GetFName() == NAME_OutputVolume)
+		FName ChangedPropName = PropertyChangedEvent.Property->GetFName();
+
+		bool bUpdateSubmixGain = false;
+
+		if (ChangedPropName == NAME_OutputVolume)
+		{
+			OutputVolumeDB = Audio::ConvertToDecibels(OutputVolume);
+			bUpdateSubmixGain = true;
+		}
+		else if (ChangedPropName == NAME_WetLevel)
+		{
+			DryLevelDB = Audio::ConvertToDecibels(DryLevel);
+			bUpdateSubmixGain = true;
+		}
+		else if (ChangedPropName == NAME_DryLevel)
+		{
+			WetLevelDB = Audio::ConvertToDecibels(OutputVolume);
+			bUpdateSubmixGain = true;
+		}
+		else if (ChangedPropName == NAME_OutputVolumeDB)
+		{
+			if (OutputVolumeDB <= -120.f)
+			{
+				OutputVolume = 0.0f;
+			}
+			else
+			{
+				OutputVolume = Audio::ConvertToLinear(OutputVolumeDB);
+			}
+			bUpdateSubmixGain = true;
+		}
+		else if (ChangedPropName == NAME_WetLevelDB)
 		{
-			FAudioDeviceManager* AudioDeviceManager = (GEngine ? GEngine->GetAudioDeviceManager() : nullptr);
-			if (AudioDeviceManager)
+			if (WetLevelDB <= -120.f)
 			{
-				AudioDeviceManager->UpdateSubmix(this);
+				WetLevel = 0.0f;
+			}
+			else
+			{
+				WetLevel = Audio::ConvertToLinear(WetLevelDB);
+			}
+			bUpdateSubmixGain = true;
+		}
+		else if (ChangedPropName == NAME_DryLevelDB)
+		{
+			if (DryLevelDB <= -120.0f)
+			{
+				DryLevel = 0.0f;
+			}
+			else
+			{
+				DryLevel = Audio::ConvertToLinear(DryLevelDB);
+			}
+			bUpdateSubmixGain = true;
+		}
+
+		// Force the properties to be initialized for this SoundSubmix on all active audio devices
+		if (FAudioDeviceManager* AudioDeviceManager = GEngine->GetAudioDeviceManager())
+		{
+			if (bUpdateSubmixGain)
+			{
+//				AudioDeviceManager->UpdateSubmix(this);
+
+				const float NewOutputVolume = OutputVolume;
+				const float NewWetLevel = WetLevel;
+				const float NewDryLevel = DryLevel;
+				USoundSubmix* SoundSubmix = this;
+				AudioDeviceManager->IterateOverAllDevices([SoundSubmix, NewOutputVolume, NewWetLevel, NewDryLevel](Audio::FDeviceId Id, FAudioDevice* Device)
+				{
+					Device->SetSubmixWetDryLevel(SoundSubmix, NewOutputVolume, NewWetLevel, NewDryLevel);
+				});
 			}
 		}
 	}
@@ -340,14 +431,14 @@ void USoundSubmixBase::PreEditChange(FProperty* PropertyAboutToChange)
 
 void USoundSubmixBase::PostEditChangeProperty(FPropertyChangedEvent& PropertyChangedEvent)
 {
-	// Whether or not we need to reinit the submix. Not all properties require reinitialization.
-	bool bReinitSubmix = true;
+	if (!GEngine)
+	{
+		return;
+	}
 
 	if (PropertyChangedEvent.Property != nullptr)
 	{
 		static const FName NAME_ChildSubmixes(TEXT("ChildSubmixes"));
-		static const FName NAME_ParentSubmix(TEXT("ParentSubmix"));
-		static const FName NAME_OutputVolume(TEXT("OutputVolume"));
 
 		if (PropertyChangedEvent.Property->GetFName() == NAME_ChildSubmixes)
 		{
@@ -388,15 +479,12 @@ void USoundSubmixBase::PostEditChangeProperty(FPropertyChangedEvent& PropertyCha
 					}
 				}
 			}
-				}
-			}
 
-	if (GEngine)
-	{
-		// Force the properties to be initialized for this SoundSubmix on all active audio devices
-		if (FAudioDeviceManager* AudioDeviceManager = GEngine->GetAudioDeviceManager())
-		{
-			AudioDeviceManager->RegisterSoundSubmix(this);
+			// Force the properties to be initialized for this SoundSubmix on all active audio devices
+			if (FAudioDeviceManager* AudioDeviceManager = GEngine->GetAudioDeviceManager())
+			{
+				AudioDeviceManager->RegisterSoundSubmix(this);
+			}
 		}
 	}
 
@@ -444,13 +532,21 @@ void USoundSubmixWithParentBase::SetParentSubmix(USoundSubmixBase* InParentSubmi
 		}
 }
 
+#if WITH_EDITOR
 void USoundSubmixWithParentBase::PostEditChangeProperty(struct FPropertyChangedEvent& PropertyChangedEvent)
 {
+	if (!GEngine)
+	{
+		return;
+	}
+
 	if (PropertyChangedEvent.Property != nullptr)
 	{
 		static const FName NAME_ParentSubmix(TEXT("ParentSubmix"));
 
-		if (PropertyChangedEvent.Property->GetFName() == NAME_ParentSubmix)
+		FName ChangedPropName = PropertyChangedEvent.Property->GetFName();
+
+		if (ChangedPropName == NAME_ParentSubmix)
 		{
 			// Add this sound class to the parent class if it's not already added
 			if (ParentSubmix)
@@ -474,15 +570,12 @@ void USoundSubmixWithParentBase::PostEditChangeProperty(struct FPropertyChangedE
 			}
 
 			Modify();
-		}
-	}
 
-	if (GEngine)
-	{
-		// Force the properties to be initialized for this SoundSubmix on all active audio devices
-		if (FAudioDeviceManager* AudioDeviceManager = GEngine->GetAudioDeviceManager())
-		{
-			AudioDeviceManager->RegisterSoundSubmix(this);
+			// Force the properties to be initialized for this SoundSubmix on all active audio devices
+			if (FAudioDeviceManager* AudioDeviceManager = GEngine->GetAudioDeviceManager())
+			{
+				AudioDeviceManager->RegisterSoundSubmix(this);
+			}
 		}
 	}
 
@@ -498,6 +591,7 @@ void USoundSubmixWithParentBase::PostDuplicate(EDuplicateMode::Type DuplicateMod
 
 	Super::PostDuplicate(DuplicateMode);
 }
+#endif
 
 void USoundSubmixBase::AddReferencedObjects(UObject* InThis, FReferenceCollector& Collector)
 {
diff --git a/Engine/Source/Runtime/Engine/Public/AudioDevice.h b/Engine/Source/Runtime/Engine/Public/AudioDevice.h
--- a/Engine/Source/Runtime/Engine/Public/AudioDevice.h
+++ b/Engine/Source/Runtime/Engine/Public/AudioDevice.h
@@ -1244,12 +1244,6 @@ public:
 		UE_LOG(LogAudio, Error, TEXT("Submixes are only supported in audio mixer."));
 	}
 
-	/** Sets the dynamic volume of the given submix. */
-	virtual void SetSubmixOutputVolume(USoundSubmix* InSubmix, float NewVolume)
-	{
-		UE_LOG(LogAudio, Error, TEXT("Submixes are only supported in audio mixer."));
-	}
-
 	/** This is called by a USoundSubmix to start recording a submix instance on this device. */
 	virtual void StartRecording(USoundSubmix* InSubmix, float ExpectedRecordingDuration)
 	{
@@ -1277,6 +1271,30 @@ public:
 		UE_LOG(LogAudio, Error, TEXT("Envelope following submixes only works with the audio mixer. Please run using -audiomixer or set INI file to use submix recording."));
 	}
 
+	/** Set the wet-dry level of the given submix */
+	virtual void SetSubmixWetDryLevel(USoundSubmix* InSoundSubmix, float InOutputVolume, float InWetLevel, float InDryLevel)
+	{
+		UE_LOG(LogAudio, Error, TEXT("Submixes are only supported in audio mixer."));
+	}
+
+	/** Set the wet-dry level of the given submix */
+	virtual void SetSubmixOutputVolume(USoundSubmix* InSoundSubmix, float InOutputVolume)
+	{
+		UE_LOG(LogAudio, Error, TEXT("Submixes are only supported in audio mixer."));
+	}
+
+	/** Set the wet-dry level of the given submix */
+	virtual void SetSubmixWetLevel(USoundSubmix* InSoundSubmix, float InWetLevel)
+	{
+		UE_LOG(LogAudio, Error, TEXT("Submixes are only supported in audio mixer."));
+	}
+
+	/** Set the wet-dry level of the given submix */
+	virtual void SetSubmixDryLevel(USoundSubmix* InSoundSubmix, float InDryLevel)
+	{
+		UE_LOG(LogAudio, Error, TEXT("Submixes are only supported in audio mixer."));
+	}
+
 	/** Adds an envelope follower delegate to the submix for this audio device. */
 	virtual void AddEnvelopeFollowerDelegate(USoundSubmix* InSubmix, const FOnSubmixEnvelopeBP& OnSubmixEnvelopeBP);
 	virtual void StartSpectrumAnalysis(USoundSubmix* InSubmix, const Audio::FSpectrumAnalyzerSettings& InSettings);
