Fixing submix issue with muting audio output in unrelated submixes

- Fixed order of operations issue with submix processing.

diff --git a/Engine/Source/Runtime/AudioMixer/Private/AudioMixerSubmix.cpp b/Engine/Source/Runtime/AudioMixer/Private/AudioMixerSubmix.cpp
--- a/Engine/Source/Runtime/AudioMixer/Private/AudioMixerSubmix.cpp
+++ b/Engine/Source/Runtime/AudioMixer/Private/AudioMixerSubmix.cpp
@@ -878,6 +878,29 @@ namespace Audio
 			FMemory::Memzero((void*)BufferPtr, sizeof(float) * NumSamples);
 		}
 
+
+		// Don't necessarily need to do this if the user isn't using this feature
+		if (!FMath::IsNearlyEqual(TargetWetLevel, CurrentWetLevel) || !FMath::IsNearlyEqual(CurrentWetLevel, 1.0f))
+		{
+			// If we've already set the volume, only need to multiply by constant
+			if (FMath::IsNearlyEqual(TargetWetLevel, CurrentWetLevel))
+			{
+				Audio::MultiplyBufferByConstantInPlace(InputBuffer, TargetWetLevel);
+			}
+			else
+			{
+				// To avoid popping, we do a fade on the buffer to the target volume
+				Audio::FadeBufferFast(InputBuffer, CurrentWetLevel, TargetWetLevel);
+				CurrentWetLevel = TargetWetLevel;
+			}
+		}
+
+		// Check to see if need to mix together the dry and wet buffers
+		if (DryChannelBuffer.Num())
+		{
+			Audio::MixInBufferFast(DryChannelBuffer, InputBuffer);
+		}
+		
 		// If we are recording, Add out buffer to the RecordingData buffer:
 		{
 			FScopeLock ScopedLock(&RecordingCriticalSection);
@@ -897,9 +920,6 @@ namespace Audio
 			SpectrumAnalyzer->PerformAnalysisIfPossible(true, true);
 		}
 
-		// If the channel types match, just do a copy
-		Audio::MixInBufferFast(InputBuffer, OutAudioBuffer);
-
 		// Perform any envelope following if we're told to do so
 		if (bIsEnvelopeFollowing)
 		{
@@ -929,44 +949,25 @@ namespace Audio
 			EnvelopeNumChannels = NumChannels;
 		}
 
-		// Don't necessarily need to do this if the user isn't using this feature
-		if (!FMath::IsNearlyEqual(TargetWetLevel, CurrentWetLevel) || !FMath::IsNearlyEqual(CurrentWetLevel, 1.0f))
-		{
-			// If we've already set the volume, only need to multiply by constant
-			if (FMath::IsNearlyEqual(TargetWetLevel, CurrentWetLevel))
-			{
-				Audio::MultiplyBufferByConstantInPlace(OutAudioBuffer, TargetWetLevel);
-			}
-			else
-			{
-				// To avoid popping, we do a fade on the buffer to the target volume
-				Audio::FadeBufferFast(OutAudioBuffer, CurrentWetLevel, TargetWetLevel);
-				CurrentWetLevel = TargetWetLevel;
-			}
-		}
-
-		// Check to see if need to mix together the dry and wet buffers
-		if (DryChannelBuffer.Num())
-		{
-			Audio::MixInBufferFast(DryChannelBuffer, OutAudioBuffer);
-		}
-
 		// Now apply the output volume
 		if (!FMath::IsNearlyEqual(TargetOutputVolume, CurrentOutputVolume) || !FMath::IsNearlyEqual(CurrentOutputVolume, 1.0f))
 		{
 			// If we've already set the output volume, only need to multiply by constant
 			if (FMath::IsNearlyEqual(TargetOutputVolume, CurrentOutputVolume))
 			{
-				Audio::MultiplyBufferByConstantInPlace(OutAudioBuffer, TargetOutputVolume);
+				Audio::MultiplyBufferByConstantInPlace(InputBuffer, TargetOutputVolume);
 			}
 			else
 			{
 				// To avoid popping, we do a fade on the buffer to the target volume
-				Audio::FadeBufferFast(OutAudioBuffer, CurrentOutputVolume, TargetOutputVolume);
+				Audio::FadeBufferFast(InputBuffer, CurrentOutputVolume, TargetOutputVolume);
 				CurrentOutputVolume = TargetOutputVolume;
 			}
 		}
 
+		// Mix the audio buffer of this submix with the audio buffer of the output buffer (i.e. with other submixes)
+		Audio::MixInBufferFast(InputBuffer, OutAudioBuffer);
+
 		// Now loop through any buffer listeners and feed the listeners the result of this audio callback
 		if(const USoundSubmix* SoundSubmix = Cast<const USoundSubmix>(OwningSubmixObject))
 		{
