Audio bus feature

- Audio bus is a new asset type that allows audio to be routed around to effects, etc.
- Extracts the underlying logic from source buses and changes source bus semantics to be a sonification of audio buses. Audio buses do not render to be audible on their own.
- Main use-case is side-chaining audio effects (e.g. side chain compression, auto-wah filters, etc).

diff --git a/Engine/Plugins/Runtime/Synthesis/Source/Synthesis/Classes/SourceEffects/SourceEffectFilter.h b/Engine/Plugins/Runtime/Synthesis/Source/Synthesis/Classes/SourceEffects/SourceEffectFilter.h
--- a/Engine/Plugins/Runtime/Synthesis/Source/Synthesis/Classes/SourceEffects/SourceEffectFilter.h
+++ b/Engine/Plugins/Runtime/Synthesis/Source/Synthesis/Classes/SourceEffects/SourceEffectFilter.h
@@ -4,8 +4,12 @@
 
 #include "Sound/SoundEffectSource.h"
 #include "DSP/Filter.h"
+#include "DSP/MultithreadedPatching.h"
+#include "DSP/EnvelopeFollower.h"
 #include "SourceEffectFilter.generated.h"
 
+class UAudioBus;
+
 UENUM(BlueprintType)
 enum class ESourceEffectFilterCircuit : uint8
 {
@@ -25,27 +29,81 @@ enum class ESourceEffectFilterType : uint8
 	Count UMETA(Hidden)
 };
 
+UENUM(BlueprintType)
+enum class ESourceEffectFilterParam : uint8
+{
+	FilterFrequency = 0,
+	FilterResonance,
+	Count UMETA(Hidden)
+};
+
+USTRUCT(BlueprintType)
+struct SYNTHESIS_API FSourceEffectFilterAudioBusModulationSettings
+{
+	GENERATED_USTRUCT_BODY()
+
+	// Audio bus to use to modulate the filter
+	UPROPERTY(EditAnywhere, BlueprintReadWrite, Category = "SourceEffect|Preset|AudioBus")
+	UAudioBus* AudioBus = nullptr;
+
+	// The amplitude envelope follower attack time (in milliseconds) on the audio bus.
+	UPROPERTY(EditAnywhere, BlueprintReadWrite, Category = "SourceEffect|Modulation Envelope", meta = (ClampMin = "0", UIMin = "0", UIMax = "2000.0", DisplayName = "Attack Time (ms)"))
+	int32 EnvelopeFollowerAttackTimeMsec = 10;
+
+	// The amplitude envelope follower release time (in milliseconds) on the audio bus.
+	UPROPERTY(EditAnywhere, BlueprintReadWrite, Category = "SourceEffect|Modulation Envelope", meta = (ClampMin = "0", UIMin = "0", UIMax = "2000.0", DisplayName = "Release Time (ms)"))
+	int32 EnvelopeFollowerReleaseTimeMsec = 100;
+
+	// An amount to scale the envelope follower output to map to the modulation values.  
+	UPROPERTY(EditAnywhere, BlueprintReadWrite, Category = "SourceEffect|Modulation Envelope", meta = (ClampMin = "0"))
+	float EnvelopeGainMultiplier = 1.0f;
+
+	// Which parameter to modulate.
+	UPROPERTY(EditAnywhere, Category = "SourceEffect|Filter Modulation Params", meta = (InlineCategoryProperty))
+	ESourceEffectFilterParam FilterParam = ESourceEffectFilterParam::FilterFrequency;
+
+	// The frequency modulation value (in semitones from the filter frequency) to use when the envelope is quietest
+	UPROPERTY(EditAnywhere, BlueprintReadWrite, Category = "SourceEffect|Filter Modulation Params", meta = (EditCondition = "FilterParam == ESourceEffectFilterParam::FilterFrequency", DisplayName = "Min Frequency Modulation Amount", EditConditionHides))
+	float MinFrequencyModulation = -12.0f;
+
+	// The frequency modulation value (in semitones from the filter frequency) to use when the envelope is loudest
+	UPROPERTY(EditAnywhere, BlueprintReadWrite, Category = "SourceEffect|Filter Modulation Params", meta = (EditCondition = "FilterParam == ESourceEffectFilterParam::FilterFrequency", DisplayName = "Max Frequency Modulation Amount", EditConditionHides))
+	float MaxFrequencyModulation = 12.0f;
+
+	// The resonance modulation value to use when the envelope is quietest
+	UPROPERTY(EditAnywhere, BlueprintReadWrite, Category = "SourceEffect|Filter Modulation Params", meta = (ClampMin = "0.0", ClampMax = "10.0", UIMin = "0", UIMax = "10.0", EditCondition = "FilterParam == ESourceEffectFilterParam::FilterResonance", DisplayName = "Min Resonance Modulation Amount", EditConditionHides))
+	float MinResonanceModulation = 0.2f;
+
+	// The resonance modulation value to use when the envelope is loudest
+	UPROPERTY(EditAnywhere, BlueprintReadWrite, Category = "SourceEffect|Filter Modulation Params", meta = (ClampMin = "0.0", ClampMax = "10.0", UIMin = "0", UIMax = "10.0", EditCondition = "FilterParam == ESourceEffectFilterParam::FilterResonance", DisplayName = "Max Resonance Modulation Amount", EditConditionHides))
+	float MaxResonanceModulation = 3.0f;
+};
+
 USTRUCT(BlueprintType)
 struct SYNTHESIS_API FSourceEffectFilterSettings
 {
 	GENERATED_USTRUCT_BODY()
 
 	// The type of filter circuit to use.
-	UPROPERTY(EditAnywhere, BlueprintReadWrite, Category = "SourceEffect|Preset")
+	UPROPERTY(EditAnywhere, BlueprintReadWrite, Category = "SourceEffect")
 	ESourceEffectFilterCircuit FilterCircuit;
 
 	// The type of filter to use.
-	UPROPERTY(EditAnywhere, BlueprintReadWrite, Category = "SourceEffect|Preset")
+	UPROPERTY(EditAnywhere, BlueprintReadWrite, Category = "SourceEffect")
 	ESourceEffectFilterType FilterType;
 
 	// The filter cutoff frequency
-	UPROPERTY(EditAnywhere, BlueprintReadWrite, Category = "SourceEffect|Preset", meta = (ClampMin = "20.0", UIMin = "20.0", UIMax = "12000.0"))
+	UPROPERTY(EditAnywhere, BlueprintReadWrite, Category = "SourceEffect", meta = (ClampMin = "20.0", UIMin = "20.0", UIMax = "12000.0"))
 	float CutoffFrequency;
 
 	// The filter resonance.
-	UPROPERTY(EditAnywhere, BlueprintReadWrite, Category = "SourceEffect|Preset", meta = (ClampMin = "0.5", ClampMax = "10.0", UIMin = "0.5", UIMax = "10.0"))
+	UPROPERTY(EditAnywhere, BlueprintReadWrite, Category = "SourceEffect", meta = (ClampMin = "0.5", ClampMax = "10.0", UIMin = "0.5", UIMax = "10.0"))
 	float FilterQ;
 
+	// Audio bus settings to use to modulate the filter frequency cutoff (auto-wah) with arbitrary audio from an audio bus
+	UPROPERTY(EditAnywhere, BlueprintReadWrite, Category = "SourceEffect")
+	TArray<FSourceEffectFilterAudioBusModulationSettings> AudioBusModulation;
+
 	FSourceEffectFilterSettings()
 		: FilterCircuit(ESourceEffectFilterCircuit::StateVariable)
 		, FilterType(ESourceEffectFilterType::LowPass)
@@ -59,6 +117,7 @@ class SYNTHESIS_API FSourceEffectFilter : public FSoundEffectSource
 {
 public:
 	FSourceEffectFilter();
+	virtual ~FSourceEffectFilter();
 
 	// Called on an audio effect at initialization on main thread before audio processing begins.
 	virtual void Init(const FSoundEffectSourceInitData& InitData) override;
@@ -77,17 +136,38 @@ protected:
 	Audio::FOnePoleFilter OnePoleFilter;
 	Audio::IFilter* CurrentFilter;
 
+	float SampleRate;
 	float CutoffFrequency;
+	float BaseCutoffFrequency;
 	float FilterQ;
+	float BaseFilterQ;
 	ESourceEffectFilterCircuit CircuitType;
 	ESourceEffectFilterType FilterType;
 
+	Audio::AlignedFloatBuffer ScratchModBuffer;
+	Audio::AlignedFloatBuffer ScratchEnvFollowerBuffer;
+
+	struct FAudioBusModulationData
+	{
+		Audio::FPatchOutputStrongPtr AudioBusPatch;
+		Audio::FEnvelopeFollower AudioBusEnvelopeFollower;
+
+		ESourceEffectFilterParam FilterParam;
+		float MinFreqModValue = 0.0f;
+		float MaxFreqModValue = 0.0f;
+		float MinResModValue = 0.0f;
+		float MaxResModValue = 0.0f;
+		float EnvelopeGain = 0.0f;
+	};
+	TArray<FAudioBusModulationData> ModData;
+
 	float AudioInput[2];
 	float AudioOutput[2];
 	int32 NumChannels;
+	uint32 AudioDeviceId;
 };
 
-UCLASS(ClassGroup = AudioSourceEffect, meta = (BlueprintSpawnableComponent))
+UCLASS(ClassGroup = AudioSourceEffect, AutoExpandCategories = (SourceEffect), meta = (BlueprintSpawnableComponent))
 class SYNTHESIS_API USourceEffectFilterPreset : public USoundEffectSourcePreset
 {
 	GENERATED_BODY()
diff --git a/Engine/Plugins/Runtime/Synthesis/Source/Synthesis/Classes/SourceEffects/SourceEffectRingModulation.h b/Engine/Plugins/Runtime/Synthesis/Source/Synthesis/Classes/SourceEffects/SourceEffectRingModulation.h
--- a/Engine/Plugins/Runtime/Synthesis/Source/Synthesis/Classes/SourceEffects/SourceEffectRingModulation.h
+++ b/Engine/Plugins/Runtime/Synthesis/Source/Synthesis/Classes/SourceEffects/SourceEffectRingModulation.h
@@ -6,6 +6,8 @@
 #include "DSP/RingModulation.h"
 #include "SourceEffectRingModulation.generated.h"
 
+class UAudioBus;
+
 UENUM(BlueprintType)
 enum class ERingModulatorTypeSourceEffect : uint8
 {
@@ -41,12 +43,17 @@ struct SYNTHESIS_API FSourceEffectRingModulationSettings
 	UPROPERTY(EditAnywhere, BlueprintReadWrite, Category = "SourceEffect|Preset", meta = (ClampMin = "0.0", ClampMax = "1.0", UIMin = "0.0", UIMax = "1.0"))
 	float WetLevel;
 
+	// Audio bus to use to modulate the effect
+	UPROPERTY(EditAnywhere, BlueprintReadWrite, Category = "SourceEffect|Preset")
+	UAudioBus* AudioBusModulator;
+
 	FSourceEffectRingModulationSettings()
 		: ModulatorType(ERingModulatorTypeSourceEffect::Sine)
 		, Frequency(100.0f)
 		, Depth(0.5f)
 		, DryLevel(0.0f)
 		, WetLevel(1.0f)
+		, AudioBusModulator(nullptr)
 	{}
 };
 
@@ -64,6 +71,8 @@ public:
 
 protected:
 	Audio::FRingModulation RingModulation;
+	uint32 AudioDeviceId;
+	int32 NumChannels = 0;
 };
 
 UCLASS(ClassGroup = AudioSourceEffect, meta = (BlueprintSpawnableComponent))
diff --git a/Engine/Plugins/Runtime/Synthesis/Source/Synthesis/Private/SourceEffectFilter.cpp b/Engine/Plugins/Runtime/Synthesis/Source/Synthesis/Private/SourceEffectFilter.cpp
--- a/Engine/Plugins/Runtime/Synthesis/Source/Synthesis/Private/SourceEffectFilter.cpp
+++ b/Engine/Plugins/Runtime/Synthesis/Source/Synthesis/Private/SourceEffectFilter.cpp
@@ -1,12 +1,19 @@
 // Copyright Epic Games, Inc. All Rights Reserved.
 
 #include "SourceEffects/SourceEffectFilter.h"
+#include "DSP/BufferVectorOperations.h"
+#include "DSP/Dsp.h"
+#include "AudioMixerDevice.h"
 
+PRAGMA_DISABLE_OPTIMIZATION
 
 FSourceEffectFilter::FSourceEffectFilter()
 	: CurrentFilter(nullptr)
+	, SampleRate(0.0f)
 	, CutoffFrequency(8000.0f)
+	, BaseCutoffFrequency(8000.0f)
 	, FilterQ(2.0f)
+	, BaseFilterQ(2.0f)
 	, CircuitType(ESourceEffectFilterCircuit::StateVariable)
 	, FilterType(ESourceEffectFilterType::LowPass)
 {
@@ -14,6 +21,11 @@ FSourceEffectFilter::FSourceEffectFilter()
 	FMemory::Memzero(AudioOutput, 2 * sizeof(float));
 }
 
+FSourceEffectFilter::~FSourceEffectFilter()
+{
+
+}
+
 void FSourceEffectFilter::Init(const FSoundEffectSourceInitData& InitData)
 {
 	bIsActive = true;
@@ -22,6 +34,9 @@ void FSourceEffectFilter::Init(const FSoundEffectSourceInitData& InitData)
 	LadderFilter.Init(InitData.SampleRate, NumChannels);
 	OnePoleFilter.Init(InitData.SampleRate, NumChannels);
 
+	SampleRate = InitData.SampleRate;
+	AudioDeviceId = InitData.AudioDeviceId;
+
 	UpdateFilter();
 }
 
@@ -81,17 +96,101 @@ void FSourceEffectFilter::OnPresetChanged()
 	CircuitType = Settings.FilterCircuit;
 	FilterType = Settings.FilterType;
 	CutoffFrequency = Settings.CutoffFrequency;
+	BaseCutoffFrequency = CutoffFrequency;
 	FilterQ = Settings.FilterQ;
 
+	ModData.Reset();
+	
+	FAudioDeviceManager* AudioDeviceManager = FAudioDeviceManager::Get();
+	if (AudioDeviceManager && Settings.AudioBusModulation.Num() > 0)
+	{
+		Audio::FMixerDevice* MixerDevice = (Audio::FMixerDevice*)AudioDeviceManager->GetAudioDeviceRaw(AudioDeviceId);
+		if (MixerDevice)
+		{
+			for (FSourceEffectFilterAudioBusModulationSettings& BusModulationSettings : Settings.AudioBusModulation)
+			{
+				if (BusModulationSettings.AudioBus)
+				{
+					FAudioBusModulationData& NewModData = ModData.AddDefaulted_GetRef();
+
+					uint32 AudioBusId = BusModulationSettings.AudioBus->GetUniqueID();
+					NewModData.AudioBusPatch = MixerDevice->AddPatchForAudioBus(AudioBusId, 1.0f);
+
+					NewModData.MinFreqModValue = BusModulationSettings.MinFrequencyModulation;
+					NewModData.MaxFreqModValue = BusModulationSettings.MaxFrequencyModulation;
+					NewModData.MinResModValue = BusModulationSettings.MinResonanceModulation;
+					NewModData.MaxResModValue = BusModulationSettings.MaxResonanceModulation;
+					NewModData.EnvelopeGain = BusModulationSettings.EnvelopeGainMultiplier;
+
+					NewModData.AudioBusEnvelopeFollower.Init(SampleRate);
+					NewModData.AudioBusEnvelopeFollower.SetAttackTime(BusModulationSettings.EnvelopeFollowerAttackTimeMsec);
+					NewModData.AudioBusEnvelopeFollower.SetReleaseTime(BusModulationSettings.EnvelopeFollowerReleaseTimeMsec);
+				}
+			}
+		}
+	}
+
 	UpdateFilter();
 }
 
 void FSourceEffectFilter::ProcessAudio(const FSoundEffectSourceInputData& InData, float* OutAudioBufferData)
 {
+	float ModFrequency = 0.0f;
+	float ModQ = 0.0f;
+
+	for (FAudioBusModulationData& Mod : ModData)
+	{
+		if (!Mod.AudioBusPatch)
+		{
+			continue;
+		}
+
+		ScratchModBuffer.Reset();
+		ScratchModBuffer.AddZeroed(InData.NumSamples * NumChannels);
+
+		Mod.AudioBusPatch->PopAudio(ScratchModBuffer.GetData(), ScratchModBuffer.Num(), true);
+
+		float EnvelopeValue = 0.0f;
+
+		// If our modulation buffer (and this source effect) are both 2 channels, we need to downmix to mono before envelope following the signal
+		if (NumChannels == 2)
+		{
+			ScratchEnvFollowerBuffer.Reset();
+			ScratchEnvFollowerBuffer.AddUninitialized(InData.NumSamples);
+
+			Audio::BufferSum2ChannelToMonoFast(ScratchModBuffer, ScratchEnvFollowerBuffer);
+			Audio::MultiplyBufferByConstantInPlace(ScratchEnvFollowerBuffer, 0.5f);
+
+			EnvelopeValue = Mod.AudioBusEnvelopeFollower.ProcessAudio(ScratchEnvFollowerBuffer.GetData(), InData.NumSamples);
+		}
+		else
+		{
+			EnvelopeValue = Mod.AudioBusEnvelopeFollower.ProcessAudio(ScratchModBuffer.GetData(), InData.NumSamples);
+		}
+
+		if (Mod.FilterParam == ESourceEffectFilterParam::FilterFrequency)
+		{
+			ModFrequency += FMath::Lerp(Mod.MinFreqModValue, Mod.MaxFreqModValue, FMath::Clamp(EnvelopeValue * Mod.EnvelopeGain, 0.0f, 1.0f));;
+		}
+		else
+		{
+			ModQ += FMath::Lerp(Mod.MinResModValue, Mod.MaxResModValue, FMath::Clamp(EnvelopeValue * Mod.EnvelopeGain, 0.0f, 1.0f));;
+		}
+
+		CutoffFrequency = BaseCutoffFrequency * Audio::GetFrequencyMultiplier(ModFrequency);
+		FilterQ = FMath::Clamp(BaseFilterQ + ModQ, 0.1f, 10.0f);
+
+
+		// Update the filter parameters
+		UpdateFilter();
+	}
+
 	CurrentFilter->ProcessAudio(InData.InputSourceEffectBufferPtr, InData.NumSamples, OutAudioBufferData);
 }
 
 void USourceEffectFilterPreset::SetSettings(const FSourceEffectFilterSettings& InSettings)
 {
 	UpdateSettings(InSettings);
-}
\ No newline at end of file
+}
+
+PRAGMA_ENABLE_OPTIMIZATION
diff --git a/Engine/Plugins/Runtime/Synthesis/Source/Synthesis/Private/SourceEffectRingModulation.cpp b/Engine/Plugins/Runtime/Synthesis/Source/Synthesis/Private/SourceEffectRingModulation.cpp
--- a/Engine/Plugins/Runtime/Synthesis/Source/Synthesis/Private/SourceEffectRingModulation.cpp
+++ b/Engine/Plugins/Runtime/Synthesis/Source/Synthesis/Private/SourceEffectRingModulation.cpp
@@ -1,11 +1,17 @@
 // Copyright Epic Games, Inc. All Rights Reserved.
 
 #include "SourceEffects/SourceEffectRingModulation.h"
+#include "AudioDevice.h"
+#include "AudioMixerDevice.h"
+#include "AudioMixerBus.h"
+#include "DSP/MultithreadedPatching.h"
 
 void FSourceEffectRingModulation::Init(const FSoundEffectSourceInitData& InitData)
 {
 	bIsActive = true;
-	RingModulation.Init(InitData.SampleRate, InitData.NumSourceChannels);
+	NumChannels = InitData.NumSourceChannels;
+	RingModulation.Init(InitData.SampleRate, NumChannels);
+	AudioDeviceId = InitData.AudioDeviceId;
 }
 
 void FSourceEffectRingModulation::OnPresetChanged()
@@ -36,6 +42,20 @@ void FSourceEffectRingModulation::OnPresetChanged()
 	RingModulation.SetModulationFrequency(Settings.Frequency);
 	RingModulation.SetDryLevel(Settings.DryLevel);
 	RingModulation.SetWetLevel(Settings.WetLevel);
+
+	// If we're modulating the ring modulator with an audio bus, lets set up that patch source from the audio bus
+	FAudioDeviceManager* AudioDeviceManager = FAudioDeviceManager::Get();
+	if (Settings.AudioBusModulator && AudioDeviceManager)
+	{
+		Audio::FMixerDevice* MixerDevice = (Audio::FMixerDevice*)AudioDeviceManager->GetAudioDeviceRaw(AudioDeviceId);
+		uint32 AudioBusId = Settings.AudioBusModulator->GetUniqueID();
+		Audio::FPatchOutputStrongPtr AudioBusPatchOutputPtr = MixerDevice->AddPatchForAudioBus(AudioBusId, 1.0f);
+		RingModulation.SetExternalPatchSource(AudioBusPatchOutputPtr);
+	}
+	else
+	{
+		RingModulation.SetExternalPatchSource(nullptr);
+	}
 }
 
 void FSourceEffectRingModulation::ProcessAudio(const FSoundEffectSourceInputData& InData, float* OutAudioBufferData)
diff --git a/Engine/Source/Editor/AudioEditor/Classes/Factories/AudioBusFactory.h b/Engine/Source/Editor/AudioEditor/Classes/Factories/AudioBusFactory.h
new file mode 100644
--- /dev/null
+++ b/Engine/Source/Editor/AudioEditor/Classes/Factories/AudioBusFactory.h
@@ -0,0 +1,20 @@
+// Copyright Epic Games, Inc. All Rights Reserved.
+
+//~=============================================================================
+// SoundSubmixFactory
+//~=============================================================================
+
+#pragma once
+#include "Factories/Factory.h"
+#include "AudioBusFactory.generated.h"
+
+UCLASS(MinimalAPI, hidecategories=Object)
+class UAudioBusFactory : public UFactory
+{
+	GENERATED_UCLASS_BODY()
+
+	//~ Begin UFactory Interface
+	virtual UObject* FactoryCreateNew(UClass* Class,UObject* InParent,FName Name,EObjectFlags Flags,UObject* Context,FFeedbackContext* Warn) override;
+	virtual bool CanCreateNew() const override;
+	//~ Begin UFactory Interface	
+};
diff --git a/Engine/Source/Editor/AudioEditor/Private/AssetTypeActions/AssetTypeActions_AudioBus.cpp b/Engine/Source/Editor/AudioEditor/Private/AssetTypeActions/AssetTypeActions_AudioBus.cpp
new file mode 100644
--- /dev/null
+++ b/Engine/Source/Editor/AudioEditor/Private/AssetTypeActions/AssetTypeActions_AudioBus.cpp
@@ -0,0 +1,33 @@
+// Copyright Epic Games, Inc. All Rights Reserved.
+
+#include "AssetTypeActions/AssetTypeActions_AudioBus.h"
+#include "Misc/PackageName.h"
+#include "AssetData.h"
+#include "EditorStyleSet.h"
+#include "Factories/AudioBusFactory.h"
+#include "IContentBrowserSingleton.h"
+#include "ContentBrowserModule.h"
+#include "PropertyCustomizationHelpers.h"
+#include "Widgets/Layout/SBox.h"
+#include "Widgets/Input/SButton.h"
+#include "Widgets/Images/SImage.h"
+#include "AudioEditorModule.h"
+
+
+#define LOCTEXT_NAMESPACE "AssetTypeActions"
+
+UClass* FAssetTypeActions_AudioBus::GetSupportedClass() const
+{
+	return UAudioBus::StaticClass();
+}
+
+const TArray<FText>& FAssetTypeActions_AudioBus::GetSubMenus() const
+{
+	static const TArray<FText> SubMenus
+	{
+		FText(LOCTEXT("AssetAudioBusMenu", "Mix"))
+	};
+
+	return SubMenus;
+}
+#undef LOCTEXT_NAMESPACE
diff --git a/Engine/Source/Editor/AudioEditor/Private/AssetTypeActions/AssetTypeActions_AudioBus.h b/Engine/Source/Editor/AudioEditor/Private/AssetTypeActions/AssetTypeActions_AudioBus.h
new file mode 100644
--- /dev/null
+++ b/Engine/Source/Editor/AudioEditor/Private/AssetTypeActions/AssetTypeActions_AudioBus.h
@@ -0,0 +1,26 @@
+// Copyright Epic Games, Inc. All Rights Reserved.
+
+#pragma once
+
+#include "CoreMinimal.h"
+#include "Toolkits/IToolkitHost.h"
+#include "AssetTypeActions_Base.h"
+#include "Sound/AudioBus.h"
+
+class FMenuBuilder;
+class UAudiobus;
+
+class FAssetTypeActions_AudioBus : public FAssetTypeActions_Base
+{
+public:
+	//~ Begin IAssetTypeActions Implementation
+	virtual FText GetName() const override { return NSLOCTEXT("AssetTypeActions", "AssetTypeActions_AudioBus", "Audio Bus"); }
+	virtual FColor GetTypeColor() const override { return FColor(97, 97, 85); }
+	virtual const TArray<FText>& GetSubMenus() const override;
+	virtual uint32 GetCategories() override { return EAssetTypeCategories::Sounds; }
+	virtual UClass* GetSupportedClass() const override;
+	virtual bool CanFilter() override { return true; }
+	virtual bool IsImportedAsset() const override { return false; }
+	virtual TSharedPtr<class SWidget> GetThumbnailOverlay(const FAssetData& AssetData) const override { return nullptr; }
+	//~ End IAssetTypeActions Implementation
+};
diff --git a/Engine/Source/Editor/AudioEditor/Private/AudioEditorModule.cpp b/Engine/Source/Editor/AudioEditor/Private/AudioEditorModule.cpp
--- a/Engine/Source/Editor/AudioEditor/Private/AudioEditorModule.cpp
+++ b/Engine/Source/Editor/AudioEditor/Private/AudioEditorModule.cpp
@@ -33,6 +33,7 @@
 #include "AssetTypeActions/AssetTypeActions_SoundSubmix.h"
 #include "AssetTypeActions/AssetTypeActions_SoundEffectPreset.h"
 #include "AssetTypeActions/AssetTypeActions_SoundSourceBus.h"
+#include "AssetTypeActions/AssetTypeActions_AudioBus.h"
 #include "Utils.h"
 #include "UObject/UObjectIterator.h"
 #include "Styling/SlateStyle.h"
@@ -161,6 +162,7 @@ public:
 			AssetTools.RegisterAssetTypeActions(MakeShareable(new FAssetTypeActions_SoundEffectSourcePreset));
 			AssetTools.RegisterAssetTypeActions(MakeShareable(new FAssetTypeActions_SoundEffectSourcePresetChain));
 			AssetTools.RegisterAssetTypeActions(MakeShareable(new FAssetTypeActions_SoundSourceBus));
+			AssetTools.RegisterAssetTypeActions(MakeShareable(new FAssetTypeActions_AudioBus));
 		}
 	}
 
diff --git a/Engine/Source/Editor/AudioEditor/Private/Factories/AudioBusFactory.cpp b/Engine/Source/Editor/AudioEditor/Private/Factories/AudioBusFactory.cpp
new file mode 100644
--- /dev/null
+++ b/Engine/Source/Editor/AudioEditor/Private/Factories/AudioBusFactory.cpp
@@ -0,0 +1,28 @@
+// Copyright Epic Games, Inc. All Rights Reserved.
+
+#include "Factories/AudioBusFactory.h"
+#include "Sound/AudioBus.h"
+#include "AudioDeviceManager.h"
+#include "Engine/Engine.h"
+#include "Sound/AudioSettings.h"
+
+UAudioBusFactory::UAudioBusFactory(const FObjectInitializer& ObjectInitializer)
+	: Super(ObjectInitializer)
+{
+	SupportedClass = UAudioBus::StaticClass();
+	bCreateNew = true;
+	bEditorImport = false;
+	bEditAfterNew = true;
+}
+
+UObject* UAudioBusFactory::FactoryCreateNew(UClass* InClass, UObject* InParent, FName InName, EObjectFlags Flags, UObject* Context, FFeedbackContext* Warn)
+{
+	UAudioBus* AudioBus = NewObject<UAudioBus>(InParent, InName, Flags);
+
+	return AudioBus;
+}
+
+bool UAudioBusFactory::CanCreateNew() const
+{
+	return GetDefault<UAudioSettings>()->IsAudioMixerEnabled();
+}
diff --git a/Engine/Source/Runtime/AudioMixer/Private/AudioMixerBlueprintLibrary.cpp b/Engine/Source/Runtime/AudioMixer/Private/AudioMixerBlueprintLibrary.cpp
--- a/Engine/Source/Runtime/AudioMixer/Private/AudioMixerBlueprintLibrary.cpp
+++ b/Engine/Source/Runtime/AudioMixer/Private/AudioMixerBlueprintLibrary.cpp
@@ -104,7 +104,7 @@ void UAudioMixerBlueprintLibrary::StartRecordingOutput(const UObject* WorldConte
 	}
 	else
 	{
-		UE_LOG(LogAudioMixer, Error, TEXT("Output recording is an audio mixer only feature. Please run the game with -audiomixer to enable this feature."));
+		UE_LOG(LogAudioMixer, Error, TEXT("Output recording is an audio mixer only feature."));
 	}
 }
 
@@ -162,7 +162,7 @@ USoundWave* UAudioMixerBlueprintLibrary::StopRecordingOutput(const UObject* Worl
 	}
 	else
 	{
-		UE_LOG(LogAudioMixer, Error, TEXT("Output recording is an audio mixer only feature. Please run the game with -audiomixer to enable this feature."));
+		UE_LOG(LogAudioMixer, Error, TEXT("Output recording is an audio mixer only feature."));
 	}
 
 	return nullptr;
@@ -176,7 +176,7 @@ void UAudioMixerBlueprintLibrary::PauseRecordingOutput(const UObject* WorldConte
 	}
 	else
 	{
-		UE_LOG(LogAudioMixer, Error, TEXT("Output recording is an audio mixer only feature. Please run the game with -audiomixer to enable this feature."));
+		UE_LOG(LogAudioMixer, Error, TEXT("Output recording is an audio mixer only feature."));
 	}
 }
 
@@ -188,7 +188,7 @@ void UAudioMixerBlueprintLibrary::ResumeRecordingOutput(const UObject* WorldCont
 	}
 	else
 	{
-		UE_LOG(LogAudioMixer, Error, TEXT("Output recording is an audio mixer only feature. Please run the game with -audiomixer to enable this feature."));
+		UE_LOG(LogAudioMixer, Error, TEXT("Output recording is an audio mixer only feature."));
 	}
 }
 
@@ -202,7 +202,7 @@ void UAudioMixerBlueprintLibrary::StartAnalyzingOutput(const UObject* WorldConte
 	}
 	else
 	{
-		UE_LOG(LogAudioMixer, Error, TEXT("Spectrum Analysis is an audio mixer only feature. Please run the game with -audiomixer to enable this feature."));
+		UE_LOG(LogAudioMixer, Error, TEXT("Spectrum Analysis is an audio mixer only feature."));
 	}
 }
 
@@ -214,7 +214,7 @@ void UAudioMixerBlueprintLibrary::StopAnalyzingOutput(const UObject* WorldContex
 	}
 	else
 	{
-		UE_LOG(LogAudioMixer, Error, TEXT("Spectrum Analysis is an audio mixer only feature. Please run the game with -audiomixer to enable this feature."));
+		UE_LOG(LogAudioMixer, Error, TEXT("Spectrum Analysis is an audio mixer only feature."));
 	}
 }
 
@@ -226,7 +226,7 @@ void UAudioMixerBlueprintLibrary::GetMagnitudeForFrequencies(const UObject* Worl
 	}
 	else
 	{
-		UE_LOG(LogAudioMixer, Error, TEXT("Output recording is an audio mixer only feature. Please run the game with -audiomixer to enable this feature."));
+		UE_LOG(LogAudioMixer, Error, TEXT("Output recording is an audio mixer only feature."));
 	}
 }
 
@@ -238,7 +238,7 @@ void UAudioMixerBlueprintLibrary::GetPhaseForFrequencies(const UObject* WorldCon
 	}
 	else
 	{
-		UE_LOG(LogAudioMixer, Error, TEXT("Output recording is an audio mixer only feature. Please run the game with -audiomixer to enable this feature."));
+		UE_LOG(LogAudioMixer, Error, TEXT("Output recording is an audio mixer only feature."));
 	}
 }
 
@@ -393,6 +393,47 @@ float UAudioMixerBlueprintLibrary::TrimAudioCache(float InMegabytesToFree)
 	return (float)(((double) NumBytesFreed / 1024) / 1024.0);
 }
 
+void UAudioMixerBlueprintLibrary::StartAudioBus(const UObject* WorldContextObject, UAudioBus* AudioBus)
+{
+	if (Audio::FMixerDevice* MixerDevice = GetAudioMixerDeviceFromWorldContext(WorldContextObject))
+	{
+		uint32 AudioBusId = AudioBus->GetUniqueID();
+		int32 NumChannels = (int32)AudioBus->AudioBusChannels + 1;
+		MixerDevice->StartAudioBus(AudioBusId, NumChannels, false);
+	}
+	else
+	{
+		UE_LOG(LogAudioMixer, Error, TEXT("Audio buses are an audio mixer only feature. Please run the game with audio mixer enabled for this feature."));
+	}
+}
+
+void UAudioMixerBlueprintLibrary::StopAudioBus(const UObject* WorldContextObject, UAudioBus* AudioBus)
+{
+	if (Audio::FMixerDevice* MixerDevice = GetAudioMixerDeviceFromWorldContext(WorldContextObject))
+	{
+		uint32 AudioBusId = AudioBus->GetUniqueID();
+		MixerDevice->StopAudioBus(AudioBusId);
+	}
+	else
+	{
+		UE_LOG(LogAudioMixer, Error, TEXT("Audio buses are an audio mixer only feature. Please run the game with audio mixer enabled for this feature."));
+	}
+}
+
+bool UAudioMixerBlueprintLibrary::IsAudioBusActive(const UObject* WorldContextObject, UAudioBus* AudioBus)
+{
+	if (Audio::FMixerDevice* MixerDevice = GetAudioMixerDeviceFromWorldContext(WorldContextObject))
+	{
+		uint32 AudioBusId = AudioBus->GetUniqueID();
+		return MixerDevice->IsAudioBusActive(AudioBusId);
+	}
+	else
+	{
+		UE_LOG(LogAudioMixer, Error, TEXT("Audio buses are an audio mixer only feature. Please run the game with audio mixer enabled for this feature."));
+		return false;
+	}
+}
+
 void UAudioMixerBlueprintLibrary::PopulateSpectrumAnalyzerSettings(EFFTSize FFTSize, EFFTPeakInterpolationMethod InterpolationMethod, EFFTWindowType WindowType, float HopSize, Audio::FSpectrumAnalyzerSettings &OutSettings)
 {
 	switch (FFTSize)
diff --git a/Engine/Source/Runtime/AudioMixer/Private/AudioMixerBus.cpp b/Engine/Source/Runtime/AudioMixer/Private/AudioMixerBus.cpp
--- a/Engine/Source/Runtime/AudioMixer/Private/AudioMixerBus.cpp
+++ b/Engine/Source/Runtime/AudioMixer/Private/AudioMixerBus.cpp
@@ -5,22 +5,17 @@
 
 namespace Audio
 {
-	FMixerBus::FMixerBus(FMixerSourceManager* InSourceManager, const int32 InNumChannels, const int32 InNumFrames)
+	FMixerAudioBus::FMixerAudioBus(FMixerSourceManager* InSourceManager, bool bInIsAutomatic, int32 InNumChannels)
 		: CurrentBufferIndex(1)
 		, NumChannels(InNumChannels)
-		, NumFrames(InNumFrames)
+		, NumFrames(InSourceManager->GetNumOutputFrames())
 		, SourceManager(InSourceManager)
+		, bIsAutomatic(bInIsAutomatic)
 	{
-		// Prepare the previous buffer with zero'd data
-		const int32 NumSamples = InNumChannels * InNumFrames;
-		for (int32 i = 0; i < 2; ++i)
-		{
-			MixedSourceData[i].Reset();
-			MixedSourceData[i].AddZeroed(NumSamples);
-		}
+		SetNumOutputChannels(NumChannels);
 	}
 
-	void FMixerBus::SetNumOutputChannels(int32 InNumOutputChannels)
+	void FMixerAudioBus::SetNumOutputChannels(int32 InNumOutputChannels)
 	{
 		NumChannels = InNumOutputChannels;
 		const int32 NumSamples = NumChannels * NumFrames;
@@ -31,49 +26,44 @@ namespace Audio
 		}
 	}
 
-	void FMixerBus::Update()
+	void FMixerAudioBus::Update()
 	{
 		CurrentBufferIndex = !CurrentBufferIndex;
 	}
 
-	void FMixerBus::AddInstanceId(const int32 InSourceId, int32 InNumOutputChannels)
+	void FMixerAudioBus::AddInstanceId(const int32 InSourceId, int32 InNumOutputChannels)
 	{
 		InstanceIds.Add(InSourceId);
-
-		if (NumChannels != InNumOutputChannels)
-		{
-			SetNumOutputChannels(InNumOutputChannels);
-		}
 	}
 
-	bool FMixerBus::RemoveInstanceId(const int32 InSourceId)
+	bool FMixerAudioBus::RemoveInstanceId(const int32 InSourceId)
 	{
 		InstanceIds.Remove(InSourceId);
 
 		// Return true if there is no more instances or sends
-		return !InstanceIds.Num() && !BusSends[(int32)EBusSendType::PreEffect].Num() && !BusSends[(int32)EBusSendType::PostEffect].Num();
+		return bIsAutomatic && !InstanceIds.Num() && !AudioBusSends[(int32)EBusSendType::PreEffect].Num() && !AudioBusSends[(int32)EBusSendType::PostEffect].Num();
 	}
 
-	void FMixerBus::AddBusSend(EBusSendType BusSendType, const FBusSend& InBusSend)
+	void FMixerAudioBus::AddSend(EBusSendType BusSendType, const FAudioBusSend& InAudioBusSend)
 	{
 		// Make sure we don't have duplicates in the bus sends
-		for (FBusSend& BusSend : BusSends[(int32)BusSendType])
+		for (FAudioBusSend& BusSend : AudioBusSends[(int32)BusSendType])
 		{
 			// If it's already added, just update the send level
-			if (BusSend.SourceId == InBusSend.SourceId)
+			if (BusSend.SourceId == InAudioBusSend.SourceId)
 			{
-				BusSend.SendLevel = InBusSend.SendLevel;
+				BusSend.SendLevel = InAudioBusSend.SendLevel;
 				return;
 			}
 		}
 
 		// It's a new source id so just add it
-		BusSends[(int32)BusSendType].Add(InBusSend);
+		AudioBusSends[(int32)BusSendType].Add(InAudioBusSend);
 	}
 
-	bool FMixerBus::RemoveBusSend(EBusSendType BusSendType, const int32 InSourceId)
+	bool FMixerAudioBus::RemoveSend(EBusSendType BusSendType, const int32 InSourceId)
 	{
-		TArray<FBusSend>& Sends = BusSends[(int32)BusSendType];
+		TArray<FAudioBusSend>& Sends = AudioBusSends[(int32)BusSendType];
 
 		for (int32 i = Sends.Num() - 1; i >= 0; --i)
 		{
@@ -87,11 +77,11 @@ namespace Audio
 			}
 		}
 
-		// Return true if there is no more instances or sends
-		return !InstanceIds.Num() && !BusSends[(int32)EBusSendType::PreEffect].Num() && !BusSends[(int32)EBusSendType::PostEffect].Num();
+		// Return true if there is no more instances or sends and this is an automatic audio bus
+		return bIsAutomatic && !InstanceIds.Num() && !AudioBusSends[(int32)EBusSendType::PreEffect].Num() && !AudioBusSends[(int32)EBusSendType::PostEffect].Num();
 	}
 
-	void FMixerBus::MixBuffer()
+	void FMixerAudioBus::MixBuffer()
 	{
 		// Reset and zero the mixed source data buffer for this bus
 		const int32 NumSamples = NumFrames * NumChannels;
@@ -104,87 +94,108 @@ namespace Audio
 		for (int32 BusSendType = 0; BusSendType < (int32)EBusSendType::Count; ++BusSendType)
 		{
 			// Loop through the send list for this bus
-			for (FBusSend& BusSend : BusSends[BusSendType])
+			for (const FAudioBusSend& AudioBusSend : AudioBusSends[BusSendType])
 			{
 				const float* SourceBufferPtr = nullptr;
 
-				// If the source is a bus, we need to use the previous renderer buffer 
-				if (SourceManager->IsBus(BusSend.SourceId))
+				// Check if the bus send is in fact from another audio bus (and not a source)
+				if (AudioBusSend.AudioBusId != INDEX_NONE)
 				{
-					SourceBufferPtr = SourceManager->GetPreviousBusBuffer(BusSend.SourceId);
-				}
-				// If the source mixing into this is not itself a bus, then simply mix the pre-attenuation audio of the source into the bus
-				// The source will have already computed its buffers for this frame
-				else if (BusSendType == (int32)EBusSendType::PostEffect)
-				{
-					SourceBufferPtr = SourceManager->GetPreDistanceAttenuationBuffer(BusSend.SourceId);
+					SourceBufferPtr = SourceManager->GetPreviousAudioBusBuffer(AudioBusSend.AudioBusId);
 				}
 				else
 				{
-					SourceBufferPtr = SourceManager->GetPreEffectBuffer(BusSend.SourceId);
+					// If the audio source mixing to this audio bus is itself a source bus, we need to use the previous renderer buffer to avoid infinite recursion
+					if (SourceManager->IsSourceBus(AudioBusSend.SourceId))
+					{
+						SourceBufferPtr = SourceManager->GetPreviousSourceBusBuffer(AudioBusSend.SourceId);
+					}
+					// If the source mixing into this is not itself a bus, then simply mix the pre-attenuation audio of the source into the bus
+					// The source will have already computed its buffers for this frame
+					else if (BusSendType == (int32)EBusSendType::PostEffect)
+					{
+						SourceBufferPtr = SourceManager->GetPreDistanceAttenuationBuffer(AudioBusSend.SourceId);
+					}
+					else
+					{
+						SourceBufferPtr = SourceManager->GetPreEffectBuffer(AudioBusSend.SourceId);
+					}
 				}
 
-				const int32 NumSourceChannels = SourceManager->GetNumChannels(BusSend.SourceId);
-				const int32 NumOutputFrames = SourceManager->GetNumOutputFrames();
-				const int32 NumSourceSamples = NumSourceChannels * NumOutputFrames;
-
-				// If source channels are 1 but the bus is 2 channels, we need to up-mix
-				if (NumSourceChannels == 1 && NumChannels == 2)
+				// It's possible we may not have a source buffer ptr here, so protect against it
+				if (ensure(SourceBufferPtr))
 				{
-					int32 BusSampleIndex = 0;
-					for (int32 SourceSampleIndex = 0; SourceSampleIndex < NumSourceSamples; ++SourceSampleIndex)
-					{
-						// Take half the source sample to up-mix to stereo
-						const float SourceSample = 0.5f * BusSend.SendLevel * SourceBufferPtr[SourceSampleIndex];
+					const int32 NumSourceChannels = SourceManager->GetNumChannels(AudioBusSend.SourceId);
+					const int32 NumOutputFrames = SourceManager->GetNumOutputFrames();
+					const int32 NumSourceSamples = NumSourceChannels * NumOutputFrames;
 
-						// Mix in the source sample
-						for (int32 BusChannel = 0; BusChannel < NumChannels; ++BusChannel)
+					// If source channels are 1 but the bus is 2 channels, we need to up-mix
+					if (NumSourceChannels == 1 && NumChannels == 2)
+					{
+						int32 BusSampleIndex = 0;
+						for (int32 SourceSampleIndex = 0; SourceSampleIndex < NumSourceSamples; ++SourceSampleIndex)
 						{
-							BusDataBufferPtr[BusSampleIndex++] += SourceSample;
+							// Take half the source sample to up-mix to stereo
+							const float SourceSample = 0.5f * AudioBusSend.SendLevel * SourceBufferPtr[SourceSampleIndex];
+
+							// Mix in the source sample
+							for (int32 BusChannel = 0; BusChannel < NumChannels; ++BusChannel)
+							{
+								BusDataBufferPtr[BusSampleIndex++] += SourceSample;
+							}
 						}
 					}
-				}
-				// If the source channels are 2 but the bus is 1 channel, we need to down-mix
-				else if (NumSourceChannels == 2 && NumChannels == 1)
-				{
-					int32 SourceSampleIndex = 0;
-					for (int32 BusSampleIndex = 0; BusSampleIndex < NumOutputFrames; ++BusSampleIndex)
+					// If the source channels are 2 but the bus is 1 channel, we need to down-mix
+					else if (NumSourceChannels == 2 && NumChannels == 1)
 					{
-						// Downmix the stereo source to mono before summing to bus
-						float SourceSample = 0.0f;
-						for (int32 SourceChannel = 0; SourceChannel < 2; ++SourceChannel)
+						int32 SourceSampleIndex = 0;
+						for (int32 BusSampleIndex = 0; BusSampleIndex < NumOutputFrames; ++BusSampleIndex)
 						{
-							SourceSample += SourceBufferPtr[SourceSampleIndex++];
-						}
+							// Downmix the stereo source to mono before summing to bus
+							float SourceSample = 0.0f;
+							for (int32 SourceChannel = 0; SourceChannel < 2; ++SourceChannel)
+							{
+								SourceSample += SourceBufferPtr[SourceSampleIndex++];
+							}
 
-						SourceSample *= 0.5f;
-						SourceSample *= BusSend.SendLevel;
+							SourceSample *= 0.5f;
+							SourceSample *= AudioBusSend.SendLevel;
 
-						BusDataBufferPtr[BusSampleIndex] += SourceSample;
+							BusDataBufferPtr[BusSampleIndex] += SourceSample;
+						}
 					}
-				}
-				// If they're the same channels, just mix it in
-				else
-				{
-					for (int32 SampleIndex = 0; SampleIndex < NumSourceSamples; ++SampleIndex)
+					// If they're the same channels, just mix it in
+					else
 					{
-						BusDataBufferPtr[SampleIndex] += (BusSend.SendLevel * SourceBufferPtr[SampleIndex]);
+						for (int32 SampleIndex = 0; SampleIndex < NumSourceSamples; ++SampleIndex)
+						{
+							BusDataBufferPtr[SampleIndex] += (AudioBusSend.SendLevel * SourceBufferPtr[SampleIndex]);
+						}
 					}
-				}
+
+					// Push the mixed data to the patch splitter
+					PatchSplitter.PushAudio(BusDataBufferPtr, NumOutputFrames * NumChannels);
+				}				
 			}
 		}
 	}
 
 
-	const float* FMixerBus::GetCurrentBusBuffer() const
+	const float* FMixerAudioBus::GetCurrentBusBuffer() const
 	{
 		return MixedSourceData[CurrentBufferIndex].GetData();
 	}
 
-	const float* FMixerBus::GetPreviousBusBuffer() const
+	const float* FMixerAudioBus::GetPreviousBusBuffer() const
 	{
 		return MixedSourceData[!CurrentBufferIndex].GetData();
 	}
 
+	FPatchOutputStrongPtr FMixerAudioBus::AddNewPatch(int32 MaxLatencyInSamples, float InGain)
+	{
+		return PatchSplitter.AddNewPatch(MaxLatencyInSamples, InGain);
+	}
+
+
 }
 
diff --git a/Engine/Source/Runtime/AudioMixer/Private/AudioMixerDevice.cpp b/Engine/Source/Runtime/AudioMixer/Private/AudioMixerDevice.cpp
--- a/Engine/Source/Runtime/AudioMixer/Private/AudioMixerDevice.cpp
+++ b/Engine/Source/Runtime/AudioMixer/Private/AudioMixerDevice.cpp
@@ -1912,6 +1912,65 @@ namespace Audio
 		}
 	}
 
+	void FMixerDevice::StartAudioBus(uint32 InAudioBusId, int32 InNumChannels, bool bInIsAutomatic)
+	{
+		if (IsInGameThread())
+		{
+			if (ActiveAudioBuses_GameThread.Contains(InAudioBusId))
+			{
+				return;
+			}
+
+			ActiveAudioBuses_GameThread.Add(InAudioBusId);
+
+			FAudioThread::RunCommandOnAudioThread([this, InAudioBusId, InNumChannels, bInIsAutomatic]()
+			{
+				SourceManager->StartAudioBus(InAudioBusId, InNumChannels, bInIsAutomatic);
+			});
+		}
+	}
+
+	void FMixerDevice::StopAudioBus(uint32 InAudioBusId)
+	{
+		if (IsInGameThread())
+		{
+			if (!ActiveAudioBuses_GameThread.Contains(InAudioBusId))
+			{
+				return;
+			}
+
+			ActiveAudioBuses_GameThread.Remove(InAudioBusId);
+
+			FAudioThread::RunCommandOnAudioThread([this, InAudioBusId]()
+			{
+				SourceManager->StopAudioBus(InAudioBusId);
+			});
+		}
+	}
+
+	bool FMixerDevice::IsAudioBusActive(uint32 InAudioBusId)
+	{
+		if (IsInGameThread())
+		{
+			return ActiveAudioBuses_GameThread.Contains(InAudioBusId);
+		}
+
+		check(IsInAudioThread());
+		return SourceManager->IsAudioBusActive(InAudioBusId);
+	}
+
+	FPatchOutputStrongPtr FMixerDevice::AddPatchForAudioBus(uint32 InAudioBusId, float PatchGain)
+	{
+		if (ensure(!IsInGameThread() && !IsInAudioThread()))
+		{
+			return SourceManager->AddPatchForAudioBus(InAudioBusId, PatchGain);
+		}
+		else
+		{
+			return nullptr;
+		}
+	}
+
 	int32 FMixerDevice::GetDeviceSampleRate() const
 	{
 		return SampleRate;
diff --git a/Engine/Source/Runtime/AudioMixer/Private/AudioMixerSource.cpp b/Engine/Source/Runtime/AudioMixer/Private/AudioMixerSource.cpp
--- a/Engine/Source/Runtime/AudioMixer/Private/AudioMixerSource.cpp
+++ b/Engine/Source/Runtime/AudioMixer/Private/AudioMixerSource.cpp
@@ -156,12 +156,24 @@ namespace Audio
 				}
 
 				// Setup the bus Id if this source is a bus
-				if (WaveInstance->WaveData->bIsBus)
+				if (WaveInstance->WaveData->bIsSourceBus)
 				{
-					InitParams.BusId = WaveInstance->WaveData->GetUniqueID();
+					// We need to check if the source bus has an audio bus specified
+					USoundSourceBus* SoundSourceBus = CastChecked<USoundSourceBus>(WaveInstance->WaveData);
+
+					// If it does, we will use that audio bus as the source of the audio data for the source bus
+					if (SoundSourceBus->AudioBus)
+					{
+						InitParams.AudioBusId = SoundSourceBus->AudioBus->GetUniqueID();
+					}
+					else
+					{
+						InitParams.AudioBusId = WaveInstance->WaveData->GetUniqueID();
+					}
+
 					if (!WaveInstance->WaveData->IsLooping())
 					{
-						InitParams.BusDuration = WaveInstance->WaveData->GetDuration();
+						InitParams.SourceBusDuration = WaveInstance->WaveData->GetDuration();
 					}
 				}
 
@@ -170,38 +182,7 @@ namespace Audio
 				InitParams.bOutputToBusOnly = WaveInstance->bOutputToBusOnly;
 				DynamicBusSendInfos.Reset();
 
-				// If this source is sending its audio to a bus
-				for (int32 BusSendType = 0; BusSendType < (int32)EBusSendType::Count; ++BusSendType)
-				{
-					// And add all the source bus sends
-					for (FSoundSourceBusSendInfo& SendInfo : WaveInstance->SoundSourceBusSends[BusSendType])
-					{
-						if (SendInfo.SoundSourceBus != nullptr)
-						{
-							FMixerBusSend BusSend;
-							BusSend.BusId = SendInfo.SoundSourceBus->GetUniqueID();
-							BusSend.SendLevel = SendInfo.SendLevel;
-							InitParams.BusSends[BusSendType].Add(BusSend);
-
-							FDynamicBusSendInfo DynamicBusSendInfo;
-							DynamicBusSendInfo.SendLevel = SendInfo.SendLevel;
-							DynamicBusSendInfo.BusId = BusSend.BusId;
-							DynamicBusSendInfo.BusSendLevelControlMethod = SendInfo.SourceBusSendLevelControlMethod;
-							DynamicBusSendInfo.BusSendType = (EBusSendType)BusSendType;
-							DynamicBusSendInfo.MinSendLevel = SendInfo.MinSendLevel;
-							DynamicBusSendInfo.MaxSendLevel = SendInfo.MaxSendLevel;
-							DynamicBusSendInfo.MinSendDistance = SendInfo.MinSendDistance;
-							DynamicBusSendInfo.MaxSendDistance = SendInfo.MaxSendDistance;
-							DynamicBusSendInfo.CustomSendLevelCurve = SendInfo.CustomSendLevelCurve;
-
-							// Copy the bus SourceBusSendInfo structs to a local copy so we can update it in the update tick
-							DynamicBusSendInfos.Add(DynamicBusSendInfo);
-
-							// Flag that we're sending audio to buses so we can check for updates to send levels
-							bSendingAudioToBuses = true;
-						}
-					}
-				}
+				InitBusSends(WaveInstance, InitParams);
 			}
 
 			// Don't set up any submixing if we're set to output to bus only
@@ -341,6 +322,71 @@ namespace Audio
 		return false;
 	}
 
+	void FMixerSource::InitBusSends(FWaveInstance* InWaveInstance, FMixerSourceVoiceInitParams& InitParams)
+	{
+		for (int32 BusSendType = 0; BusSendType < (int32)EBusSendType::Count; ++BusSendType)
+		{
+			// And add all the source bus sends
+			for (FSoundSourceBusSendInfo& SendInfo : InWaveInstance->BusSends[BusSendType])
+			{
+				// Avoid redoing duplicate code for sending audio to source bus or audio bus. Most of it is the same other than the bus id.
+				auto InitBusSend = [this](FMixerSourceVoiceInitParams& InInitParams, const FSoundSourceBusSendInfo& InSendInfo, int32 BusSendType, uint32 InBusId)
+				{
+					FInitAudioBusSend BusSend;
+					BusSend.AudioBusId = InBusId;
+					BusSend.SendLevel = InSendInfo.SendLevel;
+					InInitParams.AudioBusSends[BusSendType].Add(BusSend);
+
+					FDynamicBusSendInfo DynamicBusSendInfo;
+					DynamicBusSendInfo.SendLevel = InSendInfo.SendLevel;
+					DynamicBusSendInfo.BusId = BusSend.AudioBusId;
+					DynamicBusSendInfo.BusSendLevelControlMethod = InSendInfo.SourceBusSendLevelControlMethod;
+					DynamicBusSendInfo.BusSendType = (EBusSendType)BusSendType;
+					DynamicBusSendInfo.MinSendLevel = InSendInfo.MinSendLevel;
+					DynamicBusSendInfo.MaxSendLevel = InSendInfo.MaxSendLevel;
+					DynamicBusSendInfo.MinSendDistance = InSendInfo.MinSendDistance;
+					DynamicBusSendInfo.MaxSendDistance = InSendInfo.MaxSendDistance;
+					DynamicBusSendInfo.CustomSendLevelCurve = InSendInfo.CustomSendLevelCurve;
+
+					// Copy the bus SourceBusSendInfo structs to a local copy so we can update it in the update tick
+					DynamicBusSendInfos.Add(DynamicBusSendInfo);
+
+					// Flag that we're sending audio to buses so we can check for updates to send levels
+					bSendingAudioToBuses = true;
+				};
+
+				// Retrieve bus id of the audio bus to use
+				if (SendInfo.SoundSourceBus)
+				{						
+					uint32 BusId;
+
+					// Either use the bus id of the source bus's audio bus id if it was specified
+					if (SendInfo.SoundSourceBus->AudioBus)
+					{
+						BusId = SendInfo.SoundSourceBus->AudioBus->GetUniqueID();
+					}
+					else
+					{
+						// otherwise, use the id of the source bus itself (for an automatic source bus)
+						BusId = SendInfo.SoundSourceBus->GetUniqueID();
+					}
+
+					// Call lambda w/ the correctly derived bus id
+					InitBusSend(InitParams, SendInfo, BusSendType, BusId);
+				}
+
+				if (SendInfo.AudioBus)
+				{
+					// Only need to send audio to just the specified audio bus
+					uint32 BusId = SendInfo.AudioBus->GetUniqueID();
+
+					// Note we will be sending audio to both the specified source bus and the audio bus with the same send level
+					InitBusSend(InitParams, SendInfo, BusSendType, BusId);
+				}
+			}
+		}
+	}
+
 	void FMixerSource::Update()
 	{
 		CSV_SCOPED_TIMING_STAT(Audio, UpdateSources);
@@ -429,7 +475,7 @@ namespace Audio
 		// Not all wave data types have a non-zero duration
 		if (SoundWave.Duration > 0.0f)
 		{
-			if (!SoundWave.bIsBus)
+			if (!SoundWave.bIsSourceBus)
 			{
 				NumTotalFrames = SoundWave.Duration * SoundWave.GetSampleRateForCurrentPlatform();
 				check(NumTotalFrames > 0);
@@ -476,7 +522,7 @@ namespace Audio
 			}
 			else if (WaveInstance)
 			{
-				if (WaveInstance->WaveData->bIsBus)
+				if (WaveInstance->WaveData->bIsSourceBus)
 				{
 					// Buses don't need to do anything to play audio
 					return true;
@@ -488,7 +534,7 @@ namespace Audio
 					if ((BufferType == EBufferType::PCMRealTime || BufferType == EBufferType::Streaming) && WaveInstance->WaveData)
 					{
 						// If any of these conditions meet, we need to do an initial async decode before we're ready to start playing the sound
-						if (WaveInstance->StartTime > 0.0f || WaveInstance->WaveData->bProcedural || WaveInstance->WaveData->bIsBus || !WaveInstance->WaveData->CachedRealtimeFirstBuffer)
+						if (WaveInstance->StartTime > 0.0f || WaveInstance->WaveData->bProcedural || WaveInstance->WaveData->bIsSourceBus || !WaveInstance->WaveData->CachedRealtimeFirstBuffer)
 						{
 							// Before reading more PCMRT data, we first need to seek the buffer
 							if (WaveInstance->IsSeekable())
@@ -1017,11 +1063,7 @@ namespace Audio
 			{
 				DynamicBusSendInfo.SendLevel = SendLevel;
 
-				FMixerBusSend BusSend;
-				BusSend.BusId = DynamicBusSendInfo.BusId;
-				BusSend.SendLevel = SendLevel;
-
-				MixerSourceVoice->SetBusSendInfo(DynamicBusSendInfo.BusSendType, BusSend);
+				MixerSourceVoice->SetAudioBusSendInfo(DynamicBusSendInfo.BusSendType, DynamicBusSendInfo.BusId, SendLevel);
 			}
 		}
 	}
diff --git a/Engine/Source/Runtime/AudioMixer/Private/AudioMixerSource.h b/Engine/Source/Runtime/AudioMixer/Private/AudioMixerSource.h
--- a/Engine/Source/Runtime/AudioMixer/Private/AudioMixerSource.h
+++ b/Engine/Source/Runtime/AudioMixer/Private/AudioMixerSource.h
@@ -61,6 +61,9 @@ namespace Audio
 
 	private:
 
+		/** Initializes the bus sends. */
+		void InitBusSends(FWaveInstance* InWaveInstance, FMixerSourceVoiceInitParams& InitParams);
+
 		/** Frees any resources for this sound source. */
 		void FreeResources();
 
diff --git a/Engine/Source/Runtime/AudioMixer/Private/AudioMixerSourceBuffer.cpp b/Engine/Source/Runtime/AudioMixer/Private/AudioMixerSourceBuffer.cpp
--- a/Engine/Source/Runtime/AudioMixer/Private/AudioMixerSourceBuffer.cpp
+++ b/Engine/Source/Runtime/AudioMixer/Private/AudioMixerSourceBuffer.cpp
@@ -92,7 +92,7 @@ namespace Audio
 		, bIsSeeking(bInIsSeeking)
 		, bLoopCallback(false)
 		, bProcedural(InWave.bProcedural)
-		, bIsBus(InWave.bIsBus)
+		, bIsBus(InWave.bIsSourceBus)
 		, bForceSyncDecode(bInForceSyncDecode)
 	{
 		InWave.AddPlayingSource(this);
diff --git a/Engine/Source/Runtime/AudioMixer/Private/AudioMixerSourceManager.cpp b/Engine/Source/Runtime/AudioMixer/Private/AudioMixerSourceManager.cpp
--- a/Engine/Source/Runtime/AudioMixer/Private/AudioMixerSourceManager.cpp
+++ b/Engine/Source/Runtime/AudioMixer/Private/AudioMixerSourceManager.cpp
@@ -204,11 +204,11 @@ namespace Audio
 			SourceInfo.NumFramesPlayed = 0;
 			SourceInfo.StartTime = 0.0;
 			SourceInfo.SubmixSends.Reset();
-			SourceInfo.BusId = INDEX_NONE;
-			SourceInfo.BusDurationFrames = INDEX_NONE;
+			SourceInfo.AudioBusId = INDEX_NONE;
+			SourceInfo.SourceBusDurationFrames = INDEX_NONE;
 		
-			SourceInfo.BusSends[(int32)EBusSendType::PreEffect].Reset();
-			SourceInfo.BusSends[(int32)EBusSendType::PostEffect].Reset();
+			SourceInfo.AudioBusSends[(int32)EBusSendType::PreEffect].Reset();
+			SourceInfo.AudioBusSends[(int32)EBusSendType::PostEffect].Reset();
 
 			SourceInfo.SourceEffectChainId = INDEX_NONE;
 
@@ -388,39 +388,44 @@ namespace Audio
 		}
 #endif
 		// Remove from list of active bus or source ids depending on what type of source this is
-		if (SourceInfo.BusId != INDEX_NONE)
+		if (SourceInfo.AudioBusId != INDEX_NONE)
 		{
 			// Remove this bus from the registry of bus instances
-			FMixerBus* Bus = Buses.Find(SourceInfo.BusId);
-			AUDIO_MIXER_CHECK(Bus);
-
-			// Remove this source from the list of bus instances.
-			if (Bus->RemoveInstanceId(SourceId))
+			TSharedPtr<FMixerAudioBus> AudioBusPtr = AudioBuses.FindRef(SourceInfo.AudioBusId);
+			if (ensure(AudioBusPtr.IsValid()))
 			{
-				Buses.Remove(SourceInfo.BusId);
+				// If this audio bus was automatically created via source bus playback, this this audio bus can be removed
+				if (AudioBusPtr->RemoveInstanceId(SourceId))
+				{
+					// Only automatic buses will be getting removed here. Otherwise they need to be manually removed from the source manager.
+					ensure(AudioBusPtr->IsAutomatic());
+					AudioBuses.Remove(SourceInfo.AudioBusId);
+				}
 			}
 		}
 
 		// Remove this source's send list from the bus data registry
-		for (int32 BusSendType = 0; BusSendType < (int32)EBusSendType::Count; ++BusSendType)
+		for (int32 AudioBusSendType = 0; AudioBusSendType < (int32)EBusSendType::Count; ++AudioBusSendType)
 		{
-			for (uint32 BusId : SourceInfo.BusSends[BusSendType])
+			for (uint32 AudioBusId : SourceInfo.AudioBusSends[AudioBusSendType])
 			{
 				// we should have a bus registration entry still since the send hasn't been cleaned up yet
-				FMixerBus* Bus = Buses.Find(BusId);
-				AUDIO_MIXER_CHECK(Bus);
-
-				if (Bus->RemoveBusSend((EBusSendType)BusSendType, SourceId))
+				TSharedPtr<FMixerAudioBus> AudioBusPtr = AudioBuses.FindRef(AudioBusId);
+				if (ensure(AudioBusPtr.IsValid()))
 				{
-					Buses.Remove(BusId);
+					if (AudioBusPtr->RemoveSend((EBusSendType)AudioBusSendType, SourceId))
+					{
+						ensure(AudioBusPtr->IsAutomatic());
+						AudioBuses.Remove(AudioBusId);
+					}
 				}
 			}
 
-			SourceInfo.BusSends[BusSendType].Reset();
+			SourceInfo.AudioBusSends[AudioBusSendType].Reset();
 		}
 
-		SourceInfo.BusId = INDEX_NONE;
-		SourceInfo.BusDurationFrames = INDEX_NONE;
+		SourceInfo.AudioBusId = INDEX_NONE;
+		SourceInfo.SourceBusDurationFrames = INDEX_NONE;
 
 		// Free the mixer source buffer data
 		if (SourceInfo.MixerSourceBuffer.IsValid())
@@ -601,9 +606,9 @@ namespace Audio
 		return NumActiveSources;
 	}
 
-	int32 FMixerSourceManager::GetNumActiveBuses() const
+	int32 FMixerSourceManager::GetNumActiveAudioBuses() const
 	{
-		return Buses.Num();
+		return AudioBuses.Num();
 	}
 
 	void FMixerSourceManager::InitSource(const int32 SourceId, const FMixerSourceVoiceInitParams& InitParams)
@@ -702,6 +707,7 @@ namespace Audio
 				InitData.SampleRate = MixerDevice->SampleRate;
 				InitData.NumSourceChannels = InitParams.NumInputChannels;
 				InitData.AudioClock = MixerDevice->GetAudioTime();
+				InitData.AudioDeviceId = MixerDevice->DeviceID;
 
 				SourceInfo.SourceEffectChainId = InitParams.SourceEffectChainId;
 				BuildSourceEffectChain(SourceId, InitData, InitParams.SourceEffectChain);
@@ -709,34 +715,33 @@ namespace Audio
 				// Whether or not to output to bus only
 				SourceInfo.bOutputToBusOnly = InitParams.bOutputToBusOnly;
 
-				// If this is a bus, add this source id to the list of active bus ids
-				if (InitParams.BusId != INDEX_NONE)
+				// If this is going to be a source bus, add this source id to the list of active bus ids
+				if (InitParams.AudioBusId != INDEX_NONE)
 				{
 					// Setting this BusId will flag this source as a bus. It doesn't try to generate 
 					// audio in the normal way but instead will render in a second stage, after normal source rendering.
-					SourceInfo.BusId = InitParams.BusId;
+					SourceInfo.AudioBusId = InitParams.AudioBusId;
 
-					// Bus duration allows us to stop a bus after a given time
-					if (InitParams.BusDuration != 0.0f)
+					// Source bus duration allows us to stop a bus after a given time
+					if (InitParams.SourceBusDuration != 0.0f)
 					{
-						SourceInfo.BusDurationFrames = InitParams.BusDuration * MixerDevice->GetSampleRate();
+						SourceInfo.SourceBusDurationFrames = InitParams.SourceBusDuration * MixerDevice->GetSampleRate();
 					}
 
 					// Register this bus as an instance
-					FMixerBus* Bus = Buses.Find(InitParams.BusId);
-					if (Bus)
+					TSharedPtr<FMixerAudioBus> AudioBusPtr = AudioBuses.FindRef(SourceInfo.AudioBusId);
+					if (AudioBusPtr.IsValid())
 					{
 						// If this bus is already registered, add this as a source id
-						Bus->AddInstanceId(SourceId, InitParams.NumInputChannels);
+						AudioBusPtr->AddInstanceId(SourceId, InitParams.NumInputChannels);
 					}
 					else
 					{
-						// If the bus is not registered, make a new entry
-						FMixerBus NewBusData(this, InitParams.NumInputChannels, NumOutputFrames);
-
-						NewBusData.AddInstanceId(SourceId, InitParams.NumInputChannels);
+						// If the bus is not registered, make a new entry. This will default to an automatic audio bus until explicitly made manual later.
+						TSharedPtr<FMixerAudioBus> NewAudioBus = TSharedPtr<FMixerAudioBus>(new FMixerAudioBus(this, true, InitParams.NumInputChannels));
+						NewAudioBus->AddInstanceId(SourceId, InitParams.NumInputChannels);
 
-						Buses.Add(InitParams.BusId, NewBusData);
+						AudioBuses.Add(InitParams.AudioBusId, NewAudioBus);
 					}
 				}
 
@@ -744,34 +749,34 @@ namespace Audio
 				// Note: buses can also send their audio to other buses.
 				for (int32 BusSendType = 0; BusSendType < (int32)EBusSendType::Count; ++BusSendType)
 				{
-					for (const FMixerBusSend& BusSend : InitParams.BusSends[BusSendType])
+					for (const FInitAudioBusSend& AudioBusSend : InitParams.AudioBusSends[BusSendType])
 					{
 						// New struct to map which source (SourceId) is sending to the bus
-						FBusSend NewBusSend;
-						NewBusSend.SourceId = SourceId;
-						NewBusSend.SendLevel = BusSend.SendLevel;
+						FAudioBusSend NewAudioBusSend;
+						NewAudioBusSend.SourceId = SourceId;
+						NewAudioBusSend.SendLevel = AudioBusSend.SendLevel;
 
 						// Get existing BusId and add the send, or create new bus registration
-						FMixerBus* Bus = Buses.Find(BusSend.BusId);
-						if (Bus)
+						TSharedPtr<FMixerAudioBus> AudioBusPtr = AudioBuses.FindRef(AudioBusSend.AudioBusId);
+						if (AudioBusPtr.IsValid())
 						{
-							Bus->AddBusSend((EBusSendType)BusSendType, NewBusSend);
+							AudioBusPtr->AddSend((EBusSendType)BusSendType, NewAudioBusSend);
 						}
 						else
 						{
-							// If the bus is not registered, make a new entry
-							FMixerBus NewBusData(this, InitParams.NumInputChannels, NumOutputFrames);
+							// If the bus is not registered, make a new entry. This will default to an automatic audio bus until explicitly made manual later.
+							TSharedPtr<FMixerAudioBus> NewAudioBus(new FMixerAudioBus(this, true, InitParams.NumInputChannels));
 
 							// Add a send to it. This will not have a bus instance id (i.e. won't output audio), but 
 							// we register the send anyway in the event that this bus does play, we'll know to send this
 							// source's audio to it.
-							NewBusData.AddBusSend((EBusSendType)BusSendType, NewBusSend);
+							NewAudioBus->AddSend((EBusSendType)BusSendType, NewAudioBusSend);
 
-							Buses.Add(BusSend.BusId, NewBusData);
+							AudioBuses.Add(AudioBusSend.AudioBusId, NewAudioBus);
 						}
 
 						// Store on this source, which buses its sending its audio to
-						SourceInfo.BusSends[BusSendType].Add(BusSend.BusId);
+						SourceInfo.AudioBusSends[BusSendType].Add(AudioBusSend.AudioBusId);
 					}
 				}
 			}
@@ -892,6 +897,97 @@ namespace Audio
 		});
 	}
 
+	void FMixerSourceManager::StartAudioBus(uint32 InAudioBusId, int32 InNumChannels, bool bInIsAutomatic)
+	{
+		if (AudioBusIds_AudioThread.Contains(InAudioBusId))
+		{
+			return;
+		}
+
+		AudioBusIds_AudioThread.Add(InAudioBusId);
+
+		AudioMixerThreadCommand([this, InAudioBusId, InNumChannels, bInIsAutomatic]()
+		{
+			// If this audio bus id already exists, set it to not be automatic and return it
+			TSharedPtr<FMixerAudioBus> AudioBusPtr = AudioBuses.FindRef(InAudioBusId);
+			if (AudioBusPtr.IsValid())
+			{
+				// If this audio bus already existed, make sure the num channels lines up
+				ensure(AudioBusPtr->GetNumChannels() == InNumChannels);
+				AudioBusPtr->SetAutomatic(bInIsAutomatic);
+			}
+			else
+			{
+				// If the bus is not registered, make a new entry.
+				TSharedPtr<FMixerAudioBus> NewBusData(new FMixerAudioBus(this, bInIsAutomatic, InNumChannels));
+
+				AudioBuses.Add(InAudioBusId, NewBusData);
+			}
+
+			//  Now add any existing playing sources to this audio bus as sends if they exist
+			for (FSourceInfo& SourceInfo : SourceInfos)
+			{
+				if (SourceInfo.AudioBusId == InAudioBusId)
+				{
+					SourceInfo.bIsPlaying = false;
+					SourceInfo.bIsPaused = false;
+					SourceInfo.bIsActive = false;
+					SourceInfo.bIsStopping = false;
+				}
+			}
+		});
+	}
+
+	void FMixerSourceManager::StopAudioBus(uint32 InAudioBusId)
+	{
+		if (!AudioBusIds_AudioThread.Contains(InAudioBusId))
+		{
+			return;
+		}
+
+		AudioBusIds_AudioThread.Remove(InAudioBusId);
+
+		AudioMixerThreadCommand([this, InAudioBusId]()
+		{
+			TSharedPtr<FMixerAudioBus>* AudioBusPtr = AudioBuses.Find(InAudioBusId);
+			if (AudioBusPtr)
+			{
+				if (!(*AudioBusPtr)->IsAutomatic())
+				{
+					// Immediately stop all sources which were source buses
+					for (FSourceInfo& SourceInfo : SourceInfos)
+					{
+						if (SourceInfo.AudioBusId == InAudioBusId)
+						{
+							SourceInfo.bIsPlaying = false;
+							SourceInfo.bIsPaused = false;
+							SourceInfo.bIsActive = false;
+							SourceInfo.bIsStopping = false;
+						}
+					}
+					AudioBuses.Remove(InAudioBusId);
+				}
+			}
+		});
+	}
+
+	bool FMixerSourceManager::IsAudioBusActive(uint32 InAudioBusId)
+	{
+		AUDIO_MIXER_CHECK_GAME_THREAD(MixerDevice);
+		return AudioBusIds_AudioThread.Contains(InAudioBusId);
+	}
+
+	FPatchOutputStrongPtr FMixerSourceManager::AddPatchForAudioBus(uint32 InAudioBusId, float PatchGain)
+	{
+		AUDIO_MIXER_CHECK_AUDIO_PLAT_THREAD(MixerDevice);
+		TSharedPtr<FMixerAudioBus>* AudioBusPtr = AudioBuses.Find(InAudioBusId);
+		if (AudioBusPtr)
+		{
+			return (*AudioBusPtr)->AddNewPatch(NumOutputFrames * 2, PatchGain);
+		}
+		return nullptr;
+	}
+
 	void FMixerSourceManager::Play(const int32 SourceId)
 	{
 		AUDIO_MIXER_CHECK(SourceId < NumTotalSources);
@@ -1214,45 +1310,49 @@ namespace Audio
 		});
 	}
 
-	void FMixerSourceManager::SetBusSendInfo(const int32 SourceId, EBusSendType InBusSendType, FMixerBusSend& InBusSend)
+	void FMixerSourceManager::SetBusSendInfo(const int32 SourceId, EBusSendType InAudioBusSendType, uint32 AudioBusId, float BusSendLevel)
 	{
 		AUDIO_MIXER_CHECK(SourceId < NumTotalSources);
 		AUDIO_MIXER_CHECK(GameThreadInfo.bIsBusy[SourceId]);
 		AUDIO_MIXER_CHECK_GAME_THREAD(MixerDevice);
 
-		AudioMixerThreadCommand([this, SourceId, InBusSendType, InBusSend]()
+		AudioMixerThreadCommand([this, SourceId, InAudioBusSendType, AudioBusId, BusSendLevel]()
 		{
 			// Create mapping of source id to bus send level
-			FBusSend BusSend{ SourceId, InBusSend.SendLevel };
+			FAudioBusSend BusSend;
+			BusSend.SendLevel = BusSendLevel;
+			BusSend.SourceId = SourceId;
+			BusSend.AudioBusId = AudioBusId;
+
 			FSourceInfo& SourceInfo = SourceInfos[SourceId];
 
 			// Retrieve the bus we want to send audio to
-			FMixerBus* Bus = Buses.Find(InBusSend.BusId);
+			TSharedPtr<FMixerAudioBus>* AudioBusPtr = AudioBuses.Find(AudioBusId);
 
 			// If we already have a bus, we update the amount of audio we want to send to it
-			if (Bus)
+			if (AudioBusPtr)
 			{
-				Bus->AddBusSend(InBusSendType, BusSend);			
+				(*AudioBusPtr)->AddSend(InAudioBusSendType, BusSend);
 			}
 			else
 			{
 				// If the bus is not registered, make a new entry on the send
-				FMixerBus NewBusData(this, SourceInfo.NumInputChannels, NumOutputFrames);
+				TSharedPtr<FMixerAudioBus> NewBusData(new FMixerAudioBus(this, true, SourceInfo.NumInputChannels));
 
 				// Add a send to it. This will not have a bus instance id (i.e. won't output audio), but 
 				// we register the send anyway in the event that this bus does play, we'll know to send this
 				// source's audio to it.
-				NewBusData.AddBusSend(InBusSendType, BusSend);
+				NewBusData->AddSend(InAudioBusSendType, BusSend);
 
-				Buses.Add(InBusSend.BusId, NewBusData);
+				AudioBuses.Add(AudioBusId, NewBusData);
 			}
 
 			// Check to see if we need to create new bus data. If we are not playing a bus with this id, then we
 			// need to create a slot for it such that when a bus does play, it'll start rendering audio from this source
 			bool bExisted = false;
-			for (uint32 BusId : SourceInfo.BusSends[(int32)InBusSendType])
+			for (uint32 BusId : SourceInfo.AudioBusSends[(int32)InAudioBusSendType])
 			{
-				if (BusId == InBusSend.BusId)
+				if (BusId == AudioBusId)
 				{
 					bExisted = true;
 					break;
@@ -1261,7 +1361,7 @@ namespace Audio
 
 			if (!bExisted)
 			{
-				SourceInfo.BusSends[(int32)InBusSendType].Add(InBusSend.BusId);
+				SourceInfo.AudioBusSends[(int32)InAudioBusSendType].Add(AudioBusId);
 			}
 		});
 	}
@@ -1434,8 +1534,8 @@ namespace Audio
 				continue;
 			}
 
-			const bool bIsBus = SourceInfo.BusId != INDEX_NONE;
-			if ((bGenerateBuses && !bIsBus) || (!bGenerateBuses && bIsBus))
+			const bool bIsSourceBus = SourceInfo.AudioBusId != INDEX_NONE;
+			if ((bGenerateBuses && !bIsSourceBus) || (!bGenerateBuses && bIsSourceBus))
 			{
 				continue;
 			}
@@ -1457,27 +1557,30 @@ namespace Audio
 
 			// if this is a bus, we just want to copy the bus audio to this source's output audio
 			// Note we need to copy this since bus instances may have different audio via dynamic source effects, etc.
-			if (bIsBus)
+			if (bIsSourceBus)
 			{
-				// Get the source's rendered bus data
-				const FMixerBus* Bus = Buses.Find(SourceInfo.BusId);
-				const float* RESTRICT BusBufferPtr = Bus->GetCurrentBusBuffer();
-
-				int32 NumFramesPlayed = NumOutputFrames;
-				if (SourceInfo.BusDurationFrames != INDEX_NONE)
+				// Get the source's rendered and mixed audio bus data
+				const TSharedPtr<FMixerAudioBus>* AudioBusPtr = AudioBuses.Find(SourceInfo.AudioBusId);
+				if (ensure(AudioBusPtr != nullptr))
 				{
-					// If we're now finishing, only copy over the real data
-					if ((SourceInfo.NumFramesPlayed + NumOutputFrames) >= SourceInfo.BusDurationFrames)
+					const float* RESTRICT BusBufferPtr = (*AudioBusPtr)->GetCurrentBusBuffer();
+
+					int32 NumFramesPlayed = NumOutputFrames;
+					if (SourceInfo.SourceBusDurationFrames != INDEX_NONE)
 					{
-						NumFramesPlayed = SourceInfo.BusDurationFrames - SourceInfo.NumFramesPlayed;
-						SourceInfo.bIsLastBuffer = true;				
+						// If we're now finishing, only copy over the real data
+						if ((SourceInfo.NumFramesPlayed + NumOutputFrames) >= SourceInfo.SourceBusDurationFrames)
+						{
+							NumFramesPlayed = SourceInfo.SourceBusDurationFrames - SourceInfo.NumFramesPlayed;
+							SourceInfo.bIsLastBuffer = true;
+						}
 					}
-				}
 
-				SourceInfo.NumFramesPlayed += NumFramesPlayed;
+					SourceInfo.NumFramesPlayed += NumFramesPlayed;
 
-				// Simply copy into the pre distance attenuation buffer ptr
-				FMemory::Memcpy(PreDistanceAttenBufferPtr, BusBufferPtr, sizeof(float) * NumFramesPlayed * SourceInfo.NumInputChannels);
+					// Simply copy into the pre distance attenuation buffer ptr
+					FMemory::Memcpy(PreDistanceAttenBufferPtr, BusBufferPtr, sizeof(float) * NumFramesPlayed * SourceInfo.NumInputChannels);
+				}
 			}
 			else
 			{
@@ -1571,20 +1674,20 @@ namespace Audio
 	void FMixerSourceManager::ComputeBuses()
 	{
 		// Loop through the bus registry and mix source audio
-		for (auto& Entry : Buses)
+		for (auto& Entry : AudioBuses)
 		{
-			FMixerBus& Bus = Entry.Value;
-			Bus.MixBuffer();
+			TSharedPtr<FMixerAudioBus>& AudioBus = Entry.Value;
+			AudioBus->MixBuffer();
 		}
 	}
 
 	void FMixerSourceManager::UpdateBuses()
 	{
 		// Update the bus states post mixing. This flips the current/previous buffer indices.
-		for (auto& Entry : Buses)
+		for (auto& Entry : AudioBuses)
 		{
-			FMixerBus& Bus = Entry.Value;
-			Bus.Update();
+			TSharedPtr<FMixerAudioBus>& AudioBus = Entry.Value;
+			AudioBus->Update();
 		}
 	}
 
@@ -2031,8 +2134,8 @@ namespace Audio
 				continue;
 			}
 
-			const bool bIsBus = SourceInfo.BusId != INDEX_NONE;
-			if ((bGenerateBuses && !bIsBus) || (!bGenerateBuses && bIsBus))
+			const bool bIsSourceBus = SourceInfo.AudioBusId != INDEX_NONE;
+			if ((bGenerateBuses && !bIsSourceBus) || (!bGenerateBuses && bIsSourceBus))
 			{
 				continue;
 			}
@@ -2041,7 +2144,7 @@ namespace Audio
 
 			// Copy and store the current state of the pre-distance attenuation buffer before we feed it through our source effects
 			// This is used by pre-effect sends
-			if (SourceInfo.BusSends[(int32)EBusSendType::PreEffect].Num() > 0)
+			if (SourceInfo.AudioBusSends[(int32)EBusSendType::PreEffect].Num() > 0)
 			{
 				SourceInfo.PreEffectBuffer.Reset();
 				SourceInfo.PreEffectBuffer.Reserve(SourceInfo.PreDistanceAttenuationBuffer.Num());
@@ -2227,8 +2330,8 @@ namespace Audio
 
 			// If we're in generate buses mode and not a bus, or vice versa, or if we're set to only output audio to buses.
 			// If set to output buses, no need to do any panning for the source. The buses will do the panning.
-			const bool bIsBus = SourceInfo.BusId != INDEX_NONE;
-			if ((bGenerateBuses && !bIsBus) || (!bGenerateBuses && bIsBus) || SourceInfo.bOutputToBusOnly)
+			const bool bIsSourceBus = SourceInfo.AudioBusId != INDEX_NONE;
+			if ((bGenerateBuses && !bIsSourceBus) || (!bGenerateBuses && bIsSourceBus) || SourceInfo.bOutputToBusOnly)
 			{
 				continue;
 			}
@@ -2276,7 +2379,7 @@ namespace Audio
 	void FMixerSourceManager::GenerateSourceAudio(const bool bGenerateBuses)
 	{
 		// If there are no buses, don't need to do anything here
-		if (bGenerateBuses && !Buses.Num())
+		if (bGenerateBuses && !AudioBuses.Num())
 		{
 			return;
 		}
@@ -2396,6 +2499,7 @@ namespace Audio
 			FSoundEffectSourceInitData InitData;
 			InitData.AudioClock = MixerDevice->GetAudioClock();
 			InitData.SampleRate = MixerDevice->SampleRate;
+			InitData.AudioDeviceId = MixerDevice->DeviceID;
 
 			for (int32 SourceId = 0; SourceId < NumTotalSources; ++SourceId)
 			{
@@ -2456,11 +2560,23 @@ namespace Audio
 		return SourceInfos[SourceId].PreEffectBuffer.GetData();
 	}
 
-	const float* FMixerSourceManager::GetPreviousBusBuffer(const int32 SourceId) const
+	const float* FMixerSourceManager::GetPreviousSourceBusBuffer(const int32 SourceId) const
 	{
-		const uint32 BusId = SourceInfos[SourceId].BusId;
-		const FMixerBus* MixerBus = Buses.Find(BusId);
-		return MixerBus->GetPreviousBusBuffer();
+		if (SourceId < SourceInfos.Num())
+		{
+			return GetPreviousAudioBusBuffer(SourceInfos[SourceId].AudioBusId);
+		}
+		return nullptr;
+	}
+
+	const float* FMixerSourceManager::GetPreviousAudioBusBuffer(const int32 AudioBusId) const
+	{
+		const TSharedPtr<FMixerAudioBus>* AudioBusPtr = AudioBuses.Find(AudioBusId);
+		if (AudioBusPtr)
+		{
+			return (*AudioBusPtr)->GetPreviousBusBuffer();
+		}
+		return nullptr;
 	}
 
 	int32 FMixerSourceManager::GetNumChannels(const int32 SourceId) const
@@ -2468,9 +2584,9 @@ namespace Audio
 		return SourceInfos[SourceId].NumInputChannels;
 	}
 
-	bool FMixerSourceManager::IsBus(const int32 SourceId) const
+	bool FMixerSourceManager::IsSourceBus(const int32 SourceId) const
 	{
-		return SourceInfos[SourceId].BusId != INDEX_NONE;
+		return SourceInfos[SourceId].AudioBusId != INDEX_NONE;
 	}
 
 	void FMixerSourceManager::ComputeNextBlockOfSamples()
diff --git a/Engine/Source/Runtime/AudioMixer/Private/AudioMixerSourceManager.h b/Engine/Source/Runtime/AudioMixer/Private/AudioMixerSourceManager.h
--- a/Engine/Source/Runtime/AudioMixer/Private/AudioMixerSourceManager.h
+++ b/Engine/Source/Runtime/AudioMixer/Private/AudioMixerSourceManager.h
@@ -36,15 +36,10 @@ namespace Audio
 		AlignedFloatBuffer AudioData;
 
 		/** How many times this buffer will loop. */
-		int32 LoopCount;
+		int32 LoopCount = 0;
 
 		/** If this buffer is from real-time decoding and needs to make callbacks for more data. */
-		uint32 bRealTimeBuffer : 1;
-
-		FMixerSourceVoiceBuffer()
-		{
-			FMemory::Memzero(this, sizeof(*this));
-		}
+		bool bRealTimeBuffer = false;
 	};
 
 
@@ -71,79 +66,53 @@ namespace Audio
 		FMixerSubmixWeakPtr Submix;
 
 		// The amount of audio that is to be mixed into this submix
-		float SendLevel;
+		float SendLevel = 0.0f;
 
 		// Whather or not this is the primary send (i.e. first in the send chain)
-		bool bIsMainSend;
+		bool bIsMainSend = false;
 
 		// If this is a soundfield submix, this is a pointer to the submix's Soundfield Factory.
 		// If this is nullptr, the submix is not a soundfield submix.
-		ISoundfieldFactory* SoundfieldFactory;
+		ISoundfieldFactory* SoundfieldFactory = nullptr;
 	};
 
 	// Struct holding mappings of bus ids (unique ids) to send level
-	struct FMixerBusSend
+	struct FInitAudioBusSend
 	{
-		uint32 BusId;
-		float SendLevel;
+		uint32 AudioBusId = INDEX_NONE;
+		float SendLevel = 0.0f;
 	};
 
 	struct FMixerSourceVoiceInitParams
 	{
-		TSharedPtr<FMixerSourceBuffer> MixerSourceBuffer;
-		ISourceListener* SourceListener;
+		TSharedPtr<FMixerSourceBuffer> MixerSourceBuffer = nullptr;
+		ISourceListener* SourceListener = nullptr;
 		TArray<FMixerSourceSubmixSend> SubmixSends;
-		TArray<FMixerBusSend> BusSends[(int32)EBusSendType::Count];
-		uint32 BusId;
-		float BusDuration;
-		uint32 SourceEffectChainId;
+		TArray<FInitAudioBusSend> AudioBusSends[(int32)EBusSendType::Count];
+		uint32 AudioBusId = INDEX_NONE;
+		float SourceBusDuration = 0.0f;
+		uint32 SourceEffectChainId = INDEX_NONE;
 		TArray<FSourceEffectChainEntry> SourceEffectChain;
-		FMixerSourceVoice* SourceVoice;
-		int32 NumInputChannels;
-		int32 NumInputFrames;
-		float EnvelopeFollowerAttackTime;
-		float EnvelopeFollowerReleaseTime;
+		FMixerSourceVoice* SourceVoice = nullptr;
+		int32 NumInputChannels = 0;
+		int32 NumInputFrames = 0;
+		float EnvelopeFollowerAttackTime = 10.0f;
+		float EnvelopeFollowerReleaseTime = 100.0f;
 		FString DebugName;
-		USpatializationPluginSourceSettingsBase* SpatializationPluginSettings;
-		UOcclusionPluginSourceSettingsBase* OcclusionPluginSettings;
-		UReverbPluginSourceSettingsBase* ReverbPluginSettings;
-		USoundModulationPluginSourceSettingsBase* ModulationPluginSettings;
+		USpatializationPluginSourceSettingsBase* SpatializationPluginSettings = nullptr;
+		UOcclusionPluginSourceSettingsBase* OcclusionPluginSettings = nullptr;
+		UReverbPluginSourceSettingsBase* ReverbPluginSettings = nullptr;
+		USoundModulationPluginSourceSettingsBase* ModulationPluginSettings = nullptr;
 		FName AudioComponentUserID;
-		uint64 AudioComponentID;
-		uint8 bPlayEffectChainTails : 1;
-		uint8 bUseHRTFSpatialization : 1;
-		uint8 bIsExternalSend : 1;
-		uint8 bIsDebugMode : 1;
-		uint8 bOutputToBusOnly : 1;
-		uint8 bIsVorbis : 1;
-		uint8 bIsAmbisonics : 1;
-		uint8 bIsSeeking : 1;
-
-		FMixerSourceVoiceInitParams()
-			: MixerSourceBuffer(nullptr)
-			, SourceListener(nullptr)
-			, BusId(INDEX_NONE)
-			, BusDuration(0.0f)
-			, SourceEffectChainId(INDEX_NONE)
-			, SourceVoice(nullptr)
-			, NumInputChannels(0)
-			, NumInputFrames(0)
-			, EnvelopeFollowerAttackTime(10.0f)
-			, EnvelopeFollowerReleaseTime(100.0f)
-			, SpatializationPluginSettings(nullptr)
-			, OcclusionPluginSettings(nullptr)
-			, ReverbPluginSettings(nullptr)
-			, ModulationPluginSettings(nullptr)
-			, AudioComponentID(0)
-			, bPlayEffectChainTails(false)
-			, bUseHRTFSpatialization(false)
-			, bIsExternalSend(false)
-			, bIsDebugMode(false)
-			, bOutputToBusOnly(false)
-			, bIsVorbis(false)
-			, bIsAmbisonics(false)
-			, bIsSeeking(false)
-		{}
+		uint64 AudioComponentID = 0;
+		bool bPlayEffectChainTails = false;
+		bool bUseHRTFSpatialization = false;
+		bool bIsExternalSend = false;
+		bool bIsDebugMode  = false;
+		bool bOutputToBusOnly = false;
+		bool bIsVorbis = false;
+		bool bIsAmbisonics = false;
+		bool bIsSeeking = false;
 	};
 
 	struct FSourceChannelMap
@@ -153,8 +122,8 @@ namespace Audio
 
 		// This is the number of bytes the gain array is using:
 		// (Number of input channels * number of output channels) * sizeof float.
-		int32 CopySize;
-		bool bIsInit;
+		int32 CopySize = 0;
+		bool bIsInit = false;
 
 		FSourceChannelMap(int32 InNumInChannels, int32 InNumOutChannels) 
 			: CopySize(InNumInChannels * InNumOutChannels * sizeof(float))
@@ -201,10 +170,10 @@ namespace Audio
 	struct FSourceManagerInitParams
 	{
 		// Total number of sources to use in the source manager
-		int32 NumSources;
+		int32 NumSources = 0;
 
 		// Number of worker threads to use for the source manager.
-		int32 NumSourceWorkers;
+		int32 NumSourceWorkers = 0;
 	};
 
 	class FMixerSourceManager
@@ -218,11 +187,17 @@ namespace Audio
 
 		bool GetFreeSourceId(int32& OutSourceId);
 		int32 GetNumActiveSources() const;
-		int32 GetNumActiveBuses() const;
+		int32 GetNumActiveAudioBuses() const;
 
 		void ReleaseSourceId(const int32 SourceId);
 		void InitSource(const int32 SourceId, const FMixerSourceVoiceInitParams& InitParams);
 
+		// Creates an audio bus manually. Returns an audio bus Id.
+		void StartAudioBus(uint32 InAudioBusId, int32 InNumChannels, bool bInIsAutomatic);
+		void StopAudioBus(uint32 InAudioBusId);
+		bool IsAudioBusActive(uint32 InAudioBusId);
+		FPatchOutputStrongPtr AddPatchForAudioBus(uint32 InAudioBusId, float PatchGain);
+
 		void Play(const int32 SourceId);
 		void Stop(const int32 SourceId);
 		void StopFade(const int32 SourceId, const int32 NumFrames);
@@ -253,7 +228,8 @@ namespace Audio
 		const ISoundfieldAudioPacket* GetEncodedOutput(const int32 SourceId, const FSoundfieldEncodingKey& InKey) const;
 
 		void SetSubmixSendInfo(const int32 SourceId, const FMixerSourceSubmixSend& SubmixSend);
-		void SetBusSendInfo(const int32 SourceId, EBusSendType InBusSendType, FMixerBusSend& BusSend);
+
+		void SetBusSendInfo(const int32 SourceId, EBusSendType InAudioBusSendType, uint32 AudiobusId, float BusSendLevel);
 
 		void UpdateDeviceChannelCount(const int32 InNumOutputChannels);
 
@@ -262,10 +238,11 @@ namespace Audio
 
 		const float* GetPreDistanceAttenuationBuffer(const int32 SourceId) const;
 		const float* GetPreEffectBuffer(const int32 SourceId) const;
-		const float* GetPreviousBusBuffer(const int32 SourceId) const;
+		const float* GetPreviousSourceBusBuffer(const int32 SourceId) const;
+		const float* GetPreviousAudioBusBuffer(const int32 AudioBusId) const;
 		int32 GetNumChannels(const int32 SourceId) const;
 		int32 GetNumOutputFrames() const { return NumOutputFrames; }
-		bool IsBus(const int32 SourceId) const;
+		bool IsSourceBus(const int32 SourceId) const;
 		void PumpCommandQueue();
 		void UpdatePendingReleaseData(bool bForceWait = false);
 		void FlushCommandQueue(bool bPumpCommandQueue = false);
@@ -460,14 +437,14 @@ namespace Audio
 
 			TArray<FMixerSourceSubmixSend> SubmixSends;
 
-			// What bus Id this source is, if it is a bus. This is INDEX_NONE for sources which are not buses.
-			uint32 BusId;
+			// What audio bus Id this source is sonfiying, if it is a source bus. This is INDEX_NONE for sources which are not source buses.
+			uint32 AudioBusId;
 
-			// Number of samples to count for bus
-			int64 BusDurationFrames;
+			// Number of samples to count for source bus
+			int64 SourceBusDurationFrames;
 
 			// What buses this source is sending its audio to. Used to remove this source from the bus send list.
-			TArray<uint32> BusSends[(int32)EBusSendType::Count];
+			TArray<uint32> AudioBusSends[(int32)EBusSendType::Count];
 
 			// Interpolated source params
 			FParam PitchSourceParam;
@@ -570,8 +547,9 @@ namespace Audio
 		// sources.
 		TArray<FSourceDownmixData> DownmixDataArray;
 
-		// Map of bus object Id's to bus data. 
-		TMap<uint32, FMixerBus> Buses;
+		// Map of bus object Id's to audio bus data. 
+		TMap<uint32, TSharedPtr<FMixerAudioBus>> AudioBuses;
+		TArray<uint32> AudioBusIds_AudioThread;
 
 		// Async task workers for processing sources in parallel
 		TArray<FAsyncTask<FAudioMixerSourceWorker>*> SourceWorkers;
diff --git a/Engine/Source/Runtime/AudioMixer/Private/AudioMixerSourceVoice.cpp b/Engine/Source/Runtime/AudioMixer/Private/AudioMixerSourceVoice.cpp
--- a/Engine/Source/Runtime/AudioMixer/Private/AudioMixerSourceVoice.cpp
+++ b/Engine/Source/Runtime/AudioMixer/Private/AudioMixerSourceVoice.cpp
@@ -60,7 +60,7 @@ namespace Audio
 			AUDIO_MIXER_CHECK(InitParams.NumInputChannels > 0);
 
 			bOutputToBusOnly = InitParams.bOutputToBusOnly;
-			bIsBus = InitParams.BusId != INDEX_NONE;
+			bIsBus = InitParams.AudioBusId != INDEX_NONE;
 
 			for (int32 i = 0; i < InitParams.SubmixSends.Num(); ++i)
 			{
@@ -296,11 +296,11 @@ namespace Audio
 		}
 	}
 
-	void FMixerSourceVoice::SetBusSendInfo(EBusSendType InBusSendType, FMixerBusSend& BusSend)
+	void FMixerSourceVoice::SetAudioBusSendInfo(EBusSendType InBusSendType, uint32 AudioBusId, float BusSendLevel)
 	{
 		AUDIO_MIXER_CHECK_GAME_THREAD(MixerDevice);
 
-		SourceManager->SetBusSendInfo(SourceId, InBusSendType, BusSend);
+		SourceManager->SetBusSendInfo(SourceId, InBusSendType, AudioBusId, BusSendLevel);
 	}
 
 	void FMixerSourceVoice::OnMixBus(FMixerSourceVoiceBuffer* OutMixerSourceBuffer)
diff --git a/Engine/Source/Runtime/AudioMixer/Private/AudioMixerSourceVoice.h b/Engine/Source/Runtime/AudioMixer/Private/AudioMixerSourceVoice.h
--- a/Engine/Source/Runtime/AudioMixer/Private/AudioMixerSourceVoice.h
+++ b/Engine/Source/Runtime/AudioMixer/Private/AudioMixerSourceVoice.h
@@ -105,7 +105,7 @@ namespace Audio
 		void SetSubmixSendInfo(FMixerSubmixWeakPtr Submix, const float SendLevel);
 
 		// Set the source bus send levels
-		void SetBusSendInfo(EBusSendType InBusSendType, FMixerBusSend& BusSend);
+		void SetAudioBusSendInfo(EBusSendType InBusSendType, uint32 AudioBusId, float BusSendLevel);
 
 		// Called when the source is a bus and needs to mix other sources together to generate output
 		void OnMixBus(FMixerSourceVoiceBuffer* OutMixerSourceBuffer);
diff --git a/Engine/Source/Runtime/AudioMixer/Private/SoundWaveDecoder.cpp b/Engine/Source/Runtime/AudioMixer/Private/SoundWaveDecoder.cpp
--- a/Engine/Source/Runtime/AudioMixer/Private/SoundWaveDecoder.cpp
+++ b/Engine/Source/Runtime/AudioMixer/Private/SoundWaveDecoder.cpp
@@ -481,7 +481,7 @@ namespace Audio
 		}
 
 
-		if (InitData.SoundWave->bIsBus || InitData.SoundWave->bProcedural)
+		if (InitData.SoundWave->bIsSourceBus || InitData.SoundWave->bProcedural)
 		{
 			UE_LOG(LogAudioMixer, Warning, TEXT("Sound wave decoder does not support buses or procedural sounds."));
 			return false;
diff --git a/Engine/Source/Runtime/AudioMixer/Public/AudioMixerBlueprintLibrary.h b/Engine/Source/Runtime/AudioMixer/Public/AudioMixerBlueprintLibrary.h
--- a/Engine/Source/Runtime/AudioMixer/Public/AudioMixerBlueprintLibrary.h
+++ b/Engine/Source/Runtime/AudioMixer/Public/AudioMixerBlueprintLibrary.h
@@ -6,6 +6,7 @@
 #include "Kismet/BlueprintFunctionLibrary.h"
 #include "SubmixEffects/AudioMixerSubmixEffectDynamicsProcessor.h"
 #include "Sound/SoundEffectSource.h"
+#include "Sound/AudioBus.h"
 #include "SampleBuffer.h"
 #include "Sound/SoundCue.h"
 #include "Sound/SoundSubmixSend.h"
@@ -96,19 +97,19 @@ public:
 	static void ClearMasterSubmixEffects(const UObject* WorldContextObject);
 
 	/** Start recording audio. By leaving the Submix To Record field blank, you can record the master output of the game. */
-	UFUNCTION(BlueprintCallable, Category = "Audio", meta = (WorldContext = "WorldContextObject", AdvancedDisplay = 1))
+	UFUNCTION(BlueprintCallable, Category = "Audio|Recording", meta = (WorldContext = "WorldContextObject", AdvancedDisplay = 1))
 	static void StartRecordingOutput(const UObject* WorldContextObject, float ExpectedDuration, USoundSubmix* SubmixToRecord = nullptr);
 	
 	/** Stop recording audio. Path can be absolute, or relative (to the /Saved/BouncedWavFiles folder). By leaving the Submix To Record field blank, you can record the master output of the game.  */
-	UFUNCTION(BlueprintCallable, Category = "Audio", meta = (WorldContext = "WorldContextObject", DisplayName = "Finish Recording Output", AdvancedDisplay = 4))
+	UFUNCTION(BlueprintCallable, Category = "Audio|Recording", meta = (WorldContext = "WorldContextObject", DisplayName = "Finish Recording Output", AdvancedDisplay = 4))
 	static USoundWave* StopRecordingOutput(const UObject* WorldContextObject, EAudioRecordingExportType ExportType, const FString& Name, FString Path, USoundSubmix* SubmixToRecord = nullptr, USoundWave* ExistingSoundWaveToOverwrite= nullptr);
 
 	/** Pause recording audio, without finalizing the recording to disk. By leaving the Submix To Record field blank, you can record the master output of the game. */
-	UFUNCTION(BlueprintCallable, Category = "Audio", meta = (WorldContext = "WorldContextObject", AdvancedDisplay = 1))
+	UFUNCTION(BlueprintCallable, Category = "Audio|Recording", meta = (WorldContext = "WorldContextObject", AdvancedDisplay = 1))
 	static void PauseRecordingOutput(const UObject* WorldContextObject, USoundSubmix* SubmixToPause = nullptr);
 
 	/** Resume recording audio after pausing. By leaving the Submix To Pause field blank, you can record the master output of the game. */
-	UFUNCTION(BlueprintCallable, Category = "Audio", meta = (WorldContext = "WorldContextObject", AdvancedDisplay = 1))
+	UFUNCTION(BlueprintCallable, Category = "Audio|Recording", meta = (WorldContext = "WorldContextObject", AdvancedDisplay = 1))
 	static void ResumeRecordingOutput(const UObject* WorldContextObject, USoundSubmix* SubmixToPause = nullptr);
 
 	/** Start spectrum analysis of the audio output. By leaving the Submix To Analyze blank, you can analyze the master output of the game. */
@@ -144,17 +145,30 @@ public:
 	static int32 GetNumberOfEntriesInSourceEffectChain(const UObject* WorldContextObject, USoundEffectSourcePresetChain* PresetChain);
 
 	/** Begin loading a sound into the cache so that it can be played immediately. */
-	UFUNCTION(BlueprintCallable, Category = "Sound")
+	UFUNCTION(BlueprintCallable, Category = "Audio|Cache")
 	static void PrimeSoundForPlayback(USoundWave* SoundWave, const FOnSoundLoadComplete OnLoadCompletion);
 
 	/** Begin loading any sounds referenced by a sound cue into the cache so that it can be played immediately. */
-	UFUNCTION(BlueprintCallable, Category = "Sound")
+	UFUNCTION(BlueprintCallable, Category = "Audio|Cache")
 	static void PrimeSoundCueForPlayback(USoundCue* SoundCue);
 
 	/** Trim memory used by the audio cache. Returns the number of megabytes freed. */
-	UFUNCTION(BlueprintCallable, Category = "Sound")
+	UFUNCTION(BlueprintCallable, Category = "Audio|Cache")
 	static float TrimAudioCache(float InMegabytesToFree);
 
+	/** Starts the given audio bus. */
+	UFUNCTION(BlueprintCallable, Category = "Audio|Bus", meta = (WorldContext = "WorldContextObject"))
+	static void StartAudioBus(const UObject* WorldContextObject, UAudioBus* AudioBus);
+
+	/** Stops the given audio bus. */
+	UFUNCTION(BlueprintCallable, Category = "Audio|Bus", meta = (WorldContext = "WorldContextObject"))
+	static void StopAudioBus(const UObject* WorldContextObject, UAudioBus* AudioBus);
+
+	/** Queries if the given audio bus is active (and audio can be mixed to it). */
+	UFUNCTION(BlueprintCallable, Category = "Audio|Bus", meta = (WorldContext = "WorldContextObject"))
+	static bool IsAudioBusActive(const UObject* WorldContextObject, UAudioBus* AudioBus);
+
+
 private:
 	static void PopulateSpectrumAnalyzerSettings(EFFTSize FFTSize, EFFTPeakInterpolationMethod InterpolationMethod, EFFTWindowType WindowType, float HopSize, Audio::FSpectrumAnalyzerSettings &OutSettings);
 };
diff --git a/Engine/Source/Runtime/AudioMixer/Public/AudioMixerBus.h b/Engine/Source/Runtime/AudioMixer/Public/AudioMixerBus.h
--- a/Engine/Source/Runtime/AudioMixer/Public/AudioMixerBus.h
+++ b/Engine/Source/Runtime/AudioMixer/Public/AudioMixerBus.h
@@ -4,24 +4,45 @@
 
 #include "CoreMinimal.h"
 #include "Audio.h"
+#include "DSP/MultithreadedPatching.h"
 
 namespace Audio
 {
 	class FMixerSourceManager;
 
 	// Struct holding mappings of runtime source ids (bus instances) to bus send level
-	struct FBusSend
+	struct FAudioBusSend
 	{
-		int32 SourceId;
-		float SendLevel;
+		int32 SourceId = INDEX_NONE;
+		uint32 AudioBusId = INDEX_NONE;
+		float SendLevel = 0.0f;
 	};
 
 	// Bus instance data. Holds source id bus instances and bus sends data
-	class FMixerBus
+	class FMixerAudioBus
 	{
 	public:
-		FMixerBus(FMixerSourceManager* SourceManager, const int32 InNumChannels, const int32 InNumFrames);
+		// Allow anybody to add a patch to retrieve the audio from any thread.
+		FPatchOutputStrongPtr AddNewPatch(int32 MaxLatencyInSamples, float InGain);
+
+	private:
+
+		// Creates an audio bus.
+		// SourceMnager		The owning source manager object.
+		// bInIsAutomatic	Whether or not this audio bus was created automatically via source buses.
+		// InNumChannels	The number of channels of the source bus.
+		// InNumFrames		The number of frames to mix per mix buffer call. I.e. source manager render size.
+		FMixerAudioBus(FMixerSourceManager* SourceManager, bool bInIsAutomatic, int32 InNumChannels);
 		
+		// Sets whether or not this audio bus is automatic.
+		void SetAutomatic(bool bInIsAutomatic) { bIsAutomatic = bInIsAutomatic; }
+
+		// Returns if this is a manual audio bus vs automatic
+		bool IsAutomatic() const { return bIsAutomatic; }
+
+		// Returns the number of channels of the audio bus
+		int32 GetNumChannels() const { return NumChannels; }
+
 		// Update the mixer bus after a render block
 		void Update();
 
@@ -32,10 +53,10 @@ namespace Audio
 		bool RemoveInstanceId(const int32 InSourceId);
 
 		// Adds a buss end to the bus
-		void AddBusSend(EBusSendType BusSendType, const FBusSend& InBusSend);
+		void AddSend(EBusSendType BusSendType, const FAudioBusSend& InBusSend);
 
 		// Removes the source instance from this bus's send list
-		bool RemoveBusSend(EBusSendType BusSendType, const int32 InSourceId);
+		bool RemoveSend(EBusSendType BusSendType, const int32 InSourceId);
 
 		// Gets the current mixed bus buffer
 		const float* GetCurrentBusBuffer() const;
@@ -46,8 +67,6 @@ namespace Audio
 		// Compute the mixed buffer
 		void MixBuffer();
 
-	private:
-
 		// If this bus was constructed before
 		void SetNumOutputChannels(int32 InNumOutputChannels);
 
@@ -58,7 +77,7 @@ namespace Audio
 		TArray<int32> InstanceIds;
 
 		// Bus sends to this instance
-		TArray<FBusSend> BusSends[(int32)EBusSendType::Count];
+		TArray<FAudioBusSend> AudioBusSends[(int32)EBusSendType::Count];
 
 		// The mixed source data. This is double-buffered to allow buses to send audio to themselves.
 		// Buses feed audio to each other by storing their previous buffer. Current buses mix in previous other buses (including themselves)
@@ -76,7 +95,12 @@ namespace Audio
 		// Owning soruce manager
 		FMixerSourceManager* SourceManager;
 
-		FMixerBus();
+		Audio::FPatchSplitter PatchSplitter;
+
+		// Was created manually, not via source buses.
+		bool bIsAutomatic;
+
+		friend FMixerSourceManager;
 	};
 
 }
diff --git a/Engine/Source/Runtime/AudioMixer/Public/AudioMixerDevice.h b/Engine/Source/Runtime/AudioMixer/Public/AudioMixerDevice.h
--- a/Engine/Source/Runtime/AudioMixer/Public/AudioMixerDevice.h
+++ b/Engine/Source/Runtime/AudioMixer/Public/AudioMixerDevice.h
@@ -7,7 +7,7 @@
 #include "AudioDevice.h"
 #include "Sound/SoundSubmix.h"
 #include "DSP/BufferVectorOperations.h"
-
+#include "DSP/MultithreadedPatching.h"
 
 // Forward Declarations
 class FOnSubmixEnvelopeBP;
@@ -213,6 +213,13 @@ namespace Audio
 
 		static bool IsEndpointSubmix(const USoundSubmixBase* InSubmix);
 
+		// Audio bus API
+		void StartAudioBus(uint32 InAudioBusId, int32 InNumChannels, bool bInIsAutomatic);
+		void StopAudioBus(uint32 InAudioBusId);
+		bool IsAudioBusActive(uint32 InAudioBusId);
+		FPatchOutputStrongPtr AddPatchForAudioBus(uint32 InAudioBusId, float PatchGain);
+
+
 	protected:
 
 		virtual void InitSoundSubmixes() override;
@@ -260,6 +267,9 @@ namespace Audio
 		TArray<USoundSubmix*> MasterSubmixes;
 		TArray<FMixerSubmixPtr> MasterSubmixInstances;
 
+		// The active audio bus list accessible on the game thread
+		TArray<int32> ActiveAudioBuses_GameThread;
+
 		/** Ptr to the platform interface, which handles streaming audio to the hardware device. */
 		IAudioMixerPlatformInterface* AudioMixerPlatform;
 		
diff --git a/Engine/Source/Runtime/Engine/Classes/Components/AudioComponent.h b/Engine/Source/Runtime/Engine/Classes/Components/AudioComponent.h
--- a/Engine/Source/Runtime/Engine/Classes/Components/AudioComponent.h
+++ b/Engine/Source/Runtime/Engine/Classes/Components/AudioComponent.h
@@ -501,6 +501,16 @@ public:
 	UFUNCTION(BlueprintCallable, Category = "Audio|Components|Audio")
 	void SetSourceBusSendPostEffect(USoundSourceBus* SoundSourceBus, float SourceBusSendLevel);
 
+	/** Sets how much audio the sound should send to the given Audio Bus (PRE Source Effects).
+	if the Bus Send doesn't already exist, it will be added to the overrides on the active sound */
+	UFUNCTION(BlueprintCallable, Category = "Audio|Components|Audio")
+	void SetAudioBusSendPreEffect(UAudioBus* AudioBus, float AudioBusSendLevel);
+
+	/** Sets how much audio the sound should send to the given Audio Bus (POST Source Effects).
+		if the Audio Bus Send doesn't already exist, it will be added to the overrides on the active sound */
+	UFUNCTION(BlueprintCallable, Category = "Audio|Components|Audio")
+	void SetAudioBusSendPostEffect(UAudioBus* AudioBus, float SourceBusSendLevel);
+
 	/** Sets whether or not the low pass filter is enabled on the audio component. */
 	UFUNCTION(BlueprintCallable, Category = "Audio|Components|Audio")
 	void SetLowPassFilterEnabled(bool InLowPassFilterEnabled);
@@ -558,6 +568,8 @@ private:
 	/** Whether or not the sound is audible. */
 	bool IsInAudibleRange(float* OutMaxDistance) const;
 
+	void SetBusSendffectInternal(USoundSourceBus* InSourceBus, UAudioBus* InAudioBus, float SendLevel, EBusSendType InBusSendType);
+
 public:
 
 	/** Sets the sound instance parameter. */
diff --git a/Engine/Source/Runtime/Engine/Classes/Sound/AudioBus.h b/Engine/Source/Runtime/Engine/Classes/Sound/AudioBus.h
new file mode 100644
--- /dev/null
+++ b/Engine/Source/Runtime/Engine/Classes/Sound/AudioBus.h
@@ -0,0 +1,48 @@
+// Copyright Epic Games, Inc. All Rights Reserved.
+#pragma once
+
+#include "CoreMinimal.h"
+#include "UObject/ObjectMacros.h"
+#include "UObject/Object.h"
+#include "AudioBus.generated.h"
+
+class FAudioDevice;
+
+// The number of channels to mix audio into the source bus
+UENUM(BlueprintType)
+enum class EAudioBusChannels : uint8
+{
+	Mono,
+	Stereo,
+};
+
+// Function to retrieve an audio bus buffer given a handle
+// static float* GetAudioBusBuffer(const FAudioBusHandle& AudioBusHandle);
+
+// An audio bus is an object which represents an audio patch chord. Audio can be sent to it. It can be sonified using USoundSourceBuses.
+// Instances of the audio bus are created in the audio engine. 
+UCLASS(hidecategories= (Compression, SoundWave, Streaming, Subtitles, Sound, Info, ImportSettings), ClassGroup = Sound, meta = (BlueprintSpawnableComponent))
+class ENGINE_API UAudioBus : public UObject
+{
+	GENERATED_UCLASS_BODY()
+
+public:
+
+	/** How many channels to use for the source bus. */
+	UPROPERTY(EditAnywhere, Category = BusProperties)
+	EAudioBusChannels AudioBusChannels;
+
+	/** If the audio bus can be instantiated and destroyed automatically when sources send audio to it. If this audio bus is manually started, it will override this value to be false, meaning you will need to stop the audio bus manually.*/
+	UPROPERTY(EditAnywhere, Category = BusProperties)
+	bool bIsAutomatic;
+
+	//~ Begin UObject
+	virtual void BeginDestroy() override;
+	//~ End UObject
+
+	//~ Begin UObject Interface.
+#if WITH_EDITOR
+	virtual void PostEditChangeProperty(FPropertyChangedEvent& PropertyChangedEvent) override;
+#endif // WITH_EDITOR
+	//~ End UObject Interface.
+};
diff --git a/Engine/Source/Runtime/Engine/Classes/Sound/SoundEffectSource.h b/Engine/Source/Runtime/Engine/Classes/Sound/SoundEffectSource.h
--- a/Engine/Source/Runtime/Engine/Classes/Sound/SoundEffectSource.h
+++ b/Engine/Source/Runtime/Engine/Classes/Sound/SoundEffectSource.h
@@ -67,11 +67,15 @@ struct FSoundEffectSourceInitData
 	// The object id of the parent preset
 	uint32 ParentPresetUniqueId;
 
+	// The audio device ID of the audio device instance this source instance was created from
+	uint32 AudioDeviceId;
+
 	FSoundEffectSourceInitData()
 	: SampleRate(0.0f)
 	, NumSourceChannels(0)
 	, AudioClock(0.0)
 	, ParentPresetUniqueId(INDEX_NONE)
+	, AudioDeviceId(INDEX_NONE)
 	{}
 };
 
diff --git a/Engine/Source/Runtime/Engine/Classes/Sound/SoundSourceBus.h b/Engine/Source/Runtime/Engine/Classes/Sound/SoundSourceBus.h
--- a/Engine/Source/Runtime/Engine/Classes/Sound/SoundSourceBus.h
+++ b/Engine/Source/Runtime/Engine/Classes/Sound/SoundSourceBus.h
@@ -6,6 +6,7 @@
 #include "UObject/Object.h"
 #include "Sound/SoundWaveProcedural.h"
 #include "Sound/SoundSourceBusSend.h"
+#include "Sound/AudioBus.h"
 #include "SoundSourceBus.generated.h"
 
 class USoundSourceBus;
@@ -21,14 +22,14 @@ enum class ESourceBusChannels : uint8
 };
 
 // A source bus is a type of USoundBase and can be "played" like any sound.
-UCLASS(hidecategories= (Compression, SoundWave, Streaming, Subtitles, Sound, Info, ImportSettings), ClassGroup = Sound, meta = (BlueprintSpawnableComponent))
+UCLASS(hidecategories= (Compression, SoundWave, Streaming, Subtitles, Sound, Analysis, Format, Loading, Info, ImportSettings), ClassGroup = Sound, meta = (BlueprintSpawnableComponent))
 class ENGINE_API USoundSourceBus : public USoundWave
 {
 	GENERATED_UCLASS_BODY()
 
 public:
 
-	/** How many channels to use for the source bus. */
+	/** How many channels to use for the source bus if the audio bus is not specified, otherwise it will use the audio bus object's channel count. */
 	UPROPERTY(EditAnywhere, Category = BusProperties)
 	ESourceBusChannels SourceBusChannels;
 
@@ -36,6 +37,10 @@ public:
 	UPROPERTY(EditAnywhere, Category = BusProperties, meta = (UIMin = 0.0, ClampMin = 0.0))
 	float SourceBusDuration;
 
+	/** Audio bus to use to use as audio for this source bus. This source bus will sonify the audio from the audio bus. */
+	UPROPERTY(EditAnywhere, Category = BusProperties)
+	UAudioBus* AudioBus;
+
 	/** Stop the source bus when the volume goes to zero. */
 	UPROPERTY(EditAnywhere, Category = BusProperties)
 	uint32 bAutoDeactivateWhenSilent:1;
diff --git a/Engine/Source/Runtime/Engine/Classes/Sound/SoundSourceBusSend.h b/Engine/Source/Runtime/Engine/Classes/Sound/SoundSourceBusSend.h
--- a/Engine/Source/Runtime/Engine/Classes/Sound/SoundSourceBusSend.h
+++ b/Engine/Source/Runtime/Engine/Classes/Sound/SoundSourceBusSend.h
@@ -7,6 +7,7 @@
 #include "SoundSourceBusSend.generated.h"
 
 class USoundSourceBus;
+class UAudioBus;
 
 UENUM(BlueprintType)
 enum class ESourceBusSendLevelControlMethod : uint8
@@ -31,40 +32,45 @@ struct ENGINE_API FSoundSourceBusSendInfo
 		Linear: Interpolate between Min and Max Send Levels based on listener distance (between Distance Min and Distance Max)
 		Custom Curve: Use the float curve to map Send Level to distance (0.0-1.0 on curve maps to Distance Min - Distance Max)
 	*/
-	UPROPERTY(EditAnywhere, Category = SourceBusSend)
+	UPROPERTY(EditAnywhere, Category = BusSend)
 	ESourceBusSendLevelControlMethod SourceBusSendLevelControlMethod;
 
-	// The Source Bus to send the audio to
-	UPROPERTY(EditAnywhere, Category = SourceBusSend)
+	// A source Bus to send the audio to. Source buses sonify (make audible) the audio sent to it and are themselves sounds which take up a voice slot in the audio engine.
+	UPROPERTY(EditAnywhere, Category = BusSend)
 	USoundSourceBus* SoundSourceBus;
 
-	// The amount of audio to send to the source bus
-	UPROPERTY(EditAnywhere, Category = SourceBusSend)
+	// An audio bus to send the audio to. Audio buses can be used to route audio to DSP effects or other purposes. E.g. side-chaining, analysis, etc. Audio buses are not audible unless hooked up to a source bus.
+	UPROPERTY(EditAnywhere, Category = BusSend)
+	UAudioBus* AudioBus;
+
+	// The amount of audio to send to the bus.
+	UPROPERTY(EditAnywhere, Category = BusSend)
 	float SendLevel;
 
-	// The amount to send to the Source Bus when sound is located at a distance equal to value specified in the min send distance.
-	UPROPERTY(EditAnywhere, BlueprintReadWrite, Category = SourceBusSend)
+	// The amount to send to the bus when sound is located at a distance equal to value specified in the min send distance.
+	UPROPERTY(EditAnywhere, BlueprintReadWrite, Category = BusSend)
 	float MinSendLevel;
 
-	// The amount to send to the Source Bus when sound is located at a distance equal to value specified in the max send distance.
-	UPROPERTY(EditAnywhere, BlueprintReadWrite, Category = SourceBusSend)
+	// The amount to send to the bus when sound is located at a distance equal to value specified in the max send distance.
+	UPROPERTY(EditAnywhere, BlueprintReadWrite, Category = BusSend)
 	float MaxSendLevel;
 
-	// The distance at which the Min Send Level is sent to the source bus
-	UPROPERTY(EditAnywhere, BlueprintReadWrite, Category = SourceBusSend)
+	// The distance at which the min send Level is sent to the bus
+	UPROPERTY(EditAnywhere, BlueprintReadWrite, Category = BusSend)
 	float MinSendDistance;
 
-	// The distance at which the Max Send Level is sent to the source bus
-	UPROPERTY(EditAnywhere, BlueprintReadWrite, Category = SourceBusSend)
+	// The distance at which the max send level is sent to the bus
+	UPROPERTY(EditAnywhere, BlueprintReadWrite, Category = BusSend)
 	float MaxSendDistance;
 
-	// The custom curve to use for distance-based Source Bus send level.
-	UPROPERTY(EditAnywhere, BlueprintReadWrite, Category = SourceBusSend)
+	// The custom curve to use for distance-based bus send level.
+	UPROPERTY(EditAnywhere, BlueprintReadWrite, Category = BusSend)
 	FRuntimeFloatCurve CustomSendLevelCurve;
 
 	FSoundSourceBusSendInfo()
 		: SourceBusSendLevelControlMethod(ESourceBusSendLevelControlMethod::Manual)
 		, SoundSourceBus(nullptr)
+		, AudioBus(nullptr)
 		, SendLevel(0.0f)
 		, MinSendLevel(0.0f)
 		, MaxSendLevel(1.0f)
diff --git a/Engine/Source/Runtime/Engine/Classes/Sound/SoundWave.h b/Engine/Source/Runtime/Engine/Classes/Sound/SoundWave.h
--- a/Engine/Source/Runtime/Engine/Classes/Sound/SoundWave.h
+++ b/Engine/Source/Runtime/Engine/Classes/Sound/SoundWave.h
@@ -340,7 +340,7 @@ public:
 	uint8 bPlayingProcedural : 1;
 
 	/** Set to true of this is a bus sound source. This will result in the sound wave not generating audio for itself, but generate audio through instances. Used only in audio mixer. */
-	uint8 bIsBus:1;
+	uint8 bIsSourceBus:1;
 
 	/** Set to true for procedural waves that can be processed asynchronously. */
 	uint8 bCanProcessAsync:1;
diff --git a/Engine/Source/Runtime/Engine/Private/ActiveSound.cpp b/Engine/Source/Runtime/Engine/Private/ActiveSound.cpp
--- a/Engine/Source/Runtime/Engine/Private/ActiveSound.cpp
+++ b/Engine/Source/Runtime/Engine/Private/ActiveSound.cpp
@@ -343,20 +343,20 @@ void FActiveSound::SetSubmixSend(const FSoundSubmixSendInfo& SubmixSendInfo)
 	SoundSubmixSendsOverride.Add(SubmixSendInfo);
 }
 
-void FActiveSound::SetSourceBusSend(EBusSendType BusSendType, const FSoundSourceBusSendInfo& SourceBusSendInfo)
+void FActiveSound::SetSourceBusSend(EBusSendType BusSendType, const FSoundSourceBusSendInfo& SendInfo)
 {
 	// Override send level if the source bus send is already included in active sound
-	for (FSoundSourceBusSendInfo& Info : SoundSourceBusSendsOverride[(int32)BusSendType])
+	for (FSoundSourceBusSendInfo& Info : BusSendsOverride[(int32)BusSendType])
 	{
-		if (Info.SoundSourceBus == SourceBusSendInfo.SoundSourceBus)
+		if (Info.SoundSourceBus == SendInfo.SoundSourceBus || Info.AudioBus == SendInfo.AudioBus)
 		{
-			Info.SendLevel = SourceBusSendInfo.SendLevel;
+			Info.SendLevel = SendInfo.SendLevel;
 			return;
 		}
 	}
 
 	// Otherwise, add it to the source bus send overrides
-	SoundSourceBusSendsOverride[(int32)BusSendType].Add(SourceBusSendInfo);
+	BusSendsOverride[(int32)BusSendType].Add(SendInfo);
 }
 
 void FActiveSound::Stop()
@@ -398,7 +398,7 @@ void FActiveSound::GetSoundSubmixSends(TArray<FSoundSubmixSendInfo>& OutSends) c
 	}
 }
 
-void FActiveSound::GetSoundSourceBusSends(EBusSendType BusSendType, TArray<FSoundSourceBusSendInfo>& OutSends) const
+void FActiveSound::GetBusSends(EBusSendType BusSendType, TArray<FSoundSourceBusSendInfo>& OutSends) const
 {
 	if (Sound)
 	{
@@ -406,12 +406,12 @@ void FActiveSound::GetSoundSourceBusSends(EBusSendType BusSendType, TArray<FSoun
 		Sound->GetSoundSourceBusSends(BusSendType, OutSends);
 
 		// Loop through the overrides, which may append or override the existing send
-		for (const FSoundSourceBusSendInfo& SendInfo : SoundSourceBusSendsOverride[(int32)BusSendType])
+		for (const FSoundSourceBusSendInfo& SendInfo : BusSendsOverride[(int32)BusSendType])
 		{
 			bool bOverridden = false;
 			for (FSoundSourceBusSendInfo& OutSendInfo : OutSends)
 			{
-				if (OutSendInfo.SoundSourceBus == SendInfo.SoundSourceBus)
+				if (OutSendInfo.SoundSourceBus == SendInfo.SoundSourceBus || OutSendInfo.AudioBus == SendInfo.AudioBus)
 				{
 					OutSendInfo.SendLevel = SendInfo.SendLevel;
 					bOverridden = true;
@@ -535,7 +535,7 @@ void FActiveSound::UpdateWaveInstances(TArray<FWaveInstance*> &InWaveInstances,
 
 	for (int32 BusSendType = 0; BusSendType < (int32)EBusSendType::Count; ++BusSendType)
 	{
-		GetSoundSourceBusSends((EBusSendType)BusSendType, ParseParams.SoundSourceBusSends[BusSendType]);
+		GetBusSends((EBusSendType)BusSendType, ParseParams.BusSends[BusSendType]);
 	}
 
 	// Set up the base source effect chain.
diff --git a/Engine/Source/Runtime/Engine/Private/Audio.cpp b/Engine/Source/Runtime/Engine/Private/Audio.cpp
--- a/Engine/Source/Runtime/Engine/Private/Audio.cpp
+++ b/Engine/Source/Runtime/Engine/Private/Audio.cpp
@@ -1078,7 +1078,7 @@ bool FWaveInstance::IsSeekable() const
 		return false;
 	}
 
-	if (WaveData->bIsBus || WaveData->bProcedural)
+	if (WaveData->bIsSourceBus || WaveData->bProcedural)
 	{
 		return false;
 	}
diff --git a/Engine/Source/Runtime/Engine/Private/AudioBus.cpp b/Engine/Source/Runtime/Engine/Private/AudioBus.cpp
new file mode 100644
--- /dev/null
+++ b/Engine/Source/Runtime/Engine/Private/AudioBus.cpp
@@ -0,0 +1,50 @@
+// Copyright Epic Games, Inc. All Rights Reserved.
+
+
+#include "Sound/AudioBus.h"
+#include "AudioDeviceManager.h"
+#include "EngineGlobals.h"
+#include "Engine/Engine.h"
+#include "UObject/UObjectIterator.h"
+#include "ActiveSound.h"
+#include "AudioDevice.h"
+#include "AudioMixerDevice.h"
+
+UAudioBus::UAudioBus(const FObjectInitializer& ObjectInitializer)
+	: Super(ObjectInitializer)
+{
+}
+
+void UAudioBus::BeginDestroy()
+{
+	Super::BeginDestroy();
+
+	if (!GEngine)
+	{
+		return;
+	}
+
+	// Make sure we stop all audio bus instances on all devices if this object is getting destroyed
+	uint32 AudioBusId = GetUniqueID();
+
+	FAudioDeviceManager* AudioDeviceManager = GEngine->GetAudioDeviceManager();
+	if (AudioDeviceManager)
+	{
+		TArray<FAudioDevice*> AudioDevices = AudioDeviceManager->GetAudioDevices();
+		for (FAudioDevice* AudioDevice : AudioDevices)
+		{
+			if (AudioDevice->IsAudioMixerEnabled())
+			{
+				Audio::FMixerDevice* MixerDevice = static_cast<Audio::FMixerDevice*>(AudioDevice);
+				MixerDevice->StopAudioBus(AudioBusId);
+			}
+		}
+	}
+}
+
+#if WITH_EDITOR
+void UAudioBus::PostEditChangeProperty(FPropertyChangedEvent& PropertyChangedEvent)
+{
+}
+#endif
+
diff --git a/Engine/Source/Runtime/Engine/Private/AudioDevice.cpp b/Engine/Source/Runtime/Engine/Private/AudioDevice.cpp
--- a/Engine/Source/Runtime/Engine/Private/AudioDevice.cpp
+++ b/Engine/Source/Runtime/Engine/Private/AudioDevice.cpp
@@ -5816,7 +5816,7 @@ void FAudioDevice::Precache(USoundWave* SoundWave, bool bSynchronous, bool bTrac
 		// Procedurally created audio
 		SoundWave->DecompressionType = DTYPE_Procedural;
 	}
-	else if (SoundWave->bIsBus)
+	else if (SoundWave->bIsSourceBus)
 	{
 		// Audio data which will be generated by instanced objects, not from the sound wave asset
 		if (IsAudioMixerEnabled())
diff --git a/Engine/Source/Runtime/Engine/Private/Components/AudioComponent.cpp b/Engine/Source/Runtime/Engine/Private/Components/AudioComponent.cpp
--- a/Engine/Source/Runtime/Engine/Private/Components/AudioComponent.cpp
+++ b/Engine/Source/Runtime/Engine/Private/Components/AudioComponent.cpp
@@ -1317,46 +1317,45 @@ void UAudioComponent::SetSubmixSend(USoundSubmixBase* Submix, float SendLevel)
 	}
 }
 
-// BP function to set source bus sends (pre effect)
-void UAudioComponent::SetSourceBusSendPreEffect(USoundSourceBus* SoundSourceBus, float SourceBusSendLevel)
+void UAudioComponent::SetBusSendffectInternal(USoundSourceBus* InSourceBus, UAudioBus* InAudioBus, float SendLevel, EBusSendType InBusSendType)
 {
 	if (FAudioDevice* AudioDevice = GetAudioDevice())
 	{
 		const uint64 MyAudioComponentID = AudioComponentID;
-		FAudioThread::RunCommandOnAudioThread([AudioDevice, MyAudioComponentID, SoundSourceBus, SourceBusSendLevel]()
+		FAudioThread::RunCommandOnAudioThread([AudioDevice, MyAudioComponentID, InSourceBus, InAudioBus, SendLevel, InBusSendType]()
 		{
 			FActiveSound* ActiveSound = AudioDevice->FindActiveSound(MyAudioComponentID);
 			if (ActiveSound)
 			{
 				FSoundSourceBusSendInfo SourceBusSendInfo;
-				SourceBusSendInfo.SoundSourceBus = SoundSourceBus;
-				SourceBusSendInfo.SendLevel = SourceBusSendLevel;
+				SourceBusSendInfo.SoundSourceBus = InSourceBus;
+				SourceBusSendInfo.AudioBus = InAudioBus;
+				SourceBusSendInfo.SendLevel = SendLevel;
 
-				ActiveSound->SetSourceBusSend(EBusSendType::PreEffect, SourceBusSendInfo);
+				ActiveSound->SetSourceBusSend(InBusSendType, SourceBusSendInfo);
 			}
 		});
 	}
 }
 
-// BP function to set source bus sends (post effect)
+void UAudioComponent::SetSourceBusSendPreEffect(USoundSourceBus* SoundSourceBus, float SourceBusSendLevel)
+{
+	SetBusSendffectInternal(SoundSourceBus, nullptr, SourceBusSendLevel, EBusSendType::PreEffect);
+}
+
 void UAudioComponent::SetSourceBusSendPostEffect(USoundSourceBus * SoundSourceBus, float SourceBusSendLevel)
 {
-	if (FAudioDevice* AudioDevice = GetAudioDevice())
-	{
-		const uint64 MyAudioComponentID = AudioComponentID;
-		FAudioThread::RunCommandOnAudioThread([AudioDevice, MyAudioComponentID, SoundSourceBus, SourceBusSendLevel]()
-		{
-			FActiveSound* ActiveSound = AudioDevice->FindActiveSound(MyAudioComponentID);
-			if (ActiveSound)
-			{
-				FSoundSourceBusSendInfo SourceBusSendInfo;
-				SourceBusSendInfo.SoundSourceBus = SoundSourceBus;
-				SourceBusSendInfo.SendLevel = SourceBusSendLevel;
+	SetBusSendffectInternal(SoundSourceBus, nullptr, SourceBusSendLevel, EBusSendType::PostEffect);
+}
 
-				ActiveSound->SetSourceBusSend(EBusSendType::PostEffect, SourceBusSendInfo);
-			}
-		});
-	}
+void UAudioComponent::SetAudioBusSendPreEffect(UAudioBus* AudioBus, float AudioBusSendLevel)
+{
+	SetBusSendffectInternal(nullptr, AudioBus, AudioBusSendLevel, EBusSendType::PreEffect);
+}
+
+void UAudioComponent::SetAudioBusSendPostEffect(UAudioBus* AudioBus, float AudioBusSendLevel)
+{
+	SetBusSendffectInternal(nullptr, AudioBus, AudioBusSendLevel, EBusSendType::PostEffect);
 }
 
 void UAudioComponent::SetLowPassFilterEnabled(bool InLowPassFilterEnabled)
diff --git a/Engine/Source/Runtime/Engine/Private/SoundSourceBus.cpp b/Engine/Source/Runtime/Engine/Private/SoundSourceBus.cpp
--- a/Engine/Source/Runtime/Engine/Private/SoundSourceBus.cpp
+++ b/Engine/Source/Runtime/Engine/Private/SoundSourceBus.cpp
@@ -12,7 +12,7 @@ USoundSourceBus::USoundSourceBus(const FObjectInitializer& ObjectInitializer)
 	: Super(ObjectInitializer)
 {
 	// This is a bus. This will result in the decompression type to be set as DTYPE_Bus. Audio won't be generated from this object but from instance data in audio mixer.
-	bIsBus = true;
+	bIsSourceBus = true;
 
 	Init();
 }
diff --git a/Engine/Source/Runtime/Engine/Private/SoundWave.cpp b/Engine/Source/Runtime/Engine/Private/SoundWave.cpp
--- a/Engine/Source/Runtime/Engine/Private/SoundWave.cpp
+++ b/Engine/Source/Runtime/Engine/Private/SoundWave.cpp
@@ -1903,7 +1903,7 @@ void USoundWave::Parse(FAudioDevice* AudioDevice, const UPTRINT NodeWaveInstance
 
 	for (int32 BusSendType = 0; BusSendType < (int32)EBusSendType::Count; ++BusSendType)
 	{
-		WaveInstance->SoundSourceBusSends[BusSendType] = ParseParams.SoundSourceBusSends[BusSendType];
+		WaveInstance->BusSends[BusSendType] = ParseParams.BusSends[BusSendType];
 	}
 
 	// Pass along plugin settings to the wave instance
diff --git a/Engine/Source/Runtime/Engine/Public/ActiveSound.h b/Engine/Source/Runtime/Engine/Public/ActiveSound.h
--- a/Engine/Source/Runtime/Engine/Public/ActiveSound.h
+++ b/Engine/Source/Runtime/Engine/Public/ActiveSound.h
@@ -132,7 +132,7 @@ struct FSoundParseParameters
 	TArray<FSoundSubmixSendInfo> SoundSubmixSends;
 
 	// The source bus sends to use
-	TArray<FSoundSourceBusSendInfo> SoundSourceBusSends[(int32)EBusSendType::Count];
+	TArray<FSoundSourceBusSendInfo> BusSends[(int32)EBusSendType::Count];
 
 	// Reverb wet-level parameters
 	EReverbSendMethod ReverbSendMethod;
@@ -350,7 +350,7 @@ private:
 	TArray<FSoundSubmixSendInfo> SoundSubmixSendsOverride;
 
 	/** Optional override for the source bus sends for the sound. */
-	TArray<FSoundSourceBusSendInfo> SoundSourceBusSendsOverride[(int32)EBusSendType::Count];
+	TArray<FSoundSourceBusSendInfo> BusSendsOverride[(int32)EBusSendType::Count];
 
 	TMap<UPTRINT, FWaveInstance*> WaveInstances;
 
@@ -639,7 +639,7 @@ public:
 	void GetSoundSubmixSends(TArray<FSoundSubmixSendInfo>& OutSends) const;
 
 	/** Gets the sound source bus sends to use for this sound instance. */
-	void GetSoundSourceBusSends(EBusSendType BusSendType, TArray<FSoundSourceBusSendInfo>& OutSends) const;
+	void GetBusSends(EBusSendType BusSendType, TArray<FSoundSourceBusSendInfo>& OutSends) const;
 
 	/* Determines which of the provided listeners is the closest to the sound */
 	int32 FindClosestListener( const TArray<struct FListener>& InListeners ) const;
diff --git a/Engine/Source/Runtime/Engine/Public/Audio.h b/Engine/Source/Runtime/Engine/Public/Audio.h
--- a/Engine/Source/Runtime/Engine/Public/Audio.h
+++ b/Engine/Source/Runtime/Engine/Public/Audio.h
@@ -179,8 +179,8 @@ struct ENGINE_API FWaveInstance
 	/** Sound submix sends */
 	TArray<FSoundSubmixSendInfo> SoundSubmixSends;
 
-	/** The sound source bus sends. */
-	TArray<FSoundSourceBusSendInfo> SoundSourceBusSends[(int32)EBusSendType::Count];
+	/** The source bus and/or audio bus sends. */
+	TArray<FSoundSourceBusSendInfo> BusSends[(int32)EBusSendType::Count];
 
 	/** Sound effect chain */
 	USoundEffectSourcePresetChain* SourceEffectChain;
diff --git a/Engine/Source/Runtime/SignalProcessing/Private/EnvelopeFollower.cpp b/Engine/Source/Runtime/SignalProcessing/Private/EnvelopeFollower.cpp
--- a/Engine/Source/Runtime/SignalProcessing/Private/EnvelopeFollower.cpp
+++ b/Engine/Source/Runtime/SignalProcessing/Private/EnvelopeFollower.cpp
@@ -2,6 +2,7 @@
 
 #include "DSP/EnvelopeFollower.h"
 #include "DSP/Dsp.h"
+#include "DSP/BufferVectorOperations.h"
 
 namespace Audio
 {
@@ -76,15 +77,7 @@ namespace Audio
 
 	float FEnvelopeFollower::ProcessAudio(const float InAudioSample)
 	{
-		// Take the absolute value of the input sample
-		float Sample = FMath::Abs(InAudioSample);
-
-		// If we're not Peak detecting, then square the input
-		if (EnvMode != EPeakMode::Peak)
-		{
-			Sample = Sample * Sample;
-		}
-
+		float Sample = (EnvMode != EPeakMode::Peak) ? InAudioSample * InAudioSample : FMath::Abs(InAudioSample);
 		float TimeSamples = (Sample > CurrentEnvelopeValue) ? AttackTimeSamples : ReleaseTimeSamples;
 		float NewEnvelopeValue = TimeSamples * (CurrentEnvelopeValue - Sample) + Sample;;
 		NewEnvelopeValue = Audio::UnderflowClamp(NewEnvelopeValue);
@@ -94,17 +87,29 @@ namespace Audio
 		return CurrentEnvelopeValue = NewEnvelopeValue;
 	}
 
-	float FEnvelopeFollower::ProcessAudioNonClamped(const float InAudioSample)
+	float FEnvelopeFollower::ProcessAudio(const float* InAudioBuffer, int32 InNumSamples)
 	{
-		// Take the absolute value of the input sample
-		float Sample = FMath::Abs(InAudioSample);
+		for (int32 SampleIndex = 0; SampleIndex < InNumSamples; ++SampleIndex)
+		{
+			ProcessAudioNonClamped(InAudioBuffer[SampleIndex]);
+		}
+		return FMath::Clamp(CurrentEnvelopeValue, 0.0f, 1.0f);
+	}
 
-		// If we're not Peak detecting, then square the input
-		if (EnvMode != EPeakMode::Peak)
+	float FEnvelopeFollower::ProcessAudio(const float* InAudioBuffer, float* OutAudioBuffer, int32 InNumSamples)
+	{
+		for (int32 SampleIndex = 0; SampleIndex < InNumSamples; ++SampleIndex)
 		{
-			Sample = Sample * Sample;
+			OutAudioBuffer[SampleIndex] = ProcessAudioNonClamped(InAudioBuffer[SampleIndex]);
 		}
 
+		Audio::BufferRangeClampFast(OutAudioBuffer, InNumSamples, 0.0f, 1.0f);
+		return CurrentEnvelopeValue;
+	}
+
+	float FEnvelopeFollower::ProcessAudioNonClamped(const float InAudioSample)
+	{
+		float Sample = (EnvMode != EPeakMode::Peak) ? InAudioSample * InAudioSample : FMath::Abs(InAudioSample);
 		float TimeSamples = (Sample > CurrentEnvelopeValue) ? AttackTimeSamples : ReleaseTimeSamples;
 		float NewEnvelopeValue = TimeSamples * (CurrentEnvelopeValue - Sample) + Sample;;
 		NewEnvelopeValue = Audio::UnderflowClamp(NewEnvelopeValue);
diff --git a/Engine/Source/Runtime/SignalProcessing/Private/RingModulation.cpp b/Engine/Source/Runtime/SignalProcessing/Private/RingModulation.cpp
--- a/Engine/Source/Runtime/SignalProcessing/Private/RingModulation.cpp
+++ b/Engine/Source/Runtime/SignalProcessing/Private/RingModulation.cpp
@@ -2,6 +2,7 @@
 
 #include "DSP/RingModulation.h"
 #include "DSP/Dsp.h"
+#include "DSP/BufferVectorOperations.h"
 
 namespace Audio
 {
@@ -10,9 +11,10 @@ namespace Audio
 		, ModulationDepth(0.5f)
 		, DryLevel(0.0f)
 		, WetLevel(1.0f)
+		, Scale(1.0f)
 		, NumChannels(0)
 	{
-
+		UpdateScale();
 	}
 
 	FRingModulation::~FRingModulation()
@@ -30,6 +32,22 @@ namespace Audio
 		NumChannels = InNumChannels;
 	}
 
+	void FRingModulation::UpdateScale()
+	{
+		Scale = WetLevel * ModulationDepth;
+	}
+
+	void FRingModulation::SetExternalPatchSource(Audio::FPatchOutputStrongPtr InPatch)
+	{
+		Patch = InPatch;
+	}
+
+	void FRingModulation::SetWetLevel(const float InWetLevel)
+	{
+		WetLevel = InWetLevel;
+		UpdateScale();
+	}
+
 	void FRingModulation::SetModulatorWaveType(const EOsc::Type InType)
 	{
 		Osc.SetType(InType);
@@ -44,23 +62,37 @@ namespace Audio
 	void FRingModulation::SetModulationDepth(const float InModulationDepth)
 	{
 		ModulationDepth = FMath::Clamp(InModulationDepth, -1.0f, 1.0f);
+		UpdateScale();
 	}
 
-	void FRingModulation::ProcessAudioFrame(const float* InFrame, float* OutFrame)
+	void FRingModulation::ProcessAudio(const float* InBuffer, const int32 InNumSamples, float* OutBuffer)
 	{
-		float OscOut = Osc.Generate();
-		for (int32 Channel = 0; Channel < NumChannels; ++Channel)
+		if (ModulationBuffer.Num() != InNumSamples * NumChannels)
 		{
-			OutFrame[Channel] = DryLevel * InFrame[Channel] + WetLevel * InFrame[Channel] * OscOut * ModulationDepth;
+			ModulationBuffer.Reset();
+			ModulationBuffer.AddZeroed(InNumSamples * NumChannels);
 		}
-	}
 
-	void FRingModulation::ProcessAudio(const float* InBuffer, const int32 InNumSamples, float* OutBuffer)
-	{
-		for (int32 SampleIndex = 0; SampleIndex < InNumSamples; SampleIndex += NumChannels)
+		// If we have an external patch source to modulate against, copy that data into the modulation buffer
+		if (Patch.IsValid())
 		{
-			ProcessAudioFrame(&InBuffer[SampleIndex], &OutBuffer[SampleIndex]);
+			Patch->PopAudio(ModulationBuffer.GetData(), ModulationBuffer.Num(), true);
 		}
+		else
+		{
+			// Write the oscillator data into the modulation buffer
+			float* ModulationBufferPtr = ModulationBuffer.GetData();
+			for (int32 SampleIndex = 0; SampleIndex < InNumSamples; ++SampleIndex)
+			{
+				ModulationBufferPtr[SampleIndex] = Osc.Generate();
+			}
+		}
+
+		// Multiply the input buffer by the modulation buffer in-place
+		Audio::MultiplyBuffersInPlace(InBuffer, ModulationBuffer.GetData(), InNumSamples);
+
+		// Perform a buffer weighted sum of the modulation buffer with the Scale value and the dry level as weights
+		Audio::BufferWeightedSumFast(InBuffer, DryLevel, ModulationBuffer.GetData(), Scale, OutBuffer, InNumSamples);
 	}
 
 
diff --git a/Engine/Source/Runtime/SignalProcessing/Public/DSP/EnvelopeFollower.h b/Engine/Source/Runtime/SignalProcessing/Public/DSP/EnvelopeFollower.h
--- a/Engine/Source/Runtime/SignalProcessing/Public/DSP/EnvelopeFollower.h
+++ b/Engine/Source/Runtime/SignalProcessing/Public/DSP/EnvelopeFollower.h
@@ -50,6 +50,12 @@ namespace Audio
 		// Processes the input audio stream and returns the envelope value.
 		float ProcessAudio(const float InAudioSample);
 
+		// Process the input audio buffer and returns the last envelope value
+		float ProcessAudio(const float* InAudioBuffer, float* OutAudioBuffer, int32 InNumSamples);
+
+		// Process the input audio buffer and returns the last envelope value
+		float ProcessAudio(const float* InAudioBuffer, int32 InNumSamples);
+
 		// Processes the input audio stream and returns the envelope value.
 		float ProcessAudioNonClamped(const float InAudioSample);
 
diff --git a/Engine/Source/Runtime/SignalProcessing/Public/DSP/RingModulation.h b/Engine/Source/Runtime/SignalProcessing/Public/DSP/RingModulation.h
--- a/Engine/Source/Runtime/SignalProcessing/Public/DSP/RingModulation.h
+++ b/Engine/Source/Runtime/SignalProcessing/Public/DSP/RingModulation.h
@@ -4,6 +4,7 @@
 
 #include "CoreMinimal.h"
 #include "DSP/Osc.h"
+#include "DSP/MultithreadedPatching.h"
 
 namespace Audio
 {
@@ -30,25 +31,30 @@ namespace Audio
 		// Set the ring modulation depth
 		void SetModulationDepth(const float InModulationDepth);
 
+		// Sets that the modulation buffer is external
+		void SetExternalPatchSource(Audio::FPatchOutputStrongPtr InPatch);
+
 		// Set the dry level of the ring modulation
 		void SetDryLevel(const float InDryLevel) { DryLevel = InDryLevel; }
 
 		// Set the wet level of the ring modulation
-		void SetWetLevel(const float InWetLevel) { WetLevel = InWetLevel; }
-
-		// Processes the audio frame (audio frame must have channels equal to that used during initialization)
-		void ProcessAudioFrame(const float* InFrame, float* OutFrame);
+		void SetWetLevel(const float InWetLevel);
 
 		// Process audio buffer
 		void ProcessAudio(const float* InBuffer, const int32 InNumSamples, float* OutBuffer);
 
 	private:
+		void UpdateScale();
+
+		Audio::FPatchOutputStrongPtr Patch;
 		Audio::FOsc Osc;
 		float ModulationFrequency;
 		float ModulationDepth;
 		float DryLevel;
 		float WetLevel;
+		float Scale;
 		int32 NumChannels;
+		TArray<float> ModulationBuffer;
 	};
 
 }
