Fix compile error

diff --git a/Engine/Source/Runtime/AudioMixer/Private/AudioMixerSource.cpp b/Engine/Source/Runtime/AudioMixer/Private/AudioMixerSource.cpp
--- a/Engine/Source/Runtime/AudioMixer/Private/AudioMixerSource.cpp
+++ b/Engine/Source/Runtime/AudioMixer/Private/AudioMixerSource.cpp
@@ -330,18 +330,18 @@ namespace Audio
 			for (FSoundSourceBusSendInfo& SendInfo : InWaveInstance->BusSends[BusSendType])
 			{
 				// Avoid redoing duplicate code for sending audio to source bus or audio bus. Most of it is the same other than the bus id.
-				auto InitBusSend = [this](FMixerSourceVoiceInitParams& InInitParams, const FSoundSourceBusSendInfo& InSendInfo, int32 BusSendType, uint32 InBusId)
+				auto InitBusSend = [this](FMixerSourceVoiceInitParams& InInitParams, const FSoundSourceBusSendInfo& InSendInfo, int32 InBusSendType, uint32 InBusId)
 				{
 					FInitAudioBusSend BusSend;
 					BusSend.AudioBusId = InBusId;
 					BusSend.SendLevel = InSendInfo.SendLevel;
-					InInitParams.AudioBusSends[BusSendType].Add(BusSend);
+					InInitParams.AudioBusSends[InBusSendType].Add(BusSend);
 
 					FDynamicBusSendInfo DynamicBusSendInfo;
 					DynamicBusSendInfo.SendLevel = InSendInfo.SendLevel;
 					DynamicBusSendInfo.BusId = BusSend.AudioBusId;
 					DynamicBusSendInfo.BusSendLevelControlMethod = InSendInfo.SourceBusSendLevelControlMethod;
-					DynamicBusSendInfo.BusSendType = (EBusSendType)BusSendType;
+					DynamicBusSendInfo.BusSendType = (EBusSendType)InBusSendType;
 					DynamicBusSendInfo.MinSendLevel = InSendInfo.MinSendLevel;
 					DynamicBusSendInfo.MaxSendLevel = InSendInfo.MaxSendLevel;
 					DynamicBusSendInfo.MinSendDistance = InSendInfo.MinSendDistance;
