Fix for source rotation ambisonics sources

diff --git a/Engine/Plugins/Runtime/SoundFields/Source/SoundFields/Private/SoundFields.cpp b/Engine/Plugins/Runtime/SoundFields/Source/SoundFields/Private/SoundFields.cpp
--- a/Engine/Plugins/Runtime/SoundFields/Source/SoundFields/Private/SoundFields.cpp
+++ b/Engine/Plugins/Runtime/SoundFields/Source/SoundFields/Private/SoundFields.cpp
@@ -170,6 +170,9 @@ public:
  */
 class FAmbisonicsMixer : public ISoundfieldMixerStream
 {
+private:
+	// This is a temp buffer that we use to rotate InputData in MixTogether to the world rotation.
+	FAmbisonicsSoundfieldBuffer RotatedAudio;
 public:
 	FAmbisonicsMixer() {}
 
@@ -186,10 +189,19 @@ public:
 			OutAudio.AudioBuffer.AddZeroed(InAudio.AudioBuffer.Num());
 		}
 
+		if (RotatedAudio.NumChannels == 0)
+		{
+			RotatedAudio.NumChannels = InAudio.NumChannels;
+			RotatedAudio.AudioBuffer.Reset();
+			RotatedAudio.AudioBuffer.AddZeroed(InAudio.AudioBuffer.Num());
+		}
+
 		check(OutAudio.NumChannels == InAudio.NumChannels);
 		check(InAudio.AudioBuffer.Num() == OutAudio.AudioBuffer.Num());
 
-		Audio::MixInBufferFast(InAudio.AudioBuffer, OutAudio.AudioBuffer, InputData.SendLevel);
+		// Rotate 
+		FSoundFieldDecoder::RotateFirstOrderAmbisonicsBed(InAudio, RotatedAudio, InAudio.Rotation, InAudio.PreviousRotation);
+		Audio::MixInBufferFast(RotatedAudio.AudioBuffer, OutAudio.AudioBuffer, InputData.SendLevel);
 	}
 };
 
diff --git a/Engine/Source/Runtime/AudioMixer/Private/AudioMixerSourceManager.cpp b/Engine/Source/Runtime/AudioMixer/Private/AudioMixerSourceManager.cpp
--- a/Engine/Source/Runtime/AudioMixer/Private/AudioMixerSourceManager.cpp
+++ b/Engine/Source/Runtime/AudioMixer/Private/AudioMixerSourceManager.cpp
@@ -1848,7 +1848,8 @@ namespace Audio
 			FAmbisonicsSoundfieldBuffer AmbiBuffer;
 			AmbiBuffer.AudioBuffer = MoveTemp(*DownmixData.PostEffectBuffers);
 			AmbiBuffer.NumChannels = DownmixData.NumInputChannels;
-			AmbiBuffer.Rotation = DownmixData.PositionalData.Rotation;
+			AmbiBuffer.PreviousRotation = AmbiBuffer.Rotation;
+			AmbiBuffer.Rotation = DownmixData.SourceRotation;
 
 			DownmixData.PositionalData.NumChannels = DownmixData.NumDeviceChannels;
 			DownmixData.PositionalData.ChannelPositions = MixerDevice->GetDefaultPositionMap(DownmixData.NumDeviceChannels);
@@ -1922,7 +1923,8 @@ namespace Audio
 			FAmbisonicsSoundfieldBuffer AmbiBuffer;
 			AmbiBuffer.AudioBuffer = MoveTemp(*DownmixData.PostEffectBuffers);
 			AmbiBuffer.NumChannels = DownmixData.NumInputChannels;
-			AmbiBuffer.Rotation = DownmixData.PositionalData.Rotation;
+			AmbiBuffer.PreviousRotation = AmbiBuffer.Rotation;
+			AmbiBuffer.Rotation = DownmixData.SourceRotation;
 
 			DownmixData.PositionalData.NumChannels = DownmixData.NumDeviceChannels;
 			DownmixData.PositionalData.ChannelPositions = MixerDevice->GetDefaultPositionMap(DownmixData.NumDeviceChannels);
@@ -2051,6 +2053,8 @@ namespace Audio
 				FAmbisonicsSoundfieldBuffer AmbiBuffer;
 				AmbiBuffer.AudioBuffer = MoveTemp(*DownmixData.PostEffectBuffers);
 				AmbiBuffer.NumChannels = DownmixData.NumInputChannels;
+				AmbiBuffer.PreviousRotation = AmbiBuffer.Rotation;
+				AmbiBuffer.Rotation = DownmixData.SourceRotation;
 
 				SoundfieldData.AmbiTranscoder->Transcode(AmbiBuffer, GetAmbisonicsSourceDefaultSettings(), *SoundfieldData.EncodedPacket, *SoundfieldData.EncoderSettings);
 				*DownmixData.PostEffectBuffers = MoveTemp(AmbiBuffer.AudioBuffer);
@@ -2069,10 +2073,14 @@ namespace Audio
 			}
 			else if (SoundfieldData.bIsUnrealAmbisonicsSubmix)
 			{
+				ensure(InSource.bIsAmbisonics);
+
 				FAmbisonicsSoundfieldBuffer& OutputPacket = DowncastSoundfieldRef<FAmbisonicsSoundfieldBuffer>(*SoundfieldData.EncodedPacket);
 				// Fixme: This is an array copy. Can we serve DownmixData.PostEffectBuffers directly to this soundfield?
 				OutputPacket.AudioBuffer = *DownmixData.PostEffectBuffers;
 				OutputPacket.NumChannels = DownmixData.NumInputChannels;
+				OutputPacket.PreviousRotation = OutputPacket.Rotation;
+				OutputPacket.Rotation = DownmixData.SourceRotation;
 			}
 		}
 	}
@@ -2344,6 +2352,7 @@ namespace Audio
 			}
 
 			DownmixData.PositionalData.Rotation = SourceInfo.SpatParams.ListenerOrientation;
+			DownmixData.SourceRotation = SourceInfo.SpatParams.EmitterWorldRotation;
 
 			if (SourceInfo.bIs3D && !DownmixData.bIsInitialDownmix)
 			{
diff --git a/Engine/Source/Runtime/AudioMixer/Private/AudioMixerSourceManager.h b/Engine/Source/Runtime/AudioMixer/Private/AudioMixerSourceManager.h
--- a/Engine/Source/Runtime/AudioMixer/Private/AudioMixerSourceManager.h
+++ b/Engine/Source/Runtime/AudioMixer/Private/AudioMixerSourceManager.h
@@ -357,6 +357,10 @@ namespace Audio
 
 		struct FSourceDownmixData
 		{
+			// cached parameters for encoding to a soundfield format.
+			FSoundfieldSpeakerPositionalData PositionalData;
+			FQuat SourceRotation;
+
 			// Output data, after computing a block of sample data, this is read back from mixers
 			Audio::AlignedFloatBuffer ReverbPluginOutputBuffer;
 			Audio::AlignedFloatBuffer* PostEffectBuffers;
@@ -372,19 +376,17 @@ namespace Audio
 			uint32 NumDeviceChannels;
 			uint8 bIsInitialDownmix : 1;
 
-			// cached parameters for encoding to a soundfield format.
-			FSoundfieldSpeakerPositionalData PositionalData;
-
 			// If this source is an ambisonics source, we use this to downmix the source to our output.
 			TUniquePtr<ISoundfieldDecoderStream> AmbisonicsDecoder;
 
 			FSourceDownmixData(uint32 SourceNumChannels, uint32 NumDeviceOutputChannels, uint32 InNumFrames)
-				: PostEffectBuffers(nullptr)
+				: SourceRotation(FQuat::Identity)
+				, PostEffectBuffers(nullptr)
 				, DeviceSubmixInfo(FSubmixChannelData(SourceNumChannels, NumDeviceOutputChannels, InNumFrames))
 				, NumInputChannels(SourceNumChannels)
 				, NumFrames(InNumFrames)
 				, NumDeviceChannels(NumDeviceOutputChannels)
-				, bIsInitialDownmix(true)
+				, bIsInitialDownmix(true) 
 			{
 			}
 
diff --git a/Engine/Source/Runtime/SoundFieldRendering/Private/SoundFieldRendering.cpp b/Engine/Source/Runtime/SoundFieldRendering/Private/SoundFieldRendering.cpp
--- a/Engine/Source/Runtime/SoundFieldRendering/Private/SoundFieldRendering.cpp
+++ b/Engine/Source/Runtime/SoundFieldRendering/Private/SoundFieldRendering.cpp
@@ -323,7 +323,8 @@ void FSoundFieldDecoder::DecodeAudioToSevenOneAndDownmixToDevice(const FAmbisoni
 	check(NumVirtualChannels == 9); // if this has changed, it breaks our assumptions for fast gain indexing
 
 	// get listener orientation in spherical coordinates
-	FVector2D ListenerRotationSphericalCoord = OutputPositions.Rotation.Vector().UnitCartesianToSpherical();
+	FQuat RelativeRotation = OutputPositions.Rotation * InputData.Rotation;
+	FVector2D ListenerRotationSphericalCoord = RelativeRotation.Vector().UnitCartesianToSpherical();
 	FSphericalHarmonicCalculator::AdjustUESphericalCoordinatesForAmbisonics(ListenerRotationSphericalCoord);
 
 	// Set up output buffer for virtual speaker positions
@@ -338,7 +339,7 @@ void FSoundFieldDecoder::DecodeAudioToSevenOneAndDownmixToDevice(const FAmbisoni
 	{
 		CurrentSpeakerGains = TargetSpeakerGains;
 		TargetSpeakerGains = FoaVirtualSpeakerWordLockedGains;
-		FVector ListenerRot = OutputPositions.Rotation.Euler();
+		FVector ListenerRot = RelativeRotation.Euler();
 		FoaRotationInPlace(TargetSpeakerGains, ListenerRot.X, -ListenerRot.Y, ListenerRot.Z);
 
 		// If we just switched to this decode method (from direct to device), we need to initialize CurrentSpeakerGains size
diff --git a/Engine/Source/Runtime/SoundFieldRendering/Public/SoundFieldRendering.h b/Engine/Source/Runtime/SoundFieldRendering/Public/SoundFieldRendering.h
--- a/Engine/Source/Runtime/SoundFieldRendering/Public/SoundFieldRendering.h
+++ b/Engine/Source/Runtime/SoundFieldRendering/Public/SoundFieldRendering.h
@@ -15,11 +15,6 @@ SOUNDFIELDRENDERING_API ISoundfieldEncodingSettingsProxy& GetAmbisonicsSourceDef
 class SOUNDFIELDRENDERING_API FAmbisonicsSoundfieldBuffer : public ISoundfieldAudioPacket
 {
 public:
-	FAmbisonicsSoundfieldBuffer()
-		: NumChannels(0)
-		, Rotation(FQuat::Identity)
-	{}
-
 	// Interleaved audio buffer for all vector parts of the ambisonics stream.
 	Audio::AlignedFloatBuffer AudioBuffer;
 
@@ -30,6 +25,13 @@ public:
 
 	// This is the rotation of the ambisonics source.
 	FQuat Rotation;
+	FQuat PreviousRotation;
+
+	FAmbisonicsSoundfieldBuffer()
+		: NumChannels(0)
+		, Rotation(FQuat::Identity)
+		, PreviousRotation(FQuat::Identity)
+	{}
 
 	virtual void Serialize(FArchive& Ar) override;
 	virtual TUniquePtr<ISoundfieldAudioPacket> Duplicate() const override;
