Fix for unnecessary down mix buffer

diff --git a/Engine/Source/Runtime/AudioMixer/Private/AudioMixerSubmix.cpp b/Engine/Source/Runtime/AudioMixer/Private/AudioMixerSubmix.cpp
--- a/Engine/Source/Runtime/AudioMixer/Private/AudioMixerSubmix.cpp
+++ b/Engine/Source/Runtime/AudioMixer/Private/AudioMixerSubmix.cpp
@@ -851,7 +851,6 @@ namespace Audio
 					InputData.NumChannels = ChannelCountOverride;
 					InputData.AudioBuffer = &DownmixedBuffer;
 					SubmixEffect->ProcessAudio(InputData, OutputData);
-					DownmixBuffer(ChannelCountOverride, DownmixedBuffer, NumChannels, InputBuffer);
 				}
 				else
 				{
