Fixing application of reverb effect volume and wet reverb settings

diff --git a/Engine/Source/Runtime/AudioMixer/Classes/SubmixEffects/AudioMixerSubmixEffectReverbFast.h b/Engine/Source/Runtime/AudioMixer/Classes/SubmixEffects/AudioMixerSubmixEffectReverbFast.h
--- a/Engine/Source/Runtime/AudioMixer/Classes/SubmixEffects/AudioMixerSubmixEffectReverbFast.h
+++ b/Engine/Source/Runtime/AudioMixer/Classes/SubmixEffects/AudioMixerSubmixEffectReverbFast.h
@@ -6,6 +6,7 @@
 #include "Curves/RichCurve.h"
 #include "DSP/Amp.h"
 #include "DSP/ReverbFast.h"
+#include "DSP/BufferVectorOperations.h"
 #include "Sound/SoundEffectSubmix.h"
 
 #include "AudioMixerSubmixEffectReverbFast.generated.h"
@@ -137,6 +138,8 @@ private:
 	// Level of wet/dry signal on current buffer
 	Audio::FWetDry CurrentWetDry;
 
+	Audio::AlignedFloatBuffer WetInputBuffer;
+
 	// Curve which maps old reverb times to new decay value
 	FRichCurve DecayCurve;
 };
diff --git a/Engine/Source/Runtime/AudioMixer/Private/Effects/AudioMixerSubmixEffectReverbFast.cpp b/Engine/Source/Runtime/AudioMixer/Private/Effects/AudioMixerSubmixEffectReverbFast.cpp
--- a/Engine/Source/Runtime/AudioMixer/Private/Effects/AudioMixerSubmixEffectReverbFast.cpp
+++ b/Engine/Source/Runtime/AudioMixer/Private/Effects/AudioMixerSubmixEffectReverbFast.cpp
@@ -7,6 +7,7 @@
 #include "AudioMixer.h"
 #include "DSP/BufferVectorOperations.h"
 #include "DSP/ReverbFast.h"
+#include "DSP/Amp.h"
 
 
 static int32 DisableSubmixReverbCVarFast = 0;
@@ -106,10 +107,13 @@ void FSubmixEffectReverbFast::OnPresetChanged()
 	ReverbEffect.LateDelay = Settings.LateDelay;
 	ReverbEffect.AirAbsorptionGainHF = Settings.AirAbsorptionGainHF;
 	ReverbEffect.RoomRolloffFactor = 0.0f; // not used
-	ReverbEffect.Volume = Settings.bBypass ? 0.0f : Settings.WetLevel;
+
+	ReverbEffect.Volume = Settings.bBypass ? 0.0f : FMath::Clamp(Settings.WetLevel, MinWetness, MaxWetness);
 
 	SetParameters(ReverbEffect);
 
+	// These wet dry parameters need to be set after the call to SetParameters(ReverbEffect) parameter.
+	// SetParameters sets the WetDryParams, but they need to be overriden here. 
 	Audio::FWetDry NewWetDryParams;
 	NewWetDryParams.DryLevel = Settings.bBypass ? 1.0f : Settings.DryLevel;
 	NewWetDryParams.WetLevel = Settings.bBypass ? 0.0f : FMath::Clamp(Settings.WetLevel, MinWetness, MaxWetness);
@@ -144,8 +148,17 @@ void FSubmixEffectReverbFast::OnProcessAudio(const FSoundEffectSubmixInputData&
 		LastDry = CurrentWetDry.DryLevel;
 	}
 
-	PlateReverb->ProcessAudio(*InData.AudioBuffer, InData.NumChannels, *OutData.AudioBuffer, OutData.NumChannels);
-	Audio::FadeBufferFast(*OutData.AudioBuffer, LastWet, CurrentWetDry.WetLevel);
+
+	WetInputBuffer.Reset();
+	if (InData.AudioBuffer->Num() > 0)
+	{
+		// Wet level is applied to input audio to preserve reverb tail when changing wet level
+		WetInputBuffer.AddZeroed(InData.AudioBuffer->Num());
+		Audio::MixInBufferFast(*InData.AudioBuffer, WetInputBuffer, LastWet, CurrentWetDry.WetLevel);
+	}
+
+	PlateReverb->ProcessAudio(WetInputBuffer, InData.NumChannels, *OutData.AudioBuffer, OutData.NumChannels);
+
 	Audio::MixInBufferFast(*InData.AudioBuffer, *OutData.AudioBuffer, LastDry, CurrentWetDry.DryLevel);
 }
 
@@ -189,6 +202,11 @@ bool FSubmixEffectReverbFast::SetParameters(const FAudioEffectParameters& InPara
 	// Apply the settings the thread safe settings object
 	ReverbParams.SetParams(NewSettings);
 
+	// Apply wet/dry level
+	// When using a FAudioReverbEffect, the volume parameter controls wetness and the dry level remains 0.
+	Audio::FWetDry NewWetDry(ReverbEffectParams.Volume, 0.f);
+	WetDryParams.SetParams(NewWetDry);
+
 	return true;
 }
 
