only include libsnd dll dependency on Win32 & Win64

diff --git a/Engine/Source/Runtime/AudioMixer/AudioMixer.Build.cs b/Engine/Source/Runtime/AudioMixer/AudioMixer.Build.cs
--- a/Engine/Source/Runtime/AudioMixer/AudioMixer.Build.cs
+++ b/Engine/Source/Runtime/AudioMixer/AudioMixer.Build.cs
@@ -48,7 +48,7 @@ namespace UnrealBuildTool.Rules
 					"UELibSampleRate"
 					);
 
-			if (Target.IsInPlatformGroup(UnrealPlatformGroup.Windows))
+			if (Target.Platform == UnrealTargetPlatform.Win64 || Target.Platform == UnrealTargetPlatform.Win32)
 			{
 				string PlatformName = Target.Platform == UnrealTargetPlatform.Win32 ? "Win32" : "Win64";
 
