fix for GC crash with async task owned by FMixerSourceBuffer

diff --git a/Engine/Source/Runtime/AudioMixer/Private/AudioMixerSource.h b/Engine/Source/Runtime/AudioMixer/Private/AudioMixerSource.h
--- a/Engine/Source/Runtime/AudioMixer/Private/AudioMixerSource.h
+++ b/Engine/Source/Runtime/AudioMixer/Private/AudioMixerSource.h
@@ -119,7 +119,7 @@ namespace Audio
 
 		FMixerDevice* MixerDevice;
 		FMixerBuffer* MixerBuffer;
-		TSharedPtr<FMixerSourceBuffer> MixerSourceBuffer;
+		TSharedPtr<FMixerSourceBuffer, ESPMode::ThreadSafe> MixerSourceBuffer;
 		FMixerSourceVoice* MixerSourceVoice;
 
 		// This holds data copied from FSoundSourceBusSendInfo when a new sound starts playing
diff --git a/Engine/Source/Runtime/AudioMixer/Private/AudioMixerSourceBuffer.cpp b/Engine/Source/Runtime/AudioMixer/Private/AudioMixerSourceBuffer.cpp
--- a/Engine/Source/Runtime/AudioMixer/Private/AudioMixerSourceBuffer.cpp
+++ b/Engine/Source/Runtime/AudioMixer/Private/AudioMixerSourceBuffer.cpp
@@ -56,7 +56,7 @@ namespace Audio
 		return CurrentSample >= NumSamples;
 	}
 
-	TSharedPtr<FMixerSourceBuffer> FMixerSourceBuffer::Create(FMixerBuffer& InBuffer, USoundWave& InWave, ELoopingMode InLoopingMode, bool bInIsSeeking, bool bInForceSyncDecode)
+	TSharedPtr<FMixerSourceBuffer, ESPMode::ThreadSafe> FMixerSourceBuffer::Create(FMixerBuffer& InBuffer, USoundWave& InWave, ELoopingMode InLoopingMode, bool bInIsSeeking, bool bInForceSyncDecode)
 	{
 		LLM_SCOPE(ELLMTag::AudioMixer);
 
@@ -72,7 +72,7 @@ namespace Audio
 			return nullptr;
 		}
 
-		TSharedPtr<FMixerSourceBuffer> NewSourceBuffer = MakeShareable(new FMixerSourceBuffer(InBuffer, InWave, InLoopingMode, bInIsSeeking, bInForceSyncDecode));
+		TSharedPtr<FMixerSourceBuffer, ESPMode::ThreadSafe> NewSourceBuffer = MakeShareable(new FMixerSourceBuffer(InBuffer, InWave, InLoopingMode, bInIsSeeking, bInForceSyncDecode));
 		return NewSourceBuffer;
 	}
 
@@ -100,7 +100,7 @@ namespace Audio
 		const uint32 TotalSamples = MONO_PCM_BUFFER_SAMPLES * NumChannels;
 		for (int32 BufferIndex = 0; BufferIndex < Audio::MAX_BUFFERS_QUEUED; ++BufferIndex)
 		{
-			SourceVoiceBuffers.Add(MakeShared<FMixerSourceVoiceBuffer>());
+			SourceVoiceBuffers.Add(MakeShared<FMixerSourceVoiceBuffer, ESPMode::ThreadSafe>());
 
 			// Prepare the memory to fit the max number of samples
 			SourceVoiceBuffers[BufferIndex]->AudioData.Reset(TotalSamples);
@@ -113,29 +113,14 @@ namespace Audio
 	{
 		// GC methods may get called from the game thread during the destructor
 		// These methods will trylock and early exit if we have this lock
-		FScopeLock Lock(&DtorCritSec);
-
-		// Make sure we have completed our async realtime task before deleting the decompression state
-		if (AsyncRealtimeAudioTask)
-		{
-			AsyncRealtimeAudioTask->CancelTask();
-			delete AsyncRealtimeAudioTask;
-			AsyncRealtimeAudioTask = nullptr;
-		}
+		FScopeLock Lock(&SoundWaveCritSec);
 
+		// OnEndGenerate calls EnsureTaskFinishes,
+		// which will make sure we have completed our async realtime task before deleting the decompression state
 		OnEndGenerate();
 
 		// Clean up decompression state after things have been finished using it
-		if (DecompressionState)
-		{
-			if (BufferType == EBufferType::Streaming)
-			{
-				IStreamingManager::Get().GetAudioStreamingManager().RemoveDecoder(DecompressionState);
-			}
-
-			delete DecompressionState;
-			DecompressionState = nullptr;
-		}
+		DeleteDecoder();
 
 		if (SoundWave)
 		{
@@ -175,18 +160,18 @@ namespace Audio
 
 		switch (BufferType)
 		{
-			case EBufferType::PCM:
-			case EBufferType::PCMPreview:
-				SubmitInitialPCMBuffers();
-				break;
+		case EBufferType::PCM:
+		case EBufferType::PCMPreview:
+			SubmitInitialPCMBuffers();
+			break;
 
-			case EBufferType::PCMRealTime:
-			case EBufferType::Streaming:
-				SubmitInitialRealtimeBuffers();
-				break;
+		case EBufferType::PCMRealTime:
+		case EBufferType::Streaming:
+			SubmitInitialRealtimeBuffers();
+			break;
 
-			case EBufferType::Invalid:
-				break;
+		case EBufferType::Invalid:
+			break;
 		}
 
 		return true;
@@ -194,7 +179,9 @@ namespace Audio
 
 	void FMixerSourceBuffer::OnBufferEnd()
 	{
-		if ((NumBuffersQeueued == 0 && bBufferFinished) || (bProcedural && !SoundWave))
+		FScopeTryLock Lock(&SoundWaveCritSec);
+
+		if (!Lock.IsLocked() || (NumBuffersQeueued == 0 && bBufferFinished) || (bProcedural && !SoundWave))
 		{
 			return;
 		}
@@ -202,9 +189,26 @@ namespace Audio
 		ProcessRealTimeSource();
 	}
 
-	TSharedPtr<FMixerSourceVoiceBuffer> FMixerSourceBuffer::GetNextBuffer()
+	int32 FMixerSourceBuffer::GetNumBuffersQueued() const
 	{
-		TSharedPtr<FMixerSourceVoiceBuffer> NewBufferPtr;
+		FScopeTryLock Lock(&SoundWaveCritSec);
+		if (Lock.IsLocked())
+		{
+			return NumBuffersQeueued;
+		}
+
+		return 0;
+	}
+
+	TSharedPtr<FMixerSourceVoiceBuffer, ESPMode::ThreadSafe> FMixerSourceBuffer::GetNextBuffer()
+	{
+		FScopeTryLock Lock(&SoundWaveCritSec);
+		if(!Lock.IsLocked())
+		{
+			return nullptr;
+		}
+
+		TSharedPtr<FMixerSourceVoiceBuffer, ESPMode::ThreadSafe> NewBufferPtr;
 		BufferQueue.Dequeue(NewBufferPtr);
 		--NumBuffersQeueued;
 		return NewBufferPtr;
@@ -308,7 +312,7 @@ namespace Audio
 		{
 			ProcessRealTimeSource();
 		}
-	}
+			}
 
 	bool FMixerSourceBuffer::ReadMoreRealtimeData(ICompressedAudioInfo* InDecoder, const int32 BufferIndex, EBufferReadMode BufferReadMode)
 	{
@@ -318,17 +322,24 @@ namespace Audio
 
 		if (bProcedural)
 		{
-			check(SoundWave && SoundWave->bProcedural);
-			FProceduralAudioTaskData NewTaskData;
-			NewTaskData.ProceduralSoundWave = SoundWave;
-			NewTaskData.AudioData = SourceVoiceBuffers[BufferIndex]->AudioData.GetData();
-			NewTaskData.NumSamples = MaxSamples;
-			NewTaskData.NumChannels = NumChannels;
-			check(!AsyncRealtimeAudioTask);
-			AsyncRealtimeAudioTask = CreateAudioTask(NewTaskData);
-
-			// Procedural sound waves never loop
-			return false;
+			FScopeTryLock Lock(&SoundWaveCritSec);
+			
+			if (Lock.IsLocked() && SoundWave)
+			{
+				check(SoundWave && SoundWave->bProcedural);
+				FProceduralAudioTaskData NewTaskData;
+				NewTaskData.ProceduralSoundWave = SoundWave;
+				NewTaskData.AudioData = SourceVoiceBuffers[BufferIndex]->AudioData.GetData();
+				NewTaskData.NumSamples = MaxSamples;
+				NewTaskData.NumChannels = NumChannels;
+				check(!AsyncRealtimeAudioTask);
+				AsyncRealtimeAudioTask = CreateAudioTask(NewTaskData);
+
+				// Procedural sound waves never loop
+				return false;
+			}
+
+			return true;
 		}
 		else if (BufferType != EBufferType::PCMRealTime && BufferType != EBufferType::Streaming)
 		{
@@ -351,6 +362,7 @@ namespace Audio
 		NewTaskData.NumPrecacheFrames = NumPrecacheFrames;
 		NewTaskData.bForceSyncDecode = bForceSyncDecode;
 
+		FScopeLock Lock(&DecodeTaskCritSec);
 		check(!AsyncRealtimeAudioTask);
 		AsyncRealtimeAudioTask = CreateAudioTask(NewTaskData);
 
@@ -389,6 +401,7 @@ namespace Audio
 
 	void FMixerSourceBuffer::ProcessRealTimeSource()
 	{
+		FScopeLock Lock(&DecodeTaskCritSec);
 		if (AsyncRealtimeAudioTask)
 		{
 			AsyncRealtimeAudioTask->EnsureCompletion();
@@ -451,23 +464,40 @@ namespace Audio
 		}
 	}
 
-	void FMixerSourceBuffer::SubmitBuffer(TSharedPtr<FMixerSourceVoiceBuffer> InSourceVoiceBuffer)
+	void FMixerSourceBuffer::SubmitBuffer(TSharedPtr<FMixerSourceVoiceBuffer, ESPMode::ThreadSafe> InSourceVoiceBuffer)
 	{
 		NumBuffersQeueued++;
 		BufferQueue.Enqueue(InSourceVoiceBuffer);
 	}
 
+	void FMixerSourceBuffer::DeleteDecoder()
+	{
+		// Clean up decompression state after things have been finished using it
+		if (DecompressionState)
+		{
+			if (BufferType == EBufferType::Streaming)
+			{
+				IStreamingManager::Get().GetAudioStreamingManager().RemoveDecoder(DecompressionState);
+			}
+
+			delete DecompressionState;
+			DecompressionState = nullptr;
+		}
+	}
+
 	bool FMixerSourceBuffer::OnBeginDestroy(USoundWave * /*Wave*/)
 	{
-		FScopeTryLock Lock(&DtorCritSec);
+		FScopeTryLock Lock(&SoundWaveCritSec);
 
 		// if we don't have the lock, it means we are in ~FMixerSourceBuffer() on another thread
 		if (Lock.IsLocked() && SoundWave)
 		{
-			SoundWave = nullptr;
+			EnsureAsyncTaskFinishes();
+			DeleteDecoder();
+			ClearWave();
 			return true;
 		}
-		
+
 		return false;
 	}
 
@@ -478,22 +508,26 @@ namespace Audio
 
 	void FMixerSourceBuffer::OnFinishDestroy(USoundWave * /*Wave*/)
 	{
-		FScopeTryLock Lock(&DtorCritSec);
+		EnsureAsyncTaskFinishes();
+		FScopeTryLock Lock(&SoundWaveCritSec);
 
 		// if we don't have the lock, it means we are in ~FMixerSourceBuffer() on another thread
 		if (Lock.IsLocked() && SoundWave)
 		{
-			SoundWave = nullptr;
+			DeleteDecoder();
+			ClearWave();
 		}
 	}
 
 	bool FMixerSourceBuffer::IsAsyncTaskInProgress() const
 	{ 
-		return AsyncRealtimeAudioTask != nullptr; 
+		FScopeLock Lock(&DecodeTaskCritSec);
+		return AsyncRealtimeAudioTask != nullptr;
 	}
 
 	bool FMixerSourceBuffer::IsAsyncTaskDone() const 
 	{
+		FScopeLock Lock(&DecodeTaskCritSec);
 		if (AsyncRealtimeAudioTask)
 		{
 			return AsyncRealtimeAudioTask->IsDone();
@@ -503,6 +537,7 @@ namespace Audio
 	
 	void FMixerSourceBuffer::EnsureAsyncTaskFinishes() 
 	{
+		FScopeLock Lock(&DecodeTaskCritSec);
 		if (AsyncRealtimeAudioTask) 
 		{ 
 			AsyncRealtimeAudioTask->CancelTask(); 
@@ -514,7 +549,9 @@ namespace Audio
 
 	void FMixerSourceBuffer::OnBeginGenerate()
 	{
-		if (SoundWave && bProcedural)
+		FScopeTryLock Lock(&SoundWaveCritSec);
+
+		if (Lock.IsLocked() && SoundWave && bProcedural)
 		{
 			check(SoundWave && SoundWave->bProcedural);
 			SoundWave->OnBeginGenerate();
@@ -525,6 +562,13 @@ namespace Audio
 	{		
 		// Make sure the async task finishes!
 		EnsureAsyncTaskFinishes();
+
+		FScopeTryLock Lock(&SoundWaveCritSec);
+		if (!Lock.IsLocked())
+		{
+			return;
+		}
+
 						
 		// Only need to call OnEndGenerate and access SoundWave here if we successfully initialized
 		if (SoundWave && bInitialized && bProcedural)
diff --git a/Engine/Source/Runtime/AudioMixer/Private/AudioMixerSourceBuffer.h b/Engine/Source/Runtime/AudioMixer/Private/AudioMixerSourceBuffer.h
--- a/Engine/Source/Runtime/AudioMixer/Private/AudioMixerSourceBuffer.h
+++ b/Engine/Source/Runtime/AudioMixer/Private/AudioMixerSourceBuffer.h
@@ -43,7 +43,7 @@ namespace Audio
 		AsynchronousSkipFirstFrame
 	};
 
-	using FMixerSourceBufferPtr = TSharedPtr<class FMixerSourceBuffer>;
+	using FMixerSourceBufferPtr = TSharedPtr<class FMixerSourceBuffer, ESPMode::ThreadSafe>;
 
 	/** Class which handles decoding audio for a particular source buffer. */
 	class FMixerSourceBuffer : public ISoundWaveClient
@@ -68,10 +68,10 @@ namespace Audio
 		void OnBufferEnd();
 
 		// Return the number of buffers enqueued on the mixer source buffer
-		int32 GetNumBuffersQueued() const { return NumBuffersQeueued; }
-		
+		int32 GetNumBuffersQueued() const;
+
 		// Returns the next enqueued buffer, returns nullptr if no buffers enqueued
-		TSharedPtr<FMixerSourceVoiceBuffer> GetNextBuffer();
+		TSharedPtr<FMixerSourceVoiceBuffer, ESPMode::ThreadSafe> GetNextBuffer();
 
 		// Returns if buffer looped
 		bool DidBufferLoop() const { return bLoopCallback; }
@@ -102,14 +102,15 @@ namespace Audio
 		void SubmitInitialRealtimeBuffers();
 		void SubmitRealTimeSourceData(const bool bLooped);
 		void ProcessRealTimeSource();
-		void SubmitBuffer(TSharedPtr<FMixerSourceVoiceBuffer> InSourceVoiceBuffer);
+		void SubmitBuffer(TSharedPtr<FMixerSourceVoiceBuffer, ESPMode::ThreadSafe> InSourceVoiceBuffer);
+		void DeleteDecoder();
 
 
 		int32 NumBuffersQeueued;
 		FRawPCMDataBuffer RawPCMDataBuffer;
 
-		TArray<TSharedPtr<FMixerSourceVoiceBuffer>> SourceVoiceBuffers;
-		TQueue<TSharedPtr<FMixerSourceVoiceBuffer>> BufferQueue;
+		TArray<TSharedPtr<FMixerSourceVoiceBuffer, ESPMode::ThreadSafe>> SourceVoiceBuffers;
+		TQueue<TSharedPtr<FMixerSourceVoiceBuffer, ESPMode::ThreadSafe>> BufferQueue;
 		int32 CurrentBuffer;
 		// SoundWaves are only set for procedural sound waves
 		USoundWave* SoundWave;
@@ -121,7 +122,8 @@ namespace Audio
 		int32 NumPrecacheFrames;
 		TArray<uint8> CachedRealtimeFirstBuffer;
 
-		FCriticalSection DtorCritSec;
+		mutable FCriticalSection SoundWaveCritSec;
+		mutable FCriticalSection DecodeTaskCritSec;
 
 		uint32 bInitialized : 1;
 		uint32 bBufferFinished : 1;
diff --git a/Engine/Source/Runtime/AudioMixer/Private/AudioMixerSourceManager.cpp b/Engine/Source/Runtime/AudioMixer/Private/AudioMixerSourceManager.cpp
--- a/Engine/Source/Runtime/AudioMixer/Private/AudioMixerSourceManager.cpp
+++ b/Engine/Source/Runtime/AudioMixer/Private/AudioMixerSourceManager.cpp
@@ -375,7 +375,12 @@ namespace Audio
 
 		AUDIO_MIXER_CHECK(SourceId < NumTotalSources);
 		AUDIO_MIXER_CHECK(bInitialized);
-		AUDIO_MIXER_CHECK(MixerSources[SourceId] != nullptr);
+
+		if (MixerSources[SourceId] == nullptr)
+		{
+			UE_LOG(LogAudioMixer, Warning, TEXT("Ignoring double release of SourceId: %i"), SourceId);
+			return;
+		}
 
 		AUDIO_MIXER_DEBUG_LOG(SourceId, TEXT("Is releasing"));
 		
diff --git a/Engine/Source/Runtime/AudioMixer/Private/AudioMixerSourceManager.h b/Engine/Source/Runtime/AudioMixer/Private/AudioMixerSourceManager.h
--- a/Engine/Source/Runtime/AudioMixer/Private/AudioMixerSourceManager.h
+++ b/Engine/Source/Runtime/AudioMixer/Private/AudioMixerSourceManager.h
@@ -85,7 +85,7 @@ namespace Audio
 
 	struct FMixerSourceVoiceInitParams
 	{
-		TSharedPtr<FMixerSourceBuffer> MixerSourceBuffer = nullptr;
+		TSharedPtr<FMixerSourceBuffer, ESPMode::ThreadSafe> MixerSourceBuffer = nullptr;
 		ISourceListener* SourceListener = nullptr;
 		TArray<FMixerSourceSubmixSend> SubmixSends;
 		TArray<FInitAudioBusSend> AudioBusSends[(int32)EBusSendType::Count];
@@ -415,11 +415,11 @@ namespace Audio
 			~FSourceInfo() {}
 
 			// Object which handles source buffer decoding
-			TSharedPtr<FMixerSourceBuffer> MixerSourceBuffer;
+			TSharedPtr<FMixerSourceBuffer, ESPMode::ThreadSafe> MixerSourceBuffer;
 			ISourceListener* SourceListener;
 
 			// Data used for rendering sources
-			TSharedPtr<FMixerSourceVoiceBuffer> CurrentPCMBuffer;
+			TSharedPtr<FMixerSourceVoiceBuffer, ESPMode::ThreadSafe> CurrentPCMBuffer;
 			int32 CurrentAudioChunkNumFrames;
 
 			// The post-attenuation source buffer, used to send audio to submixes
@@ -557,7 +557,7 @@ namespace Audio
 		TArray<FAsyncTask<FAudioMixerSourceWorker>*> SourceWorkers;
 
 		// Array of task data waiting to finished. Processed on audio render thread.
-		TArray<TSharedPtr<FMixerSourceBuffer>> PendingSourceBuffers;
+		TArray<TSharedPtr<FMixerSourceBuffer, ESPMode::ThreadSafe>> PendingSourceBuffers;
 
 		// General information about sources in source manager accessible from game thread
 		struct FGameThreadInfo
diff --git a/Engine/Source/Runtime/AudioMixer/Private/Components/SynthComponent.cpp b/Engine/Source/Runtime/AudioMixer/Private/Components/SynthComponent.cpp
--- a/Engine/Source/Runtime/AudioMixer/Private/Components/SynthComponent.cpp
+++ b/Engine/Source/Runtime/AudioMixer/Private/Components/SynthComponent.cpp
@@ -219,7 +219,7 @@ void USynthComponent::Initialize(int32 SampleRateOverride)
 
 		if (!Synth)
 		{
-			Synth = NewObject<USynthSound>(this, TEXT("Synth"));
+			Synth = NewObject<USynthSound>(this, FName(TEXT("TimeSynth"), this->GetUniqueID()));
 		}
 
 		// Copy sound base data to the sound
diff --git a/Engine/Source/Runtime/AudioMixer/Private/SoundWaveDecoder.cpp b/Engine/Source/Runtime/AudioMixer/Private/SoundWaveDecoder.cpp
--- a/Engine/Source/Runtime/AudioMixer/Private/SoundWaveDecoder.cpp
+++ b/Engine/Source/Runtime/AudioMixer/Private/SoundWaveDecoder.cpp
@@ -1,6 +1,8 @@
 // Copyright Epic Games, Inc. All Rights Reserved.
 
 #include "SoundWaveDecoder.h"
+#include "Engine/Public/AudioThread.h"
+#include "Misc/ScopeTryLock.h"
 #include "AudioThread.h"
 #include "AudioDecompress.h"
 #include "AudioMixer.h"
@@ -26,6 +28,8 @@ namespace Audio
 
 	FDecodingSoundSource::~FDecodingSoundSource()
 	{
+		FScopeLock Lock(&DtorCritSec);
+
 		if (MixerSourceBuffer.IsValid())
 		{
 			MixerSourceBuffer->OnEndGenerate();
@@ -152,6 +156,15 @@ namespace Audio
 
 	void FDecodingSoundSource::ReadFrame()
 	{
+		if (!MixerSourceBuffer.IsValid() || !SourceInfo.CurrentPCMBuffer.IsValid())
+		{
+			SourceInfo.bIsLastBuffer = true;
+			return;
+		}
+
+		TSharedPtr<FMixerSourceBuffer, ESPMode::ThreadSafe>MixerSourceBufferLocal = MixerSourceBuffer;
+		TSharedPtr<FMixerSourceVoiceBuffer, ESPMode::ThreadSafe>CurrentPCMBufferLocal = SourceInfo.CurrentPCMBuffer;
+
 		bool bNextFrameOutOfRange = (SourceInfo.CurrentFrameIndex + 1) >= SourceInfo.CurrentAudioChunkNumFrames;
 		bool bCurrentFrameOutOfRange = SourceInfo.CurrentFrameIndex >= SourceInfo.CurrentAudioChunkNumFrames;
 
@@ -163,30 +176,41 @@ namespace Audio
 			{
 				bReadCurrentFrame = false;
 
-				const float* AudioData = SourceInfo.CurrentPCMBuffer->AudioData.GetData();
+				const float* AudioData = CurrentPCMBufferLocal->AudioData.GetData();
+
+				if (!AudioData)
+				{
+					SourceInfo.bIsLastBuffer = true;
+					return;
+				}
+
 				const int32 CurrentSampleIndex = SourceInfo.CurrentFrameIndex * SourceInfo.NumSourceChannels;
 
 				for (int32 Channel = 0; Channel < SourceInfo.NumSourceChannels; ++Channel)
 				{
 					SourceInfo.CurrentFrameValues[Channel] = AudioData[CurrentSampleIndex + Channel];
 				}
-			}
 
-			if (SourceInfo.CurrentPCMBuffer.IsValid())
-			{
-				if (SourceInfo.CurrentPCMBuffer->LoopCount == Audio::LOOP_FOREVER && !SourceInfo.CurrentPCMBuffer->bRealTimeBuffer)
+				if (CurrentPCMBufferLocal->LoopCount == Audio::LOOP_FOREVER && !CurrentPCMBufferLocal->bRealTimeBuffer)
 				{
 					SourceInfo.CurrentFrameIndex = FMath::Max(SourceInfo.CurrentFrameIndex - SourceInfo.CurrentAudioChunkNumFrames, 0);
 					break;
 				}
 
-				MixerSourceBuffer->OnBufferEnd();
+				MixerSourceBufferLocal->OnBufferEnd();
 			}
 
-			if (MixerSourceBuffer->GetNumBuffersQueued() > 0)
+			if (MixerSourceBufferLocal->GetNumBuffersQueued() > 0 && (SourceInfo.NumSourceChannels > 0))
 			{
-				SourceInfo.CurrentPCMBuffer = MixerSourceBuffer->GetNextBuffer();
-				SourceInfo.CurrentAudioChunkNumFrames = SourceInfo.CurrentPCMBuffer->AudioData.Num() / SourceInfo.NumSourceChannels;
+				check(MixerSourceBufferLocal.IsValid());
+				CurrentPCMBufferLocal = MixerSourceBufferLocal->GetNextBuffer();
+				if (!CurrentPCMBufferLocal)
+				{
+					SourceInfo.bIsLastBuffer = true;
+					return;
+				}
+
+				SourceInfo.CurrentAudioChunkNumFrames = CurrentPCMBufferLocal->AudioData.Num() / SourceInfo.NumSourceChannels;
 
 				if (bReadCurrentFrame)
 				{
@@ -199,6 +223,7 @@ namespace Audio
 			}
 			else
 			{
+				SourceInfo.bIsLastBuffer = true;
 				return;
 			}
 
@@ -206,26 +231,30 @@ namespace Audio
 			bCurrentFrameOutOfRange = SourceInfo.CurrentFrameIndex >= SourceInfo.CurrentAudioChunkNumFrames;
 		}
 
-		if (SourceInfo.CurrentPCMBuffer.IsValid())
+		const float* AudioData = CurrentPCMBufferLocal->AudioData.GetData();
+
+		if (!AudioData)
 		{
-			const float* AudioData = SourceInfo.CurrentPCMBuffer->AudioData.GetData();
-			const int32 NextSampleIndex = (SourceInfo.CurrentFrameIndex + 1)  * SourceInfo.NumSourceChannels;
+			SourceInfo.bIsLastBuffer = true;
+			return;
+		}
 
-			if (bReadCurrentFrame)
+		const int32 NextSampleIndex = (SourceInfo.CurrentFrameIndex + 1)  * SourceInfo.NumSourceChannels;
+
+		if (bReadCurrentFrame)
+		{
+			const int32 CurrentSampleIndex = SourceInfo.CurrentFrameIndex * SourceInfo.NumSourceChannels;
+			for (int32 Channel = 0; Channel < SourceInfo.NumSourceChannels; ++Channel)
 			{
-				const int32 CurrentSampleIndex = SourceInfo.CurrentFrameIndex * SourceInfo.NumSourceChannels;
-				for (int32 Channel = 0; Channel < SourceInfo.NumSourceChannels; ++Channel)
-				{
-					SourceInfo.CurrentFrameValues[Channel] = AudioData[CurrentSampleIndex + Channel];
-					SourceInfo.NextFrameValues[Channel] = AudioData[NextSampleIndex + Channel];
-				}
+				SourceInfo.CurrentFrameValues[Channel] = AudioData[CurrentSampleIndex + Channel];
+				SourceInfo.NextFrameValues[Channel] = AudioData[NextSampleIndex + Channel];
 			}
-			else
+		}
+		else
+		{
+			for (int32 Channel = 0; Channel < SourceInfo.NumSourceChannels; ++Channel)
 			{
-				for (int32 Channel = 0; Channel < SourceInfo.NumSourceChannels; ++Channel)
-				{
-					SourceInfo.NextFrameValues[Channel] = AudioData[NextSampleIndex + Channel];
-				}
+				SourceInfo.NextFrameValues[Channel] = AudioData[NextSampleIndex + Channel];
 			}
 		}
 	}
@@ -266,11 +295,24 @@ namespace Audio
 				SourceInfo.CurrentFrameAlpha -= 1.0f;
 			}
 
+
+			if (!SourceInfo.CurrentPCMBuffer.IsValid() || !MixerSourceBuffer.IsValid())
+			{
+				bReadFrame = false;
+				SourceInfo.bIsLastBuffer = true;
+				break;
+			}
+
 			if (bReadFrame)
 			{
 				ReadFrame();
+				if (SourceInfo.bIsLastBuffer)
+				{
+					break;
+				}
 			}
 
+
 			const float CurrentVolumeScale = SourceInfo.VolumeParam.Update();
 
 			for (int32 Channel = 0; Channel < SourceInfo.NumSourceChannels; ++Channel)
@@ -303,7 +345,9 @@ namespace Audio
 
 	bool FDecodingSoundSource::GetAudioBuffer(const int32 InNumFrames, const int32 InNumChannels, AlignedFloatBuffer& OutAudioBuffer)
 	{
-		if (!bInitialized)
+		FScopeTryLock Lock(&DtorCritSec);
+
+		if (!bInitialized || !Lock.IsLocked())
 		{
 			return false;
 		}
@@ -357,7 +401,6 @@ namespace Audio
 						BufferPtr[OutputSampleIndex++] = 0.5f * (ScratchBufferPtr[InputSampleIndex] + ScratchBufferPtr[InputSampleIndex + 1]);
 					}
 				}
-
 			}
 		}
 
diff --git a/Engine/Source/Runtime/AudioMixer/Public/SoundWaveDecoder.h b/Engine/Source/Runtime/AudioMixer/Public/SoundWaveDecoder.h
--- a/Engine/Source/Runtime/AudioMixer/Public/SoundWaveDecoder.h
+++ b/Engine/Source/Runtime/AudioMixer/Public/SoundWaveDecoder.h
@@ -100,7 +100,9 @@ namespace Audio
 		FMixerBuffer* MixerBuffer;
 
 		// Object which handles bulk of decoding operations
-		TSharedPtr<FMixerSourceBuffer> MixerSourceBuffer;
+		TSharedPtr<FMixerSourceBuffer, ESPMode::ThreadSafe> MixerSourceBuffer;
+
+		FCriticalSection DtorCritSec;
 
 		// Sample rate of the source
 		int32 SampleRate;
@@ -161,7 +163,7 @@ namespace Audio
 			TArray<float> NextFrameValues;
 
 			// The current decoded PCM buffer we are reading from
-			TSharedPtr<FMixerSourceVoiceBuffer> CurrentPCMBuffer;
+			TSharedPtr<FMixerSourceVoiceBuffer, ESPMode::ThreadSafe> CurrentPCMBuffer;
 
 			// If this sound is done (has decoded all data)
 			bool bIsDone;
