Fix for time synth regrission in caused by early exiting on IsValid() check on uninitialized shared ptr (in PIE).

diff --git a/Engine/Source/Runtime/AudioMixer/Private/SoundWaveDecoder.cpp b/Engine/Source/Runtime/AudioMixer/Private/SoundWaveDecoder.cpp
--- a/Engine/Source/Runtime/AudioMixer/Private/SoundWaveDecoder.cpp
+++ b/Engine/Source/Runtime/AudioMixer/Private/SoundWaveDecoder.cpp
@@ -156,14 +156,13 @@ namespace Audio
 
 	void FDecodingSoundSource::ReadFrame()
 	{
-		if (!MixerSourceBuffer.IsValid() || !SourceInfo.CurrentPCMBuffer.IsValid())
+		if (!MixerSourceBuffer.IsValid())
 		{
 			SourceInfo.bIsLastBuffer = true;
 			return;
 		}
 
 		TSharedPtr<FMixerSourceBuffer, ESPMode::ThreadSafe>MixerSourceBufferLocal = MixerSourceBuffer;
-		TSharedPtr<FMixerSourceVoiceBuffer, ESPMode::ThreadSafe>CurrentPCMBufferLocal = SourceInfo.CurrentPCMBuffer;
 
 		bool bNextFrameOutOfRange = (SourceInfo.CurrentFrameIndex + 1) >= SourceInfo.CurrentAudioChunkNumFrames;
 		bool bCurrentFrameOutOfRange = SourceInfo.CurrentFrameIndex >= SourceInfo.CurrentAudioChunkNumFrames;
@@ -176,7 +175,7 @@ namespace Audio
 			{
 				bReadCurrentFrame = false;
 
-				const float* AudioData = CurrentPCMBufferLocal->AudioData.GetData();
+				const float* AudioData = SourceInfo.CurrentPCMBuffer->AudioData.GetData();
 
 				if (!AudioData)
 				{
@@ -191,7 +190,7 @@ namespace Audio
 					SourceInfo.CurrentFrameValues[Channel] = AudioData[CurrentSampleIndex + Channel];
 				}
 
-				if (CurrentPCMBufferLocal->LoopCount == Audio::LOOP_FOREVER && !CurrentPCMBufferLocal->bRealTimeBuffer)
+				if (SourceInfo.CurrentPCMBuffer->LoopCount == Audio::LOOP_FOREVER && !SourceInfo.CurrentPCMBuffer->bRealTimeBuffer)
 				{
 					SourceInfo.CurrentFrameIndex = FMath::Max(SourceInfo.CurrentFrameIndex - SourceInfo.CurrentAudioChunkNumFrames, 0);
 					break;
@@ -203,14 +202,14 @@ namespace Audio
 			if (MixerSourceBufferLocal->GetNumBuffersQueued() > 0 && (SourceInfo.NumSourceChannels > 0))
 			{
 				check(MixerSourceBufferLocal.IsValid());
-				CurrentPCMBufferLocal = MixerSourceBufferLocal->GetNextBuffer();
-				if (!CurrentPCMBufferLocal)
+				SourceInfo.CurrentPCMBuffer = MixerSourceBufferLocal->GetNextBuffer();
+				if (!SourceInfo.CurrentPCMBuffer)
 				{
 					SourceInfo.bIsLastBuffer = true;
 					return;
 				}
 
-				SourceInfo.CurrentAudioChunkNumFrames = CurrentPCMBufferLocal->AudioData.Num() / SourceInfo.NumSourceChannels;
+				SourceInfo.CurrentAudioChunkNumFrames = SourceInfo.CurrentPCMBuffer->AudioData.Num() / SourceInfo.NumSourceChannels;
 
 				if (bReadCurrentFrame)
 				{
@@ -231,7 +230,7 @@ namespace Audio
 			bCurrentFrameOutOfRange = SourceInfo.CurrentFrameIndex >= SourceInfo.CurrentAudioChunkNumFrames;
 		}
 
-		const float* AudioData = CurrentPCMBufferLocal->AudioData.GetData();
+		const float* AudioData = SourceInfo.CurrentPCMBuffer->AudioData.GetData();
 
 		if (!AudioData)
 		{
@@ -296,7 +295,7 @@ namespace Audio
 			}
 
 
-			if (!SourceInfo.CurrentPCMBuffer.IsValid() || !MixerSourceBuffer.IsValid())
+			if (!MixerSourceBuffer.IsValid())
 			{
 				bReadFrame = false;
 				SourceInfo.bIsLastBuffer = true;
