Fix for invalid USoundSubmix ptr on shutdown being passed to be unloaded

diff --git a/Engine/Source/Runtime/AudioMixer/Private/AudioMixerDevice.cpp b/Engine/Source/Runtime/AudioMixer/Private/AudioMixerDevice.cpp
--- a/Engine/Source/Runtime/AudioMixer/Private/AudioMixerDevice.cpp
+++ b/Engine/Source/Runtime/AudioMixer/Private/AudioMixerDevice.cpp
@@ -1265,10 +1265,14 @@ namespace Audio
 		{
 			DECLARE_CYCLE_STAT(TEXT("FAudioThreadTask.UnregisterSoundSubmix"), STAT_AudioUnregisterSoundSubmix, STATGROUP_AudioThreadCommands);
 
-			FAudioThread::RunCommandOnAudioThread([this, InSoundSubmix]()
+			TWeakObjectPtr<USoundSubmix> SubmixToUnload = InSoundSubmix;
+			FAudioThread::RunCommandOnAudioThread([this, SubmixToUnload]()
 			{
 				CSV_SCOPED_TIMING_STAT(Audio, UnregisterSubmix);
-				UnregisterSoundSubmix(InSoundSubmix);
+				if (SubmixToUnload.IsValid())
+				{
+					UnloadSoundSubmix(*SubmixToUnload.Get());
+				}
 			}, GET_STATID(STAT_AudioUnregisterSoundSubmix));
 			return;
 		}
