Fix build due to collision with Submix Endpoint updates

diff --git a/Engine/Source/Runtime/AudioMixer/Private/AudioMixerDevice.cpp b/Engine/Source/Runtime/AudioMixer/Private/AudioMixerDevice.cpp
--- a/Engine/Source/Runtime/AudioMixer/Private/AudioMixerDevice.cpp
+++ b/Engine/Source/Runtime/AudioMixer/Private/AudioMixerDevice.cpp
@@ -1265,7 +1265,7 @@ namespace Audio
 		{
 			DECLARE_CYCLE_STAT(TEXT("FAudioThreadTask.UnregisterSoundSubmix"), STAT_AudioUnregisterSoundSubmix, STATGROUP_AudioThreadCommands);
 
-			TWeakObjectPtr<USoundSubmix> SubmixToUnload = InSoundSubmix;
+			TWeakObjectPtr<USoundSubmixBase> SubmixToUnload = InSoundSubmix;
 			FAudioThread::RunCommandOnAudioThread([this, SubmixToUnload]()
 			{
 				CSV_SCOPED_TIMING_STAT(Audio, UnregisterSubmix);
