Adding some diagnostics to help trackdown OOM issue.

diff --git a/Engine/Source/Runtime/AudioMixer/Private/AudioMixerSourceManager.cpp b/Engine/Source/Runtime/AudioMixer/Private/AudioMixerSourceManager.cpp
--- a/Engine/Source/Runtime/AudioMixer/Private/AudioMixerSourceManager.cpp
+++ b/Engine/Source/Runtime/AudioMixer/Private/AudioMixerSourceManager.cpp
@@ -289,6 +289,23 @@ namespace Audio
 			MaxChannelsSupportedBySpatializationPlugin = MixerDevice->MaxChannelsSupportedBySpatializationPlugin;
 		}
 
+		// Spam command queue with nops.
+		static FAutoConsoleCommand SpamNopsCmd(
+			TEXT("au.SpamCommandQueue"),
+			TEXT(""),
+			FConsoleCommandDelegate::CreateLambda([this]() 
+			{				
+				struct FSpamPayload
+				{
+					uint8 JunkBytes[1024];
+				} Payload;
+				for (int32 i = 0; i < 65536; ++i)
+				{
+					AudioMixerThreadCommand([Payload] {});
+				}
+			})
+		);
+
 		bInitialized = true;
 		bPumpQueue = false;
 	}
@@ -2686,6 +2703,20 @@ namespace Audio
 
 		// Add the function to the command queue:
 		int32 AudioThreadCommandIndex = !RenderThreadCommandBufferIndex.GetValue();
+		
+#if !NO_LOGGING
+		static uint32 WarnSize = 1024 * 1024;
+		SIZE_T Size = CommandBuffers[AudioThreadCommandIndex].SourceCommandQueue.GetAllocatedSize();
+		if (Size > WarnSize )
+		{		
+			SIZE_T Num = CommandBuffers[AudioThreadCommandIndex].SourceCommandQueue.Num();
+			// NOTE: Although not really and error we want this to show up in shipping builds.
+			UE_LOG(LogAudioMixer, Error, TEXT("Command Queue has grown to %uk bytes, containing %d cmds, last pump was %fms ago."), 
+				Size >> 10, Num, FPlatformTime::ToMilliseconds64(FPlatformTime::Cycles64() - LastPumpTimeInCycles));
+			WarnSize *= 2;
+		}
+#endif //!NO_LOGGING
+
 		CommandBuffers[AudioThreadCommandIndex].SourceCommandQueue.Add(MoveTemp(InFunction));
 		NumCommands.Increment();
 	}
@@ -2713,6 +2744,7 @@ namespace Audio
 			NumCommands.Decrement();
 		}
 
+		LastPumpTimeInCycles = FPlatformTime::Cycles64();
 		Commands.SourceCommandQueue.Reset();
 
 		if (FPlatformProcess::SupportsMultithreading())
diff --git a/Engine/Source/Runtime/AudioMixer/Private/AudioMixerSourceManager.h b/Engine/Source/Runtime/AudioMixer/Private/AudioMixerSourceManager.h
--- a/Engine/Source/Runtime/AudioMixer/Private/AudioMixerSourceManager.h
+++ b/Engine/Source/Runtime/AudioMixer/Private/AudioMixerSourceManager.h
@@ -584,6 +584,7 @@ namespace Audio
 
 		// Set to true when the audio source manager should pump the command queue
 		FThreadSafeBool bPumpQueue;
+		uint64 LastPumpTimeInCycles = 0;
 
 		friend class FMixerSourceVoice;
 	};
