Fix for heisencrash caused by race condition on app init between FMixerSubmix::InitInternal on the audio thread mutating the EffectChain array andFMixerSubmix::ProcessAudio already being in-flight.

#localization none
#tests none

diff --git a/Engine/Source/Runtime/AudioMixer/Private/AudioMixerSubmix.cpp b/Engine/Source/Runtime/AudioMixer/Private/AudioMixerSubmix.cpp
--- a/Engine/Source/Runtime/AudioMixer/Private/AudioMixerSubmix.cpp
+++ b/Engine/Source/Runtime/AudioMixer/Private/AudioMixerSubmix.cpp
@@ -93,10 +93,14 @@ namespace Audio
 			CurrentOutputVolume = FMath::Clamp(SoundSubmix->OutputVolume, 0.0f, 1.0f);
 			TargetOutputVolume = CurrentOutputVolume;
 
+		{
+			FScopeLock ScopeLock(&EffectChainMutationCriticalSection);
 			// Set the initialized output volume
 			CurrentWetLevel = FMath::Clamp(SoundSubmix->WetLevel, 0.0f, 1.0f);
 			TargetWetLevel = CurrentWetLevel;
 
+			// Loop through the submix's presets and make new instances of effects in the same order as the presets
+			ClearSoundEffectSubmixes();
 			CurrentDryLevel = FMath::Clamp(SoundSubmix->DryLevel, 0.0f, 1.0f);
 			TargetDryLevel = CurrentDryLevel;
 
@@ -365,6 +369,8 @@ namespace Audio
 
 	void FMixerSubmix::ClearSoundEffectSubmixes()
 	{
+		FScopeLock ScopeLock(&EffectChainMutationCriticalSection);
+
 		for (FSubmixEffectInfo& Info : EffectSubmixChain)
 		{
 			if (Info.EffectInstance.IsValid())
@@ -379,6 +385,8 @@ namespace Audio
 
 	void FMixerSubmix::ReplaceSoundEffectSubmix(int32 InIndex, int32 InPresetId, FSoundEffectSubmixPtr InEffectInstance)
 	{
+		FScopeLock ScopeLock(&EffectChainMutationCriticalSection);
+
 		if (InEffectInstance->GetPreset() != nullptr && InIndex < EffectSubmixChain.Num())
 		{
 			EffectSubmixChain[InIndex].PresetId = InPresetId;
@@ -805,41 +813,49 @@ namespace Audio
 			}
 		}
 
-		if (EffectSubmixChain.Num() > 0)
 		{
-			CSV_SCOPED_TIMING_STAT(Audio, SubmixEffectProcessing);
-
-			// Setup the input data buffer
-			FSoundEffectSubmixInputData InputData;
-			InputData.AudioClock = MixerDevice->GetAudioTime();
+			FScopeLock ScopeLock(&EffectChainMutationCriticalSection);
+			if (EffectSubmixChain.Num() > 0)
+			{
+				CSV_SCOPED_TIMING_STAT(Audio, SubmixEffectProcessing);
 
-			// Compute the number of frames of audio. This will be independent of if we downmix our wet buffer.
-			InputData.NumFrames = NumSamples / NumChannels;
-			InputData.NumChannels = NumChannels;
-			InputData.NumDeviceChannels = MixerDevice->GetNumDeviceChannels();
-			InputData.ListenerTransforms = MixerDevice->GetListenerTransforms();
-			InputData.AudioClock = MixerDevice->GetAudioClock();
+				// Setup the input data buffer
+				FSoundEffectSubmixInputData InputData;
+				InputData.AudioClock = MixerDevice->GetAudioTime();
 
-			FSoundEffectSubmixOutputData OutputData;
-			OutputData.AudioBuffer = &ScratchBuffer;
-			OutputData.NumChannels = NumChannels;
+				// Compute the number of frames of audio. This will be independent of if we downmix our wet buffer.
+				InputData.NumFrames = NumSamples / NumChannels;
+				InputData.NumChannels = NumChannels;
+				InputData.NumDeviceChannels = MixerDevice->GetNumDeviceChannels();
+				InputData.ListenerTransforms = MixerDevice->GetListenerTransforms();
+				InputData.AudioClock = MixerDevice->GetAudioClock();
 
-			for (FSubmixEffectInfo& SubmixEffectInfo : EffectSubmixChain)
-			{
-				FSoundEffectSubmixPtr SubmixEffect = SubmixEffectInfo.EffectInstance;
+				FSoundEffectSubmixOutputData OutputData;
+				OutputData.AudioBuffer = &ScratchBuffer;
+				OutputData.NumChannels = NumChannels;
 
-				// SubmixEffectInfo.EffectInstance will be null if FMixerSubmix::RemoveSoundEffectSubmix was called earlier.
-				if (!SubmixEffect.IsValid())
+				for (FSubmixEffectInfo& SubmixEffectInfo : EffectSubmixChain)
 				{
-					continue;
-				}
+					FSoundEffectSubmixPtr SubmixEffect = SubmixEffectInfo.EffectInstance;
 
-				// Reset the output scratch buffer
-				ScratchBuffer.Reset(NumSamples);
-				ScratchBuffer.AddZeroed(NumSamples);
+					// SubmixEffectInfo.EffectInstance will be null if FMixerSubmix::RemoveSoundEffectSubmix was called earlier.
+					if (!SubmixEffect.IsValid())
+					{
+						continue;
+					}
 
-				// Check to see if we need to down-mix our audio before sending to the submix effect
-				const uint32 ChannelCountOverride = SubmixEffect->GetDesiredInputChannelCountOverride();
+					// Reset the output scratch buffer
+					ScratchBuffer.Reset(NumSamples);
+					ScratchBuffer.AddZeroed(NumSamples);
+
+					// Check to see if we need to down-mix our audio before sending to the submix effect
+					const uint32 ChannelCountOverride = SubmixEffect->GetDesiredInputChannelCountOverride();
+
+					// Only support downmixing to stereo. TODO: change GetDesiredInputChannelCountOverride() API to be "DownmixToStereo"
+					if (ChannelCountOverride < (uint32)NumChannels && ChannelCountOverride == 2)
+					{
+						// Perform the down-mix operation with the down-mixed scratch buffer
+						FormatChangeBuffer(ESubmixChannelFormat::Stereo, InputBuffer, DownmixedBuffer);
 
 				if (ChannelCountOverride != INDEX_NONE && ChannelCountOverride != NumChannels)
 				{
@@ -848,23 +864,14 @@ namespace Audio
 					DownmixedBuffer.AddUninitialized(NumOutputFrames * ChannelCountOverride);
 					DownmixBuffer(NumChannels, InputBuffer, ChannelCountOverride, DownmixedBuffer);
 
-					InputData.NumChannels = ChannelCountOverride;
-					InputData.AudioBuffer = &DownmixedBuffer;
-					SubmixEffect->ProcessAudio(InputData, OutputData);
-				}
-				else
-				{
-					// If we're not down-mixing, then just pass in the current wet buffer and our channel count is the same as the output channel count
-					InputData.NumChannels = NumChannels;
-					InputData.AudioBuffer = &InputBuffer;
-					SubmixEffect->ProcessAudio(InputData, OutputData);
-				}
+					// Mix in the dry signal directly
+					const float DryLevel = SubmixEffect->GetDryLevel();
+					if (DryLevel > 0.0f)
+					{
+						Audio::MixInBufferFast(InputBuffer, ScratchBuffer, DryLevel);
+					}
 
-				// Mix in the dry signal directly
-				const float DryLevel = SubmixEffect->GetDryLevel();
-				if (DryLevel > 0.0f)
-				{
-					Audio::MixInBufferFast(InputBuffer, ScratchBuffer, DryLevel);
+					FMemory::Memcpy((void*)BufferPtr, (void*)ScratchBuffer.GetData(), sizeof(float) * NumSamples);
 				}
 
 				FMemory::Memcpy((void*)BufferPtr, (void*)ScratchBuffer.GetData(), sizeof(float) * NumSamples);
diff --git a/Engine/Source/Runtime/AudioMixer/Public/AudioMixerSubmix.h b/Engine/Source/Runtime/AudioMixer/Public/AudioMixerSubmix.h
--- a/Engine/Source/Runtime/AudioMixer/Public/AudioMixerSubmix.h
+++ b/Engine/Source/Runtime/AudioMixer/Public/AudioMixerSubmix.h
@@ -473,6 +473,9 @@ namespace Audio
 		// Critical section used for when we are appending recorded data.
 		FCriticalSection RecordingCriticalSection;
 
+		// Critical section for mutation of the effect chain.
+		FCriticalSection EffectChainMutationCriticalSection;
+
 		// Handle back to the owning USoundSubmix. Used when the device is shutdown to prematurely end a recording.
 		const USoundSubmixBase* OwningSubmixObject;
 
