Fix of bad merge resolution.

diff --git a/Engine/Source/Runtime/AudioMixer/Private/AudioMixerSubmix.cpp b/Engine/Source/Runtime/AudioMixer/Private/AudioMixerSubmix.cpp
--- a/Engine/Source/Runtime/AudioMixer/Private/AudioMixerSubmix.cpp
+++ b/Engine/Source/Runtime/AudioMixer/Private/AudioMixerSubmix.cpp
@@ -99,8 +99,6 @@ namespace Audio
 			CurrentWetLevel = FMath::Clamp(SoundSubmix->WetLevel, 0.0f, 1.0f);
 			TargetWetLevel = CurrentWetLevel;
 
-			// Loop through the submix's presets and make new instances of effects in the same order as the presets
-			ClearSoundEffectSubmixes();
 			CurrentDryLevel = FMath::Clamp(SoundSubmix->DryLevel, 0.0f, 1.0f);
 			TargetDryLevel = CurrentDryLevel;
 
@@ -851,12 +849,6 @@ namespace Audio
 					// Check to see if we need to down-mix our audio before sending to the submix effect
 					const uint32 ChannelCountOverride = SubmixEffect->GetDesiredInputChannelCountOverride();
 
-					// Only support downmixing to stereo. TODO: change GetDesiredInputChannelCountOverride() API to be "DownmixToStereo"
-					if (ChannelCountOverride < (uint32)NumChannels && ChannelCountOverride == 2)
-					{
-						// Perform the down-mix operation with the down-mixed scratch buffer
-						FormatChangeBuffer(ESubmixChannelFormat::Stereo, InputBuffer, DownmixedBuffer);
-
 				if (ChannelCountOverride != INDEX_NONE && ChannelCountOverride != NumChannels)
 				{
 					// Perform the down-mix operation with the down-mixed scratch buffer
@@ -864,14 +856,23 @@ namespace Audio
 					DownmixedBuffer.AddUninitialized(NumOutputFrames * ChannelCountOverride);
 					DownmixBuffer(NumChannels, InputBuffer, ChannelCountOverride, DownmixedBuffer);
 
-					// Mix in the dry signal directly
-					const float DryLevel = SubmixEffect->GetDryLevel();
-					if (DryLevel > 0.0f)
-					{
-						Audio::MixInBufferFast(InputBuffer, ScratchBuffer, DryLevel);
-					}
+					InputData.NumChannels = ChannelCountOverride;
+					InputData.AudioBuffer = &DownmixedBuffer;
+					SubmixEffect->ProcessAudio(InputData, OutputData);
+				}
+				else
+				{
+					// If we're not down-mixing, then just pass in the current wet buffer and our channel count is the same as the output channel count
+					InputData.NumChannels = NumChannels;
+					InputData.AudioBuffer = &InputBuffer;
+					SubmixEffect->ProcessAudio(InputData, OutputData);
+				}
 
-					FMemory::Memcpy((void*)BufferPtr, (void*)ScratchBuffer.GetData(), sizeof(float) * NumSamples);
+				// Mix in the dry signal directly
+				const float DryLevel = SubmixEffect->GetDryLevel();
+				if (DryLevel > 0.0f)
+				{
+					Audio::MixInBufferFast(InputBuffer, ScratchBuffer, DryLevel);
 				}
 
 				FMemory::Memcpy((void*)BufferPtr, (void*)ScratchBuffer.GetData(), sizeof(float) * NumSamples);
