Fix extraneous bracket from merge.

diff --git a/Engine/Source/Runtime/AudioMixer/Private/AudioMixerSubmix.cpp b/Engine/Source/Runtime/AudioMixer/Private/AudioMixerSubmix.cpp
--- a/Engine/Source/Runtime/AudioMixer/Private/AudioMixerSubmix.cpp
+++ b/Engine/Source/Runtime/AudioMixer/Private/AudioMixerSubmix.cpp
@@ -93,7 +93,7 @@ namespace Audio
 			CurrentOutputVolume = FMath::Clamp(SoundSubmix->OutputVolume, 0.0f, 1.0f);
 			TargetOutputVolume = CurrentOutputVolume;
 
-		{
+
 			FScopeLock ScopeLock(&EffectChainMutationCriticalSection);
 			// Set the initialized output volume
 			CurrentWetLevel = FMath::Clamp(SoundSubmix->WetLevel, 0.0f, 1.0f);
