Fix missing bracket from merge.

diff --git a/Engine/Source/Runtime/AudioMixer/Private/AudioMixerSubmix.cpp b/Engine/Source/Runtime/AudioMixer/Private/AudioMixerSubmix.cpp
--- a/Engine/Source/Runtime/AudioMixer/Private/AudioMixerSubmix.cpp
+++ b/Engine/Source/Runtime/AudioMixer/Private/AudioMixerSubmix.cpp
@@ -849,23 +849,24 @@ namespace Audio
 					// Check to see if we need to down-mix our audio before sending to the submix effect
 					const uint32 ChannelCountOverride = SubmixEffect->GetDesiredInputChannelCountOverride();
 
-				if (ChannelCountOverride != INDEX_NONE && ChannelCountOverride != NumChannels)
-				{
-					// Perform the down-mix operation with the down-mixed scratch buffer
-					DownmixedBuffer.Reset();
-					DownmixedBuffer.AddUninitialized(NumOutputFrames * ChannelCountOverride);
-					DownmixBuffer(NumChannels, InputBuffer, ChannelCountOverride, DownmixedBuffer);
-
-					InputData.NumChannels = ChannelCountOverride;
-					InputData.AudioBuffer = &DownmixedBuffer;
-					SubmixEffect->ProcessAudio(InputData, OutputData);
-				}
-				else
-				{
-					// If we're not down-mixing, then just pass in the current wet buffer and our channel count is the same as the output channel count
-					InputData.NumChannels = NumChannels;
-					InputData.AudioBuffer = &InputBuffer;
-					SubmixEffect->ProcessAudio(InputData, OutputData);
+					if (ChannelCountOverride != INDEX_NONE && ChannelCountOverride != NumChannels)
+					{
+						// Perform the down-mix operation with the down-mixed scratch buffer
+						DownmixedBuffer.Reset();
+						DownmixedBuffer.AddUninitialized(NumOutputFrames * ChannelCountOverride);
+						DownmixBuffer(NumChannels, InputBuffer, ChannelCountOverride, DownmixedBuffer);
+
+						InputData.NumChannels = ChannelCountOverride;
+						InputData.AudioBuffer = &DownmixedBuffer;
+						SubmixEffect->ProcessAudio(InputData, OutputData);
+					}
+					else
+					{
+						// If we're not down-mixing, then just pass in the current wet buffer and our channel count is the same as the output channel count
+						InputData.NumChannels = NumChannels;
+						InputData.AudioBuffer = &InputBuffer;
+						SubmixEffect->ProcessAudio(InputData, OutputData);
+					}
 				}
 
 				// Mix in the dry signal directly
