Compile fix: Fix out of scope reference downstream.

diff --git a/Engine/Source/Runtime/AudioMixer/Private/AudioMixerSubmix.cpp b/Engine/Source/Runtime/AudioMixer/Private/AudioMixerSubmix.cpp
--- a/Engine/Source/Runtime/AudioMixer/Private/AudioMixerSubmix.cpp
+++ b/Engine/Source/Runtime/AudioMixer/Private/AudioMixerSubmix.cpp
@@ -867,13 +867,13 @@ namespace Audio
 						InputData.AudioBuffer = &InputBuffer;
 						SubmixEffect->ProcessAudio(InputData, OutputData);
 					}
-				}
-
-				// Mix in the dry signal directly
-				const float DryLevel = SubmixEffect->GetDryLevel();
-				if (DryLevel > 0.0f)
-				{
-					Audio::MixInBufferFast(InputBuffer, ScratchBuffer, DryLevel);
+					
+					// Mix in the dry signal directly
+					const float DryLevel = SubmixEffect->GetDryLevel();
+					if (DryLevel > 0.0f)
+					{
+						Audio::MixInBufferFast(InputBuffer, ScratchBuffer, DryLevel);
+					}
 				}
 
 				FMemory::Memcpy((void*)BufferPtr, (void*)ScratchBuffer.GetData(), sizeof(float) * NumSamples);
