Series of speculative fixes for observed submix initialization crash. Main fix here is no matter what, skip over submix effect processing until we change the CVar.

diff --git a/Engine/Plugins/Runtime/ResonanceAudio/Source/ResonanceAudio/Private/ResonanceAudioReverb.cpp b/Engine/Plugins/Runtime/ResonanceAudio/Source/ResonanceAudio/Private/ResonanceAudioReverb.cpp
--- a/Engine/Plugins/Runtime/ResonanceAudio/Source/ResonanceAudio/Private/ResonanceAudioReverb.cpp
+++ b/Engine/Plugins/Runtime/ResonanceAudio/Source/ResonanceAudio/Private/ResonanceAudioReverb.cpp
@@ -9,6 +9,15 @@
 #include "ResonanceAudioSettings.h"
 #include "Sound/SoundSubmix.h"
 
+static int32 IgnoreUserResonanceSubmixCVar = 1;
+FAutoConsoleVariableRef CVarIgnoreUserResonanceSubmix(
+	TEXT("au.IgnoreUserResonanceSubmix"),
+	IgnoreUserResonanceSubmixCVar,
+	TEXT("When set to 1, the resonance project setting will be bypassed.\n")
+	TEXT("1: Submix Effects are disabled."),
+	ECVF_Default);
+
+
 namespace ResonanceAudio
 {
 	/*********************************************/
@@ -158,7 +167,13 @@ namespace ResonanceAudio
 		const UResonanceAudioSettings* Settings = GetDefault<UResonanceAudioSettings>();
 		check(Settings);
 
-		USoundSubmix* ReverbSubmix = Cast<USoundSubmix>(Settings->OutputSubmix.TryLoad());
+		USoundSubmix* ReverbSubmix = nullptr;
+		
+		if (!IgnoreUserResonanceSubmixCVar)
+		{
+			ReverbSubmix = Cast<USoundSubmix>(Settings->OutputSubmix.TryLoad());
+		}
+
 		if (!ReverbSubmix)
 		{
 			static const FString DefaultSubmixName = TEXT("Resonance Reverb Submix");
diff --git a/Engine/Source/Runtime/AudioMixer/Private/AudioMixerDevice.cpp b/Engine/Source/Runtime/AudioMixer/Private/AudioMixerDevice.cpp
--- a/Engine/Source/Runtime/AudioMixer/Private/AudioMixerDevice.cpp
+++ b/Engine/Source/Runtime/AudioMixer/Private/AudioMixerDevice.cpp
@@ -51,7 +51,7 @@ namespace Audio
 		, GameOrAudioThreadId(INDEX_NONE)
 		, AudioPlatformThreadId(INDEX_NONE)
 		, bDebugOutputEnabled(false)
-		, bSubmixRegistrationDisabled(false)
+		, bSubmixRegistrationDisabled(true)
 	{
 		// This audio device is the audio mixer
 		bAudioMixerModuleLoaded = true;
@@ -1156,7 +1156,13 @@ namespace Audio
 
 	void FMixerDevice::RegisterSoundSubmix(const USoundSubmixBase* InSoundSubmix, bool bInit)
 	{
-		if (!InSoundSubmix || bSubmixRegistrationDisabled)
+		if (InSoundSubmix && bSubmixRegistrationDisabled)
+		{
+			UE_LOG(LogAudioMixer, Warning, TEXT("Attempted register Submix %s before the submix graph was initialized."), *InSoundSubmix->GetFullName());
+			return;
+		}
+
+		if (!InSoundSubmix)
 		{
 			return;
 		}
@@ -1174,6 +1180,8 @@ namespace Audio
 			return;
 		}
 
+		UE_LOG(LogAudioMixer, Display, TEXT("Registering submix %s."), *InSoundSubmix->GetFullName());
+
 		const bool bIsMasterSubmix = IsMasterSubmixType(InSoundSubmix);
 
 		if (!bIsMasterSubmix)
@@ -1191,6 +1199,10 @@ namespace Audio
 
 			LoadSoundSubmix(*InSoundSubmix);
 		}
+		else
+		{
+			UE_LOG(LogAudioMixer, Display, TEXT("Submix %s was already registered as one of the master submixes."), *InSoundSubmix->GetFullName());
+		}
 
 		FMixerSubmixPtr SubmixPtr = GetSubmixInstance(InSoundSubmix).Pin();
 
diff --git a/Engine/Source/Runtime/AudioMixer/Private/AudioMixerSubmix.cpp b/Engine/Source/Runtime/AudioMixer/Private/AudioMixerSubmix.cpp
--- a/Engine/Source/Runtime/AudioMixer/Private/AudioMixerSubmix.cpp
+++ b/Engine/Source/Runtime/AudioMixer/Private/AudioMixerSubmix.cpp
@@ -19,6 +19,14 @@ FAutoConsoleVariableRef CVarRecoverRecordingOnShutdown(
 	TEXT("0: Disabled, 1: Enabled"),
 	ECVF_Default);
 
+static int32 BypassAllSubmixEffectsCVar = 1;
+FAutoConsoleVariableRef CVarBypassAllSubmixEffects(
+	TEXT("au.BypassAllSubmixEffects"),
+	BypassAllSubmixEffectsCVar,
+	TEXT("When set to 1, all submix effects will be bypassed.\n")
+	TEXT("1: Submix Effects are disabled."),
+	ECVF_Default);
+
 namespace Audio
 {
 	// Unique IDs for mixer submixes
@@ -327,6 +335,7 @@ namespace Audio
 
 	void FMixerSubmix::AddSoundEffectSubmix(uint32 SubmixPresetId, FSoundEffectSubmixPtr InSoundEffectSubmix)
 	{
+		FScopeLock ScopeLock(&EffectChainMutationCriticalSection);
 		AUDIO_MIXER_CHECK_AUDIO_PLAT_THREAD(MixerDevice);
 
 		// Look to see if the submix preset ID is already present
@@ -349,6 +358,7 @@ namespace Audio
 
 	void FMixerSubmix::RemoveSoundEffectSubmix(uint32 SubmixPresetId)
 	{
+		FScopeLock ScopeLock(&EffectChainMutationCriticalSection);
 		AUDIO_MIXER_CHECK_AUDIO_PLAT_THREAD(MixerDevice);
 
 		for (FSubmixEffectInfo& Effect : EffectSubmixChain)
@@ -813,7 +823,7 @@ namespace Audio
 
 		{
 			FScopeLock ScopeLock(&EffectChainMutationCriticalSection);
-			if (EffectSubmixChain.Num() > 0)
+			if (!BypassAllSubmixEffectsCVar && EffectSubmixChain.Num() > 0)
 			{
 				CSV_SCOPED_TIMING_STAT(Audio, SubmixEffectProcessing);
 
