Set BypassAllSubmixEffects back to 0.

diff --git a/Engine/Source/Runtime/AudioMixer/Private/AudioMixerSubmix.cpp b/Engine/Source/Runtime/AudioMixer/Private/AudioMixerSubmix.cpp
--- a/Engine/Source/Runtime/AudioMixer/Private/AudioMixerSubmix.cpp
+++ b/Engine/Source/Runtime/AudioMixer/Private/AudioMixerSubmix.cpp
@@ -19,7 +19,7 @@ FAutoConsoleVariableRef CVarRecoverRecordingOnShutdown(
 	TEXT("0: Disabled, 1: Enabled"),
 	ECVF_Default);
 
-static int32 BypassAllSubmixEffectsCVar = 1;
+static int32 BypassAllSubmixEffectsCVar = 0;
 FAutoConsoleVariableRef CVarBypassAllSubmixEffects(
 	TEXT("au.BypassAllSubmixEffects"),
 	BypassAllSubmixEffectsCVar,
