Editgrate ambisonics rotation fix.

diff --git a/Engine/Plugins/Runtime/SoundFields/Source/SoundFields/Private/SoundFields.cpp b/Engine/Plugins/Runtime/SoundFields/Source/SoundFields/Private/SoundFields.cpp
--- a/Engine/Plugins/Runtime/SoundFields/Source/SoundFields/Private/SoundFields.cpp
+++ b/Engine/Plugins/Runtime/SoundFields/Source/SoundFields/Private/SoundFields.cpp
@@ -49,13 +49,27 @@ private:
 	FQuat PreviousRotation = FQuat::Identity;
 
 	FAmbisonicsSoundfieldBuffer RotatedAudio;
+	Audio::AlignedFloatBuffer DecoderOutputBuffer;
 
 public:
 	virtual ~FAmbisonicsSoundfieldDecoder() {};
 
 	virtual void Decode(const FSoundfieldDecoderInputData& InputData, FSoundfieldDecoderOutputData& OutputData) override
 	{
-		DecodeAndMixIn(InputData, OutputData);
+		const FAmbisonicsSoundfieldBuffer& InputAudio = DowncastSoundfieldRef<const FAmbisonicsSoundfieldBuffer>(InputData.SoundfieldBuffer);
+		const FSoundfieldSpeakerPositionalData& OutputPositions = InputData.PositionalData;
+		Audio::AlignedFloatBuffer& OutputAudio = OutputData.AudioBuffer;
+
+		const FAmbisonicsSoundfieldBuffer* AudioToDecode = &InputAudio;
+
+		if (VirtualIntermediateChannelsCvar)
+		{
+			Decoder.DecodeAudioToSevenOneAndDownmixToDevice(*AudioToDecode, OutputPositions, OutputAudio);
+		}
+		else
+		{
+			Decoder.DecodeAudioDirectlyToDeviceOutputPositions(*AudioToDecode, OutputPositions, OutputAudio);
+		}
 	}
 
 
@@ -67,24 +81,19 @@ public:
 
 		const FAmbisonicsSoundfieldBuffer* AudioToDecode = &InputAudio;
 
-		const bool bShouldRotateOutsideOfDecoder = false;
-		if (bShouldRotateOutsideOfDecoder && InputAudio.NumChannels == 4)
-		{
-			RotatedAudio.AudioBuffer.Reset();
-			RotatedAudio.AudioBuffer.AddUninitialized(InputAudio.AudioBuffer.Num());
-			RotatedAudio.NumChannels = InputAudio.NumChannels;
-			FSoundFieldDecoder::RotateFirstOrderAmbisonicsBed(InputAudio, RotatedAudio, OutputPositions.Rotation, PreviousRotation);
-			AudioToDecode = &RotatedAudio;
-		}
+		DecoderOutputBuffer.Reset();
+		DecoderOutputBuffer.AddUninitialized(OutputData.AudioBuffer.Num());
 
 		if (VirtualIntermediateChannelsCvar)
 		{
-			Decoder.DecodeAudioToSevenOneAndDownmixToDevice(*AudioToDecode, OutputPositions, OutputAudio);
+			Decoder.DecodeAudioToSevenOneAndDownmixToDevice(*AudioToDecode, OutputPositions, DecoderOutputBuffer);
 		}
 		else
 		{
-			Decoder.DecodeAudioDirectlyToDeviceOutputPositions(*AudioToDecode, OutputPositions, OutputAudio);
+			Decoder.DecodeAudioDirectlyToDeviceOutputPositions(*AudioToDecode, OutputPositions, DecoderOutputBuffer);
 		}
+
+		Audio::MixInBufferFast(DecoderOutputBuffer, OutputAudio);
 	}
 };
 
@@ -166,7 +175,7 @@ public:
 };
 
 /**
- * Class that sums together two ambisonics streams.
+ * Class that sums an ambisonics stream into an output stream, rotating the ambisonics stream as necessary to get the ambisonics field in world coordinates.
  */
 class FAmbisonicsMixer : public ISoundfieldMixerStream
 {
@@ -181,7 +190,7 @@ public:
 		const FAmbisonicsSoundfieldBuffer& InAudio = DowncastSoundfieldRef<const FAmbisonicsSoundfieldBuffer>(InputData.InputPacket);
 		FAmbisonicsSoundfieldBuffer& OutAudio = DowncastSoundfieldRef<FAmbisonicsSoundfieldBuffer>(PacketToSumTo);
 
-		// Lazily initialize our output packet.
+		// Lazily initialize our output and rotated packet.
 		if (OutAudio.NumChannels == 0)
 		{
 			OutAudio.NumChannels = InAudio.NumChannels;
diff --git a/Engine/Source/Runtime/AudioMixer/Private/AudioMixerSourceManager.cpp b/Engine/Source/Runtime/AudioMixer/Private/AudioMixerSourceManager.cpp
--- a/Engine/Source/Runtime/AudioMixer/Private/AudioMixerSourceManager.cpp
+++ b/Engine/Source/Runtime/AudioMixer/Private/AudioMixerSourceManager.cpp
@@ -678,6 +678,7 @@ namespace Audio
 			SourceInfo.bUseHRTFSpatializer = InitParams.bUseHRTFSpatialization;
 			SourceInfo.bIsExternalSend = InitParams.bIsExternalSend;
 			SourceInfo.bIsVorbis = InitParams.bIsVorbis;
+			SourceInfo.bIsAmbisonics = InitParams.bIsAmbisonics;
 			SourceInfo.AudioComponentID = InitParams.AudioComponentID;
 			SourceInfo.bIsAmbisonics = InitParams.bIsAmbisonics;
 
@@ -871,7 +872,8 @@ namespace Audio
 					}
 					else
 					{
-						FSubmixChannelData& SubmixChannelInfo = GetChannelInfoForDevice(DownmixData);
+						// Flag this source as needing to downmix its audio.
+						DownmixData.bIsSourceBeingSentToDeviceSubmix = true;
 					}
 				}
 			}
@@ -2016,7 +2018,6 @@ namespace Audio
 
 		// First, build our SoundfieldSpeakerPositionalData out of our SpatializationParams.
 		DownmixData.PositionalData.NumChannels = DownmixData.NumInputChannels;
-		DownmixData.PositionalData.Rotation = InSource.SpatParams.EmitterWorldRotation;
 		
 		// Build our input channel positions from the Spatialization params if this is a 3D source, otherwise use the default channel positions.
 		DownmixData.InputChannelPositions.Reset();
@@ -2069,7 +2070,8 @@ namespace Audio
 
 			SoundfieldData.EncodedPacket->Reset();
 
-			// If this is an ambisonics source, transcode it. Otherwise, encode the source.
+			// If this is an ambisonics source, transcode it, or if this is going to an ambisonics submix, simply forward the buffer. Otherwise, encode the source 
+			// to whatever soundfield format the destination submix is.
 			if (SoundfieldData.AmbiTranscoder)
 			{
 				FAmbisonicsSoundfieldBuffer AmbiBuffer;
@@ -2098,7 +2100,7 @@ namespace Audio
 				ensure(InSource.bIsAmbisonics);
 
 				FAmbisonicsSoundfieldBuffer& OutputPacket = DowncastSoundfieldRef<FAmbisonicsSoundfieldBuffer>(*SoundfieldData.EncodedPacket);
-				// Fixme: This is an array copy. Can we serve DownmixData.PostEffectBuffers directly to this soundfield?
+				// Fixme: This is an array copy. Can we serve DownmixData.PostEffectBuffers directly to this soundfield, then return it back at the end of the render loop?
 				OutputPacket.AudioBuffer = *DownmixData.PostEffectBuffers;
 				OutputPacket.NumChannels = DownmixData.NumInputChannels;
 				OutputPacket.PreviousRotation = OutputPacket.Rotation;
@@ -2376,14 +2378,18 @@ namespace Audio
 			DownmixData.PositionalData.Rotation = SourceInfo.SpatParams.ListenerOrientation;
 			DownmixData.SourceRotation = SourceInfo.SpatParams.EmitterWorldRotation;
 
-			if (SourceInfo.bIs3D && !DownmixData.bIsInitialDownmix)
-			{
-				ComputeDownmix3D(DownmixData, MixerDevice);
-			}
-			else
+			// If we are sending audio to a non-soundfield submix, we downmix audio to the device configuration here.
+			if (DownmixData.bIsSourceBeingSentToDeviceSubmix)
 			{
-				ComputeDownmix2D(DownmixData, MixerDevice);
-				DownmixData.bIsInitialDownmix = false;
+				if (SourceInfo.bIs3D && !DownmixData.bIsInitialDownmix)
+				{
+					ComputeDownmix3D(DownmixData, MixerDevice);
+				}
+				else
+				{
+					ComputeDownmix2D(DownmixData, MixerDevice);
+					DownmixData.bIsInitialDownmix = false;
+				}
 			}
 
 			if (DownmixData.EncodedSoundfieldDownmixes.Num())
@@ -2442,11 +2448,11 @@ namespace Audio
 		if (SendLevel > 0.0f)
 		{
 			const FSourceInfo& SourceInfo = SourceInfos[SourceId];
+			const FSourceDownmixData& DownmixData = DownmixDataArray[SourceId];
 
 			// Don't need to mix into submixes if the source is paused
-			if (!SourceInfo.bIsPaused && !SourceInfo.bIsDone && SourceInfo.bIsPlaying)
+			if (!SourceInfo.bIsPaused && !SourceInfo.bIsDone && SourceInfo.bIsPlaying && DownmixData.bIsSourceBeingSentToDeviceSubmix)
 			{
-				const FSourceDownmixData& DownmixData = DownmixDataArray[SourceId];
 				const FSubmixChannelData& ChannelInfo = GetChannelInfoForDevice(DownmixData);
 
 				const float* RESTRICT SourceOutputBufferPtr = ChannelInfo.OutputBuffer.GetData();
@@ -2479,6 +2485,12 @@ namespace Audio
 		return nullptr;
 	}
 
+	const FQuat FMixerSourceManager::GetListenerRotation(const int32 SourceId) const
+	{
+		const FSourceDownmixData& DownmixData = DownmixDataArray[SourceId];
+		return DownmixData.PositionalData.Rotation;
+	}
+
 	void FMixerSourceManager::UpdateDeviceChannelCount(const int32 InNumOutputChannels)
 	{
 		AudioMixerThreadCommand([this, InNumOutputChannels]()
diff --git a/Engine/Source/Runtime/AudioMixer/Private/AudioMixerSourceManager.h b/Engine/Source/Runtime/AudioMixer/Private/AudioMixerSourceManager.h
--- a/Engine/Source/Runtime/AudioMixer/Private/AudioMixerSourceManager.h
+++ b/Engine/Source/Runtime/AudioMixer/Private/AudioMixerSourceManager.h
@@ -9,7 +9,6 @@
 #include "AudioMixerBuffer.h"
 #include "AudioMixerSubmix.h"
 #include "AudioMixerBus.h"
-#include "AudioMixerDevice.h"
 #include "DSP/InterpolatedOnePole.h"
 #include "DSP/Filter.h"
 #include "DSP/InterpolatedOnePole.h"
@@ -42,6 +41,8 @@ namespace Audio
 		bool bRealTimeBuffer = false;
 	};
 
+	typedef TSharedPtr<FMixerSubmix, ESPMode::ThreadSafe> FMixerSubmixPtr;
+	typedef TWeakPtr<FMixerSubmix, ESPMode::ThreadSafe> FMixerSubmixWeakPtr;
 
 	class ISourceListener
 	{
@@ -227,6 +228,8 @@ namespace Audio
 		// Returned nonnull pointers are only guaranteed to be valid on the audio mixer render thread.
 		const ISoundfieldAudioPacket* GetEncodedOutput(const int32 SourceId, const FSoundfieldEncodingKey& InKey) const;
 
+		const FQuat GetListenerRotation(const int32 SourceId) const;
+
 		void SetSubmixSendInfo(const int32 SourceId, const FMixerSourceSubmixSend& SubmixSend);
 
 		void SetBusSendInfo(const int32 SourceId, EBusSendType InAudioBusSendType, uint32 AudiobusId, float BusSendLevel);
@@ -365,9 +368,12 @@ namespace Audio
 			Audio::AlignedFloatBuffer ReverbPluginOutputBuffer;
 			Audio::AlignedFloatBuffer* PostEffectBuffers;
 
-			// Data needed for outputting to submixes for each channel configuration.
+			// Data needed for outputting to submixes for the default channel configuration for the output device.
 			FSubmixChannelData DeviceSubmixInfo;
-			
+
+			// Whether this source's output is being sent to a device submix (i.e. a USoundSubmix).
+			bool bIsSourceBeingSentToDeviceSubmix;
+
 			TMap<FSoundfieldEncodingKey, FSubmixSoundfieldData> EncodedSoundfieldDownmixes;
 			TArray<Audio::FChannelPositionInfo> InputChannelPositions;
 
@@ -383,6 +389,7 @@ namespace Audio
 				: SourceRotation(FQuat::Identity)
 				, PostEffectBuffers(nullptr)
 				, DeviceSubmixInfo(FSubmixChannelData(SourceNumChannels, NumDeviceOutputChannels, InNumFrames))
+				, bIsSourceBeingSentToDeviceSubmix(false)
 				, NumInputChannels(SourceNumChannels)
 				, NumFrames(InNumFrames)
 				, NumDeviceChannels(NumDeviceOutputChannels)
@@ -398,6 +405,7 @@ namespace Audio
 
 			void ResetData(const uint32 InNumInputChannels, int32 InNumDeviceChannels)
 			{
+				bIsSourceBeingSentToDeviceSubmix = false;
 				NumDeviceChannels = InNumDeviceChannels;
 				NumInputChannels = InNumInputChannels;
 				PostEffectBuffers = nullptr;
@@ -406,6 +414,8 @@ namespace Audio
 				EncodedSoundfieldDownmixes.Reset();
 				AmbisonicsDecoder.Reset();
 				bIsInitialDownmix = true;
+
+				PositionalData.Rotation = FQuat::Identity;
 			}
 		};
 
diff --git a/Engine/Source/Runtime/AudioMixer/Private/AudioMixerSourceVoice.cpp b/Engine/Source/Runtime/AudioMixer/Private/AudioMixerSourceVoice.cpp
--- a/Engine/Source/Runtime/AudioMixer/Private/AudioMixerSourceVoice.cpp
+++ b/Engine/Source/Runtime/AudioMixer/Private/AudioMixerSourceVoice.cpp
@@ -267,6 +267,11 @@ namespace Audio
 		return SourceManager->GetEncodedOutput(SourceId, InKey);
 	}
 
+	const FQuat FMixerSourceVoice::GetListenerRotationForVoice() const
+	{
+		return SourceManager->GetListenerRotation(SourceId);
+	}
+
 	void FMixerSourceVoice::SetSubmixSendInfo(FMixerSubmixWeakPtr Submix, const float SendLevel)
 	{
 		AUDIO_MIXER_CHECK_GAME_THREAD(MixerDevice);
diff --git a/Engine/Source/Runtime/AudioMixer/Private/AudioMixerSourceVoice.h b/Engine/Source/Runtime/AudioMixer/Private/AudioMixerSourceVoice.h
--- a/Engine/Source/Runtime/AudioMixer/Private/AudioMixerSourceVoice.h
+++ b/Engine/Source/Runtime/AudioMixer/Private/AudioMixerSourceVoice.h
@@ -2,7 +2,6 @@
 
 #pragma once
 
-#include "AudioMixerDevice.h"
 #include "AudioMixerBuffer.h"
 #include "AudioMixerSourceManager.h"
 
@@ -101,6 +100,9 @@ namespace Audio
 		// For soundfield conversions, get the encoded audio.
 		const ISoundfieldAudioPacket* GetEncodedOutput(const FSoundfieldEncodingKey& InKey) const;
 
+		// This will return the listener rotation used for this source voice.
+		const FQuat GetListenerRotationForVoice() const;
+
 		// Sets the submix send levels
 		void SetSubmixSendInfo(FMixerSubmixWeakPtr Submix, const float SendLevel);
 
diff --git a/Engine/Source/Runtime/AudioMixer/Private/AudioMixerSubmix.cpp b/Engine/Source/Runtime/AudioMixer/Private/AudioMixerSubmix.cpp
--- a/Engine/Source/Runtime/AudioMixer/Private/AudioMixerSubmix.cpp
+++ b/Engine/Source/Runtime/AudioMixer/Private/AudioMixerSubmix.cpp
@@ -518,6 +518,11 @@ namespace Audio
 		SoundfieldStreams.Mixer->MixTogether(InputData, PacketToSumTo);
 	}
 
+	void FMixerSubmix::UpdateListenerRotation(const FQuat& InRotation)
+	{
+		SoundfieldStreams.CachedPositionalData.Rotation = InRotation;
+	}
+
 	void FMixerSubmix::MixInChildSubmix(FChildSubmixInfo& Child, ISoundfieldAudioPacket& PacketToSumTo)
 	{
 		check(IsSoundfieldSubmix());
@@ -581,6 +586,10 @@ namespace Audio
 				// No conversion necessary.
 				ChildSubmixSharedPtr->ProcessAudio(PacketToSumTo);
 			}
+
+			//Propogate listener rotation down to this submix.
+			// This is required if this submix doesn't have any sources sending to it, but does have at least one child submix.
+			UpdateListenerRotation(ChildSubmixSharedPtr->SoundfieldStreams.CachedPositionalData.Rotation);
 		}
 	}
 
@@ -756,8 +765,8 @@ namespace Audio
 		const int32 NumOutputFrames = OutAudioBuffer.Num() / NumChannels;
 		NumSamples = NumChannels * NumOutputFrames;
 
-		InputBuffer.Reset(NumSamples);
-		InputBuffer.AddZeroed(NumSamples);
+ 		InputBuffer.Reset(NumSamples);
+ 		InputBuffer.AddZeroed(NumSamples);
 
 		float* BufferPtr = InputBuffer.GetData();
 
@@ -1055,6 +1064,8 @@ namespace Audio
 
 				// if this voice has a valid encoded packet, mix it in.
 				const ISoundfieldAudioPacket* Packet = MixerSourceVoice->GetEncodedOutput(GetKeyForSubmixEncoding());
+				UpdateListenerRotation(MixerSourceVoice->GetListenerRotationForVoice());
+
 				if (Packet)
 				{
 					FSoundfieldMixerInputData InputData =
diff --git a/Engine/Source/Runtime/AudioMixer/Public/AudioMixerSubmix.h b/Engine/Source/Runtime/AudioMixer/Public/AudioMixerSubmix.h
--- a/Engine/Source/Runtime/AudioMixer/Public/AudioMixerSubmix.h
+++ b/Engine/Source/Runtime/AudioMixer/Public/AudioMixerSubmix.h
@@ -250,6 +250,8 @@ namespace Audio
 		// Encode a source and sum it into the mixed soundfield.
 		void MixInSource(const ISoundfieldAudioPacket& InAudio, const ISoundfieldEncodingSettingsProxy& InSettings, ISoundfieldAudioPacket& PacketToSumTo);
 
+		void UpdateListenerRotation(const FQuat& InRotation);
+
 		// Calls ProcessAudio on the child submix, performs all neccessary conversions and mixes in it's resulting audio.
 		void MixInChildSubmix(FChildSubmixInfo& Child, ISoundfieldAudioPacket& PacketToSumTo);
 
diff --git a/Engine/Source/Runtime/SoundFieldRendering/Private/SoundFieldRendering.cpp b/Engine/Source/Runtime/SoundFieldRendering/Private/SoundFieldRendering.cpp
--- a/Engine/Source/Runtime/SoundFieldRendering/Private/SoundFieldRendering.cpp
+++ b/Engine/Source/Runtime/SoundFieldRendering/Private/SoundFieldRendering.cpp
@@ -83,6 +83,11 @@ FSoundFieldDecoder::FSoundFieldDecoder()
 
 void FSoundFieldDecoder::DecodeAudioDirectlyToDeviceOutputPositions(const FAmbisonicsSoundfieldBuffer& InputData, const FSoundfieldSpeakerPositionalData& OutputPositions, Audio::AlignedFloatBuffer& OutputData)
 {
+	if (InputData.NumChannels == 0 || InputData.AudioBuffer.Num() == 0)
+	{
+		FMemory::Memzero(OutputData.GetData(), OutputData.Num() * sizeof(float));
+		return;
+	}
 
 	constexpr const float DEG_2_RAD = PI / 180.0f;
 
@@ -278,6 +283,7 @@ void FSoundFieldDecoder::DecodeAudioToSevenOneAndDownmixToDevice(const FAmbisoni
 {
 	if (InputData.NumChannels == 0 || InputData.AudioBuffer.Num() == 0)
 	{
+		FMemory::Memzero(OutputData.GetData(), OutputData.Num() * sizeof(float));
 		return;
 	}
 
