Based on slack thread: Move bypass default value to BaseEngine.ini and ConsoleVariables.ini, in case race condition on startup on desktop platforms persists.

diff --git a/Engine/Config/ConsoleVariables.ini b/Engine/Config/ConsoleVariables.ini
--- a/Engine/Config/ConsoleVariables.ini
+++ b/Engine/Config/ConsoleVariables.ini
@@ -85,3 +85,6 @@ p.chaos.AllowCreatePhysxBodies=1
 
 ;Bypass compiler optimizations to speed up Niagara compilation
 fx.SkipVectorVMBackendOptimizations=1
+
+; Comment this out to turn off audio submix effects.
+au.BypassAllSubmixEffects=0
\ No newline at end of file
diff --git a/Engine/Source/Runtime/AudioMixer/Private/AudioMixerSubmix.cpp b/Engine/Source/Runtime/AudioMixer/Private/AudioMixerSubmix.cpp
--- a/Engine/Source/Runtime/AudioMixer/Private/AudioMixerSubmix.cpp
+++ b/Engine/Source/Runtime/AudioMixer/Private/AudioMixerSubmix.cpp
@@ -19,7 +19,7 @@ FAutoConsoleVariableRef CVarRecoverRecordingOnShutdown(
 	TEXT("0: Disabled, 1: Enabled"),
 	ECVF_Default);
 
-static int32 BypassAllSubmixEffectsCVar = 0;
+static int32 BypassAllSubmixEffectsCVar = 1;
 FAutoConsoleVariableRef CVarBypassAllSubmixEffects(
 	TEXT("au.BypassAllSubmixEffects"),
 	BypassAllSubmixEffectsCVar,
