Moving libsndfile from runtime to editor time dependency.

diff --git a/Engine/Source/Editor/AudioEditor/AudioEditor.Build.cs b/Engine/Source/Editor/AudioEditor/AudioEditor.Build.cs
--- a/Engine/Source/Editor/AudioEditor/AudioEditor.Build.cs
+++ b/Engine/Source/Editor/AudioEditor/AudioEditor.Build.cs
@@ -51,6 +51,18 @@ public class AudioEditor : ModuleRules
 
 		if (Target.IsInPlatformGroup(UnrealPlatformGroup.Windows))
 		{
+			string PlatformName = Target.Platform == UnrealTargetPlatform.Win32 ? "Win32" : "Win64";
+
+			string LibSndFilePath = Target.UEThirdPartyBinariesDirectory + "libsndfile/";
+			LibSndFilePath += PlatformName;
+
+
+			PublicAdditionalLibraries.Add(LibSndFilePath + "/libsndfile-1.lib");
+			PublicDelayLoadDLLs.Add("libsndfile-1.dll");
+			PublicIncludePathModuleNames.Add("UELibSampleRate");
+
+			RuntimeDependencies.Add("$(EngineDir)/Binaries/ThirdParty/libsndfile/" + PlatformName + "/libsndfile-1.dll");
+
 			PublicDefinitions.Add("WITH_SNDFILE_IO=1");
 		}
 		else
diff --git a/Engine/Source/Runtime/AudioMixer/AudioMixer.Build.cs b/Engine/Source/Runtime/AudioMixer/AudioMixer.Build.cs
--- a/Engine/Source/Runtime/AudioMixer/AudioMixer.Build.cs
+++ b/Engine/Source/Runtime/AudioMixer/AudioMixer.Build.cs
@@ -48,20 +48,7 @@ namespace UnrealBuildTool.Rules
 					"UELibSampleRate"
 					);
 
-			if (Target.Platform == UnrealTargetPlatform.Win64 || Target.Platform == UnrealTargetPlatform.Win32)
-			{
-				string PlatformName = Target.Platform == UnrealTargetPlatform.Win32 ? "Win32" : "Win64";
 
-                string LibSndFilePath = Target.UEThirdPartyBinariesDirectory + "libsndfile/";
-                LibSndFilePath += PlatformName;
-
-
-                PublicAdditionalLibraries.Add(LibSndFilePath + "/libsndfile-1.lib");
-				PublicDelayLoadDLLs.Add("libsndfile-1.dll");
-				PublicIncludePathModuleNames.Add("UELibSampleRate");
-
-                RuntimeDependencies.Add("$(EngineDir)/Binaries/ThirdParty/libsndfile/" + PlatformName + "/libsndfile-1.dll");
-            }
 		}
 	}
 }
