Unshelved from pending changelist '11129597':

PR #6617: Changed static Up- and Downmix Scratchbuffer to member variable in FDecodingSoundSource (Contributed by [REDACTED_PERSON])

diff --git a/Engine/Source/Runtime/AudioMixer/Private/SoundWaveDecoder.cpp b/Engine/Source/Runtime/AudioMixer/Private/SoundWaveDecoder.cpp
--- a/Engine/Source/Runtime/AudioMixer/Private/SoundWaveDecoder.cpp
+++ b/Engine/Source/Runtime/AudioMixer/Private/SoundWaveDecoder.cpp
@@ -365,7 +365,7 @@ namespace Audio
 		}
 		else
 		{
-			static AlignedFloatBuffer ScratchBuffer;
+
 			ScratchBuffer.Reset();
 			ScratchBuffer.AddZeroed(InNumFrames * SourceInfo.NumSourceChannels);
 
diff --git a/Engine/Source/Runtime/AudioMixer/Public/SoundWaveDecoder.h b/Engine/Source/Runtime/AudioMixer/Public/SoundWaveDecoder.h
--- a/Engine/Source/Runtime/AudioMixer/Public/SoundWaveDecoder.h
+++ b/Engine/Source/Runtime/AudioMixer/Public/SoundWaveDecoder.h
@@ -102,6 +102,9 @@ namespace Audio
 		// Object which handles bulk of decoding operations
 		TSharedPtr<FMixerSourceBuffer, ESPMode::ThreadSafe> MixerSourceBuffer;
 
+		// Scratch buffer used for upmixing and downmixing the audio
+		AlignedFloatBuffer ScratchBuffer;
+
 		FCriticalSection DtorCritSec;
 
 		// Sample rate of the source
