minor: sanity check num channels from the audio device on linux. Callstack for crash indicates an OOB access on envelope array.

diff --git a/Engine/Source/Runtime/AudioMixer/Private/AudioMixerSubmix.cpp b/Engine/Source/Runtime/AudioMixer/Private/AudioMixerSubmix.cpp
--- a/Engine/Source/Runtime/AudioMixer/Private/AudioMixerSubmix.cpp
+++ b/Engine/Source/Runtime/AudioMixer/Private/AudioMixerSubmix.cpp
@@ -50,6 +50,8 @@ namespace Audio
 		, bIsBackgroundMuted(false)
 		, OwningSubmixObject(nullptr)
 	{
+		EnvelopeFollowers.Reset();
+		EnvelopeFollowers.AddDefaulted(AUDIO_MIXER_MAX_OUTPUT_CHANNELS);
 	}
 
 	FMixerSubmix::~FMixerSubmix()
@@ -791,6 +793,13 @@ namespace Audio
 
 		// Device format may change channels if device is hot swapped
 		NumChannels = MixerDevice->GetNumDeviceChannels();
+
+		// If we hit this, it means that platform info gave us a garbage NumChannel count.
+		if (!ensure(NumChannels <= AUDIO_MIXER_MAX_OUTPUT_CHANNELS))
+		{
+			return;
+		}
+
 		const int32 NumOutputFrames = OutAudioBuffer.Num() / NumChannels;
 		NumSamples = NumChannels * NumOutputFrames;
 
diff --git a/Engine/Source/Runtime/AudioMixer/Public/AudioMixerSubmix.h b/Engine/Source/Runtime/AudioMixer/Public/AudioMixerSubmix.h
--- a/Engine/Source/Runtime/AudioMixer/Public/AudioMixerSubmix.h
+++ b/Engine/Source/Runtime/AudioMixer/Public/AudioMixerSubmix.h
@@ -441,7 +441,7 @@ namespace Audio
 
 		// Envelope following data
 		float EnvelopeValues[AUDIO_MIXER_MAX_OUTPUT_CHANNELS];
-		Audio::FEnvelopeFollower EnvelopeFollowers[AUDIO_MIXER_MAX_OUTPUT_CHANNELS];
+		TArray<Audio::FEnvelopeFollower> EnvelopeFollowers;
 		int32 EnvelopeNumChannels;
 		FCriticalSection EnvelopeCriticalSection;
 
diff --git a/Engine/Source/Runtime/Linux/AudioMixerSDL/Private/AudioMixerPlatformSDL.cpp b/Engine/Source/Runtime/Linux/AudioMixerSDL/Private/AudioMixerPlatformSDL.cpp
--- a/Engine/Source/Runtime/Linux/AudioMixerSDL/Private/AudioMixerPlatformSDL.cpp
+++ b/Engine/Source/Runtime/Linux/AudioMixerSDL/Private/AudioMixerPlatformSDL.cpp
@@ -146,7 +146,9 @@ namespace Audio
 		OutInfo.SampleRate = ActualSpec.freq;
 		
 		OutInfo.Format = GetAudioStreamFormat();
-		OutInfo.NumChannels = ActualSpec.channels;
+
+		ensure(ActualSpec.channels <= AUDIO_MIXER_MAX_OUTPUT_CHANNELS);
+		OutInfo.NumChannels = FMath::Min(ActualSpec.channels, AUDIO_MIXER_MAX_OUTPUT_CHANNELS);
 
 		// Assume default channel map order, SDL doesn't support us querying it directly
 		OutInfo.OutputChannelArray.Reset();
