Fix for check getting hit for files marked as ambisonics that don't have four channels. Still log a warning but play the source as a non-ambisonics source.

diff --git a/Engine/Source/Runtime/AudioMixer/Private/AudioMixerSource.cpp b/Engine/Source/Runtime/AudioMixer/Private/AudioMixerSource.cpp
--- a/Engine/Source/Runtime/AudioMixer/Private/AudioMixerSource.cpp
+++ b/Engine/Source/Runtime/AudioMixer/Private/AudioMixerSource.cpp
@@ -118,12 +118,13 @@ namespace Audio
 			InitParams.SourceVoice = MixerSourceVoice;
 			InitParams.bUseHRTFSpatialization = UseObjectBasedSpatialization();
 			InitParams.bIsExternalSend = MixerDevice->bSpatializationIsExternalSend;
-			InitParams.bIsAmbisonics = WaveInstance->bIsAmbisonics;
+			InitParams.bIsAmbisonics = WaveInstance->bIsAmbisonics && (WaveInstance->WaveData->NumChannels == 4);
 
-			if (InitParams.bIsAmbisonics)
+			if (WaveInstance->bIsAmbisonics && (WaveInstance->WaveData->NumChannels != 4))
 			{
-				checkf(InitParams.NumInputChannels == 4, TEXT("Only allow 4 channel source if file is ambisonics format."));
+				UE_LOG(LogAudioMixer, Warning, TEXT("Sound wave %s was flagged as being ambisonics but had a channel count of %d. Currently the audio engine only supports FOA sources that have four channels."), *InWaveInstance->GetName(), WaveInstance->WaveData->NumChannels);
 			}
+
 			InitParams.AudioComponentUserID = WaveInstance->ActiveSound->GetAudioComponentUserID();
 
 			InitParams.AudioComponentID = WaveInstance->ActiveSound->GetAudioComponentID();
