Fix merge issue caused by new BP API added on main
#rnx

diff --git a/Engine/Source/Runtime/AudioMixer/Private/AudioMixerBlueprintLibrary.cpp b/Engine/Source/Runtime/AudioMixer/Private/AudioMixerBlueprintLibrary.cpp
--- a/Engine/Source/Runtime/AudioMixer/Private/AudioMixerBlueprintLibrary.cpp
+++ b/Engine/Source/Runtime/AudioMixer/Private/AudioMixerBlueprintLibrary.cpp
@@ -95,17 +95,17 @@ void UAudioMixerBlueprintLibrary::ClearMasterSubmixEffects(const UObject* WorldC
 
 int32 UAudioMixerBlueprintLibrary::AddSubmixEffect(const UObject* WorldContextObject, USoundSubmix* InSoundSubmix, USoundEffectSubmixPreset* SubmixEffectPreset)
 {
-	if (Audio::FMixerDevice* MixerDevice = GetAudioMixerDeviceFromWorldContext(WorldContextObject))
+	if (!SubmixEffectPreset || !InSoundSubmix)
 	{
-		// Immediately create a new sound effect base here before the object becomes potentially invalidated
-		FSoundEffectSubmixPtr SoundEffectSubmix = MakeShareable(static_cast<FSoundEffectSubmix*>(SubmixEffectPreset->CreateNewEffect()));
+		return 0;
+	}
 
+	if (Audio::FMixerDevice* MixerDevice = GetAudioMixerDeviceFromWorldContext(WorldContextObject))
+	{
 		FSoundEffectSubmixInitData InitData;
 		InitData.SampleRate = MixerDevice->GetSampleRate();
 
-		// Initialize and set the preset immediately
-		SoundEffectSubmix->Init(InitData);
-		SoundEffectSubmix->SetPreset(SubmixEffectPreset);
+		TSoundEffectSubmixPtr SoundEffectSubmix = USoundEffectPreset::CreateInstance<FSoundEffectSubmixInitData, FSoundEffectSubmix>(InitData, *SubmixEffectPreset);
 		SoundEffectSubmix->SetEnabled(true);
 
 		// Get a unique id for the preset object on the game thread. Used to refer to the object on audio render thread.
@@ -113,6 +113,7 @@ int32 UAudioMixerBlueprintLibrary::AddSubmixEffect(const UObject* WorldContextOb
 
 		return MixerDevice->AddSubmixEffect(InSoundSubmix, SubmixPresetUniqueId, SoundEffectSubmix);
 	}
+
 	return 0;
 }
 
@@ -135,17 +136,17 @@ void UAudioMixerBlueprintLibrary::RemoveSubmixEffectPresetAtIndex(const UObject*
 
 void UAudioMixerBlueprintLibrary::ReplaceSoundEffectSubmix(const UObject* WorldContextObject, USoundSubmix* InSoundSubmix, int32 SubmixChainIndex, USoundEffectSubmixPreset* SubmixEffectPreset)
 {
-	if (Audio::FMixerDevice* MixerDevice = GetAudioMixerDeviceFromWorldContext(WorldContextObject))
+	if (!SubmixEffectPreset || !InSoundSubmix)
 	{
-		// Immediately create a new sound effect base here before the object becomes potentially invalidated
-		FSoundEffectSubmixPtr SoundEffectSubmix = MakeShareable(static_cast<FSoundEffectSubmix*>(SubmixEffectPreset->CreateNewEffect()));
+		return 0;
+	}
 
+	if (Audio::FMixerDevice* MixerDevice = GetAudioMixerDeviceFromWorldContext(WorldContextObject))
+	{
 		FSoundEffectSubmixInitData InitData;
 		InitData.SampleRate = MixerDevice->GetSampleRate();
 
-		// Initialize and set the preset immediately
-		SoundEffectSubmix->Init(InitData);
-		SoundEffectSubmix->SetPreset(SubmixEffectPreset);
+		TSoundEffectSubmixPtr SoundEffectSubmix = USoundEffectPreset::CreateInstance<FSoundEffectSubmixInitData, FSoundEffectSubmix>(InitData, *SubmixEffectPreset);
 		SoundEffectSubmix->SetEnabled(true);
 
 		// Get a unique id for the preset object on the game thread. Used to refer to the object on audio render thread.
