Return value copy pasta from last merge conflict resolve
#rnx

diff --git a/Engine/Source/Runtime/AudioMixer/Private/AudioMixerBlueprintLibrary.cpp b/Engine/Source/Runtime/AudioMixer/Private/AudioMixerBlueprintLibrary.cpp
--- a/Engine/Source/Runtime/AudioMixer/Private/AudioMixerBlueprintLibrary.cpp
+++ b/Engine/Source/Runtime/AudioMixer/Private/AudioMixerBlueprintLibrary.cpp
@@ -138,7 +138,7 @@ void UAudioMixerBlueprintLibrary::ReplaceSoundEffectSubmix(const UObject* WorldC
 {
 	if (!SubmixEffectPreset || !InSoundSubmix)
 	{
-		return 0;
+		return;
 	}
 
 	if (Audio::FMixerDevice* MixerDevice = GetAudioMixerDeviceFromWorldContext(WorldContextObject))
