Merge over effect chain lock from UE Main.

diff --git a/Engine/Source/Runtime/AudioMixer/Private/AudioMixerSubmix.cpp b/Engine/Source/Runtime/AudioMixer/Private/AudioMixerSubmix.cpp
--- a/Engine/Source/Runtime/AudioMixer/Private/AudioMixerSubmix.cpp
+++ b/Engine/Source/Runtime/AudioMixer/Private/AudioMixerSubmix.cpp
@@ -116,6 +116,7 @@ namespace Audio
 
 			NumSubmixEffects = 0;
 
+			FScopeLock ScopeLock(&EffectChainMutationCriticalSection);			
 			for (USoundEffectSubmixPreset* EffectPreset : SoundSubmix->SubmixEffectChain)
 			{
 				if (EffectPreset)
@@ -348,6 +349,7 @@ namespace Audio
 
 	void FMixerSubmix::AddSoundEffectSubmix(uint32 SubmixPresetId, FSoundEffectSubmixPtr InSoundEffectSubmix)
 	{
+		FScopeLock ScopeLock(&EffectChainMutationCriticalSection);
 		AUDIO_MIXER_CHECK_AUDIO_PLAT_THREAD(MixerDevice);
 
 		// Look to see if the submix preset ID is already present
@@ -372,6 +374,7 @@ namespace Audio
 
 	void FMixerSubmix::RemoveSoundEffectSubmix(uint32 SubmixPresetId)
 	{
+		FScopeLock ScopeLock(&EffectChainMutationCriticalSection);
 		AUDIO_MIXER_CHECK_AUDIO_PLAT_THREAD(MixerDevice);
 
 		for (FSubmixEffectInfo& Effect : EffectSubmixChain)
@@ -406,6 +409,8 @@ namespace Audio
 
 	void FMixerSubmix::ClearSoundEffectSubmixes()
 	{
+		FScopeLock ScopeLock(&EffectChainMutationCriticalSection);
+
 		for (FSubmixEffectInfo& Info : EffectSubmixChain)
 		{
 			if (Info.EffectInstance.IsValid())
@@ -847,6 +852,7 @@ namespace Audio
 			}
 		}
 
+		FScopeLock ScopeLock(&EffectChainMutationCriticalSection);
 		if (EffectSubmixChain.Num() > 0)
 		{
 			CSV_SCOPED_TIMING_STAT(Audio, SubmixEffectProcessing);
diff --git a/Engine/Source/Runtime/AudioMixer/Public/AudioMixerSubmix.h b/Engine/Source/Runtime/AudioMixer/Public/AudioMixerSubmix.h
--- a/Engine/Source/Runtime/AudioMixer/Public/AudioMixerSubmix.h
+++ b/Engine/Source/Runtime/AudioMixer/Public/AudioMixerSubmix.h
@@ -484,6 +484,9 @@ namespace Audio
 		// Critical section used for when we are appending recorded data.
 		FCriticalSection RecordingCriticalSection;
 
+		// Critical section for mutation of the effect chain.
+		FCriticalSection EffectChainMutationCriticalSection;
+
 		// Handle back to the owning USoundSubmix. Used when the device is shutdown to prematurely end a recording.
 		const USoundSubmixBase* OwningSubmixObject;
 
