Integrating changes from 4.25 for signal processing and audio synesthesia

Integrating CLs: 11645006, 12210223, 12211452

diff --git a/Engine/Source/Runtime/AudioMixer/Private/Effects/AudioMixerSubmixEffectReverbFast.cpp b/Engine/Source/Runtime/AudioMixer/Private/Effects/AudioMixerSubmixEffectReverbFast.cpp
--- a/Engine/Source/Runtime/AudioMixer/Private/Effects/AudioMixerSubmixEffectReverbFast.cpp
+++ b/Engine/Source/Runtime/AudioMixer/Private/Effects/AudioMixerSubmixEffectReverbFast.cpp
@@ -180,13 +180,13 @@ bool FSubmixEffectReverbFast::SetParameters(const FAudioEffectParameters& InPara
 	Audio::FPlateReverbFastSettings NewSettings;
 
 	// Early Reflections
-	NewSettings.EarlyReflections.Gain = FMath::GetMappedRangeValueClamped({ 0.0f, 3.16f }, { 0.0f, 1.0f }, ReverbEffectParams.ReflectionsGain * ReverbEffectParams.Volume);
+	NewSettings.EarlyReflections.Gain = FMath::GetMappedRangeValueClamped({ 0.0f, 3.16f }, { 0.0f, 1.0f }, ReverbEffectParams.ReflectionsGain);
 	NewSettings.EarlyReflections.PreDelayMsec = FMath::GetMappedRangeValueClamped({ 0.0f, 0.3f }, { 0.0f, 300.0f }, ReverbEffectParams.ReflectionsDelay);
 	NewSettings.EarlyReflections.Bandwidth = FMath::GetMappedRangeValueClamped({ 0.0f, 1.0f }, { 0.0f, 1.0f }, ReverbEffectParams.GainHF);
 
 	// LateReflections
 	NewSettings.LateReflections.LateDelayMsec = FMath::GetMappedRangeValueClamped({ 0.0f, 0.1f }, { 0.0f, 100.0f }, ReverbEffectParams.LateDelay);
-	NewSettings.LateReflections.LateGainDB = FMath::GetMappedRangeValueClamped({ 0.0f, 1.0f }, { 0.0f, 1.0f }, ReverbEffectParams.Gain * ReverbEffectParams.Volume);
+	NewSettings.LateReflections.LateGainDB = FMath::GetMappedRangeValueClamped({ 0.0f, 1.0f }, { 0.0f, 1.0f }, ReverbEffectParams.Gain);
 	NewSettings.LateReflections.Bandwidth = FMath::GetMappedRangeValueClamped({ 0.0f, 1.0f }, { 0.1f, 0.6f }, ReverbEffectParams.AirAbsorptionGainHF);
 	NewSettings.LateReflections.Diffusion = FMath::GetMappedRangeValueClamped({ 0.05f, 1.0f }, { 0.0f, 0.95f }, ReverbEffectParams.Diffusion);
 	NewSettings.LateReflections.Dampening = FMath::GetMappedRangeValueClamped({ 0.05f, 1.95f }, { 0.0f, 0.999f }, ReverbEffectParams.DecayHFRatio);
