edigrate of procudural sound GC fixes to Release 4.25

diff --git a/Engine/Source/Runtime/AudioMixer/Private/AudioMixerSource.h b/Engine/Source/Runtime/AudioMixer/Private/AudioMixerSource.h
--- a/Engine/Source/Runtime/AudioMixer/Private/AudioMixerSource.h
+++ b/Engine/Source/Runtime/AudioMixer/Private/AudioMixerSource.h
@@ -61,6 +61,9 @@ namespace Audio
 
 	private:
 
+		/** Initializes the bus sends. */
+		void InitBusSends(FWaveInstance* InWaveInstance, FMixerSourceVoiceInitParams& InitParams);
+
 		/** Frees any resources for this sound source. */
 		void FreeResources();
 
@@ -116,7 +119,7 @@ namespace Audio
 
 		FMixerDevice* MixerDevice;
 		FMixerBuffer* MixerBuffer;
-		TSharedPtr<FMixerSourceBuffer> MixerSourceBuffer;
+		TSharedPtr<FMixerSourceBuffer, ESPMode::ThreadSafe> MixerSourceBuffer;
 		FMixerSourceVoice* MixerSourceVoice;
 
 		// This holds data copied from FSoundSourceBusSendInfo when a new sound starts playing
diff --git a/Engine/Source/Runtime/AudioMixer/Private/AudioMixerSourceBuffer.cpp b/Engine/Source/Runtime/AudioMixer/Private/AudioMixerSourceBuffer.cpp
--- a/Engine/Source/Runtime/AudioMixer/Private/AudioMixerSourceBuffer.cpp
+++ b/Engine/Source/Runtime/AudioMixer/Private/AudioMixerSourceBuffer.cpp
@@ -56,7 +56,7 @@ namespace Audio
 		return CurrentSample >= NumSamples;
 	}
 
-	TSharedPtr<FMixerSourceBuffer> FMixerSourceBuffer::Create(FMixerBuffer& InBuffer, USoundWave& InWave, ELoopingMode InLoopingMode, bool bInIsSeeking, bool bInForceSyncDecode)
+	TSharedPtr<FMixerSourceBuffer, ESPMode::ThreadSafe> FMixerSourceBuffer::Create(FMixerBuffer& InBuffer, USoundWave& InWave, ELoopingMode InLoopingMode, bool bInIsSeeking, bool bInForceSyncDecode)
 	{
 		LLM_SCOPE(ELLMTag::AudioMixer);
 
@@ -72,7 +72,7 @@ namespace Audio
 			return nullptr;
 		}
 
-		TSharedPtr<FMixerSourceBuffer> NewSourceBuffer = MakeShareable(new FMixerSourceBuffer(InBuffer, InWave, InLoopingMode, bInIsSeeking, bInForceSyncDecode));
+		TSharedPtr<FMixerSourceBuffer, ESPMode::ThreadSafe> NewSourceBuffer = MakeShareable(new FMixerSourceBuffer(InBuffer, InWave, InLoopingMode, bInIsSeeking, bInForceSyncDecode));
 		return NewSourceBuffer;
 	}
 
@@ -92,7 +92,7 @@ namespace Audio
 		, bIsSeeking(bInIsSeeking)
 		, bLoopCallback(false)
 		, bProcedural(InWave.bProcedural)
-		, bIsBus(InWave.bIsBus)
+		, bIsBus(InWave.bIsSourceBus)
 		, bForceSyncDecode(bInForceSyncDecode)
 	{
 		InWave.AddPlayingSource(this);
@@ -100,7 +100,7 @@ namespace Audio
 		const uint32 TotalSamples = MONO_PCM_BUFFER_SAMPLES * NumChannels;
 		for (int32 BufferIndex = 0; BufferIndex < Audio::MAX_BUFFERS_QUEUED; ++BufferIndex)
 		{
-			SourceVoiceBuffers.Add(MakeShared<FMixerSourceVoiceBuffer>());
+			SourceVoiceBuffers.Add(MakeShared<FMixerSourceVoiceBuffer, ESPMode::ThreadSafe>());
 
 			// Prepare the memory to fit the max number of samples
 			SourceVoiceBuffers[BufferIndex]->AudioData.Reset(TotalSamples);
@@ -113,29 +113,14 @@ namespace Audio
 	{
 		// GC methods may get called from the game thread during the destructor
 		// These methods will trylock and early exit if we have this lock
-		FScopeLock Lock(&DtorCritSec);
-
-		// Make sure we have completed our async realtime task before deleting the decompression state
-		if (AsyncRealtimeAudioTask)
-		{
-			AsyncRealtimeAudioTask->CancelTask();
-			delete AsyncRealtimeAudioTask;
-			AsyncRealtimeAudioTask = nullptr;
-		}
+		FScopeLock Lock(&SoundWaveCritSec);
 
+		// OnEndGenerate calls EnsureTaskFinishes,
+		// which will make sure we have completed our async realtime task before deleting the decompression state
 		OnEndGenerate();
 
 		// Clean up decompression state after things have been finished using it
-		if (DecompressionState)
-		{
-			if (BufferType == EBufferType::Streaming)
-			{
-				IStreamingManager::Get().GetAudioStreamingManager().RemoveDecoder(DecompressionState);
-			}
-
-			delete DecompressionState;
-			DecompressionState = nullptr;
-		}
+		DeleteDecoder();
 
 		if (SoundWave)
 		{
@@ -175,18 +160,18 @@ namespace Audio
 
 		switch (BufferType)
 		{
-			case EBufferType::PCM:
-			case EBufferType::PCMPreview:
-				SubmitInitialPCMBuffers();
-				break;
+		case EBufferType::PCM:
+		case EBufferType::PCMPreview:
+			SubmitInitialPCMBuffers();
+			break;
 
-			case EBufferType::PCMRealTime:
-			case EBufferType::Streaming:
-				SubmitInitialRealtimeBuffers();
-				break;
+		case EBufferType::PCMRealTime:
+		case EBufferType::Streaming:
+			SubmitInitialRealtimeBuffers();
+			break;
 
-			case EBufferType::Invalid:
-				break;
+		case EBufferType::Invalid:
+			break;
 		}
 
 		return true;
@@ -194,7 +179,9 @@ namespace Audio
 
 	void FMixerSourceBuffer::OnBufferEnd()
 	{
-		if ((NumBuffersQeueued == 0 && bBufferFinished) || (bProcedural && !SoundWave))
+		FScopeTryLock Lock(&SoundWaveCritSec);
+
+		if (!Lock.IsLocked() || (NumBuffersQeueued == 0 && bBufferFinished) || (bProcedural && !SoundWave))
 		{
 			return;
 		}
@@ -202,9 +189,26 @@ namespace Audio
 		ProcessRealTimeSource();
 	}
 
-	TSharedPtr<FMixerSourceVoiceBuffer> FMixerSourceBuffer::GetNextBuffer()
+	int32 FMixerSourceBuffer::GetNumBuffersQueued() const
 	{
-		TSharedPtr<FMixerSourceVoiceBuffer> NewBufferPtr;
+		FScopeTryLock Lock(&SoundWaveCritSec);
+		if (Lock.IsLocked())
+		{
+			return NumBuffersQeueued;
+		}
+
+		return 0;
+	}
+
+	TSharedPtr<FMixerSourceVoiceBuffer, ESPMode::ThreadSafe> FMixerSourceBuffer::GetNextBuffer()
+	{
+		FScopeTryLock Lock(&SoundWaveCritSec);
+		if(!Lock.IsLocked())
+		{
+			return nullptr;
+		}
+
+		TSharedPtr<FMixerSourceVoiceBuffer, ESPMode::ThreadSafe> NewBufferPtr;
 		BufferQueue.Dequeue(NewBufferPtr);
 		--NumBuffersQeueued;
 		return NewBufferPtr;
@@ -308,7 +312,7 @@ namespace Audio
 		{
 			ProcessRealTimeSource();
 		}
-	}
+			}
 
 	bool FMixerSourceBuffer::ReadMoreRealtimeData(ICompressedAudioInfo* InDecoder, const int32 BufferIndex, EBufferReadMode BufferReadMode)
 	{
@@ -318,17 +322,24 @@ namespace Audio
 
 		if (bProcedural)
 		{
-			check(SoundWave && SoundWave->bProcedural);
-			FProceduralAudioTaskData NewTaskData;
-			NewTaskData.ProceduralSoundWave = SoundWave;
-			NewTaskData.AudioData = SourceVoiceBuffers[BufferIndex]->AudioData.GetData();
-			NewTaskData.NumSamples = MaxSamples;
-			NewTaskData.NumChannels = NumChannels;
-			check(!AsyncRealtimeAudioTask);
-			AsyncRealtimeAudioTask = CreateAudioTask(NewTaskData);
-
-			// Procedural sound waves never loop
-			return false;
+			FScopeTryLock Lock(&SoundWaveCritSec);
+			
+			if (Lock.IsLocked() && SoundWave)
+			{
+				check(SoundWave && SoundWave->bProcedural);
+				FProceduralAudioTaskData NewTaskData;
+				NewTaskData.ProceduralSoundWave = SoundWave;
+				NewTaskData.AudioData = SourceVoiceBuffers[BufferIndex]->AudioData.GetData();
+				NewTaskData.NumSamples = MaxSamples;
+				NewTaskData.NumChannels = NumChannels;
+				check(!AsyncRealtimeAudioTask);
+				AsyncRealtimeAudioTask = CreateAudioTask(NewTaskData);
+
+				// Procedural sound waves never loop
+				return false;
+			}
+
+			return true;
 		}
 		else if (BufferType != EBufferType::PCMRealTime && BufferType != EBufferType::Streaming)
 		{
@@ -351,6 +362,7 @@ namespace Audio
 		NewTaskData.NumPrecacheFrames = NumPrecacheFrames;
 		NewTaskData.bForceSyncDecode = bForceSyncDecode;
 
+		FScopeLock Lock(&DecodeTaskCritSec);
 		check(!AsyncRealtimeAudioTask);
 		AsyncRealtimeAudioTask = CreateAudioTask(NewTaskData);
 
@@ -389,6 +401,7 @@ namespace Audio
 
 	void FMixerSourceBuffer::ProcessRealTimeSource()
 	{
+		FScopeLock Lock(&DecodeTaskCritSec);
 		if (AsyncRealtimeAudioTask)
 		{
 			AsyncRealtimeAudioTask->EnsureCompletion();
@@ -451,23 +464,40 @@ namespace Audio
 		}
 	}
 
-	void FMixerSourceBuffer::SubmitBuffer(TSharedPtr<FMixerSourceVoiceBuffer> InSourceVoiceBuffer)
+	void FMixerSourceBuffer::SubmitBuffer(TSharedPtr<FMixerSourceVoiceBuffer, ESPMode::ThreadSafe> InSourceVoiceBuffer)
 	{
 		NumBuffersQeueued++;
 		BufferQueue.Enqueue(InSourceVoiceBuffer);
 	}
 
+	void FMixerSourceBuffer::DeleteDecoder()
+	{
+		// Clean up decompression state after things have been finished using it
+		if (DecompressionState)
+		{
+			if (BufferType == EBufferType::Streaming)
+			{
+				IStreamingManager::Get().GetAudioStreamingManager().RemoveDecoder(DecompressionState);
+			}
+
+			delete DecompressionState;
+			DecompressionState = nullptr;
+		}
+	}
+
 	bool FMixerSourceBuffer::OnBeginDestroy(USoundWave * /*Wave*/)
 	{
-		FScopeTryLock Lock(&DtorCritSec);
+		FScopeTryLock Lock(&SoundWaveCritSec);
 
 		// if we don't have the lock, it means we are in ~FMixerSourceBuffer() on another thread
 		if (Lock.IsLocked() && SoundWave)
 		{
-			SoundWave = nullptr;
+			EnsureAsyncTaskFinishes();
+			DeleteDecoder();
+			ClearWave();
 			return true;
 		}
-		
+
 		return false;
 	}
 
@@ -478,22 +508,26 @@ namespace Audio
 
 	void FMixerSourceBuffer::OnFinishDestroy(USoundWave * /*Wave*/)
 	{
-		FScopeTryLock Lock(&DtorCritSec);
+		EnsureAsyncTaskFinishes();
+		FScopeTryLock Lock(&SoundWaveCritSec);
 
 		// if we don't have the lock, it means we are in ~FMixerSourceBuffer() on another thread
 		if (Lock.IsLocked() && SoundWave)
 		{
-			SoundWave = nullptr;
+			DeleteDecoder();
+			ClearWave();
 		}
 	}
 
 	bool FMixerSourceBuffer::IsAsyncTaskInProgress() const
 	{ 
-		return AsyncRealtimeAudioTask != nullptr; 
+		FScopeLock Lock(&DecodeTaskCritSec);
+		return AsyncRealtimeAudioTask != nullptr;
 	}
 
 	bool FMixerSourceBuffer::IsAsyncTaskDone() const 
 	{
+		FScopeLock Lock(&DecodeTaskCritSec);
 		if (AsyncRealtimeAudioTask)
 		{
 			return AsyncRealtimeAudioTask->IsDone();
@@ -503,6 +537,7 @@ namespace Audio
 	
 	void FMixerSourceBuffer::EnsureAsyncTaskFinishes() 
 	{
+		FScopeLock Lock(&DecodeTaskCritSec);
 		if (AsyncRealtimeAudioTask) 
 		{ 
 			AsyncRealtimeAudioTask->CancelTask(); 
@@ -514,7 +549,9 @@ namespace Audio
 
 	void FMixerSourceBuffer::OnBeginGenerate()
 	{
-		if (SoundWave && bProcedural)
+		FScopeTryLock Lock(&SoundWaveCritSec);
+
+		if (Lock.IsLocked() && SoundWave && bProcedural)
 		{
 			check(SoundWave && SoundWave->bProcedural);
 			SoundWave->OnBeginGenerate();
@@ -525,6 +562,13 @@ namespace Audio
 	{		
 		// Make sure the async task finishes!
 		EnsureAsyncTaskFinishes();
+
+		FScopeTryLock Lock(&SoundWaveCritSec);
+		if (!Lock.IsLocked())
+		{
+			return;
+		}
+
 						
 		// Only need to call OnEndGenerate and access SoundWave here if we successfully initialized
 		if (SoundWave && bInitialized && bProcedural)
diff --git a/Engine/Source/Runtime/AudioMixer/Private/AudioMixerSourceBuffer.h b/Engine/Source/Runtime/AudioMixer/Private/AudioMixerSourceBuffer.h
--- a/Engine/Source/Runtime/AudioMixer/Private/AudioMixerSourceBuffer.h
+++ b/Engine/Source/Runtime/AudioMixer/Private/AudioMixerSourceBuffer.h
@@ -43,7 +43,7 @@ namespace Audio
 		AsynchronousSkipFirstFrame
 	};
 
-	using FMixerSourceBufferPtr = TSharedPtr<class FMixerSourceBuffer>;
+	using FMixerSourceBufferPtr = TSharedPtr<class FMixerSourceBuffer, ESPMode::ThreadSafe>;
 
 	/** Class which handles decoding audio for a particular source buffer. */
 	class FMixerSourceBuffer : public ISoundWaveClient
@@ -68,10 +68,10 @@ namespace Audio
 		void OnBufferEnd();
 
 		// Return the number of buffers enqueued on the mixer source buffer
-		int32 GetNumBuffersQueued() const { return NumBuffersQeueued; }
-		
+		int32 GetNumBuffersQueued() const;
+
 		// Returns the next enqueued buffer, returns nullptr if no buffers enqueued
-		TSharedPtr<FMixerSourceVoiceBuffer> GetNextBuffer();
+		TSharedPtr<FMixerSourceVoiceBuffer, ESPMode::ThreadSafe> GetNextBuffer();
 
 		// Returns if buffer looped
 		bool DidBufferLoop() const { return bLoopCallback; }
@@ -102,14 +102,15 @@ namespace Audio
 		void SubmitInitialRealtimeBuffers();
 		void SubmitRealTimeSourceData(const bool bLooped);
 		void ProcessRealTimeSource();
-		void SubmitBuffer(TSharedPtr<FMixerSourceVoiceBuffer> InSourceVoiceBuffer);
+		void SubmitBuffer(TSharedPtr<FMixerSourceVoiceBuffer, ESPMode::ThreadSafe> InSourceVoiceBuffer);
+		void DeleteDecoder();
 
 
 		int32 NumBuffersQeueued;
 		FRawPCMDataBuffer RawPCMDataBuffer;
 
-		TArray<TSharedPtr<FMixerSourceVoiceBuffer>> SourceVoiceBuffers;
-		TQueue<TSharedPtr<FMixerSourceVoiceBuffer>> BufferQueue;
+		TArray<TSharedPtr<FMixerSourceVoiceBuffer, ESPMode::ThreadSafe>> SourceVoiceBuffers;
+		TQueue<TSharedPtr<FMixerSourceVoiceBuffer, ESPMode::ThreadSafe>> BufferQueue;
 		int32 CurrentBuffer;
 		// SoundWaves are only set for procedural sound waves
 		USoundWave* SoundWave;
@@ -121,7 +122,8 @@ namespace Audio
 		int32 NumPrecacheFrames;
 		TArray<uint8> CachedRealtimeFirstBuffer;
 
-		FCriticalSection DtorCritSec;
+		mutable FCriticalSection SoundWaveCritSec;
+		mutable FCriticalSection DecodeTaskCritSec;
 
 		uint32 bInitialized : 1;
 		uint32 bBufferFinished : 1;
diff --git a/Engine/Source/Runtime/AudioMixer/Private/AudioMixerSourceManager.cpp b/Engine/Source/Runtime/AudioMixer/Private/AudioMixerSourceManager.cpp
--- a/Engine/Source/Runtime/AudioMixer/Private/AudioMixerSourceManager.cpp
+++ b/Engine/Source/Runtime/AudioMixer/Private/AudioMixerSourceManager.cpp
@@ -203,11 +203,11 @@ namespace Audio
 			SourceInfo.NumFramesPlayed = 0;
 			SourceInfo.StartTime = 0.0;
 			SourceInfo.SubmixSends.Reset();
-			SourceInfo.BusId = INDEX_NONE;
-			SourceInfo.BusDurationFrames = INDEX_NONE;
+			SourceInfo.AudioBusId = INDEX_NONE;
+			SourceInfo.SourceBusDurationFrames = INDEX_NONE;
 		
-			SourceInfo.BusSends[(int32)EBusSendType::PreEffect].Reset();
-			SourceInfo.BusSends[(int32)EBusSendType::PostEffect].Reset();
+			SourceInfo.AudioBusSends[(int32)EBusSendType::PreEffect].Reset();
+			SourceInfo.AudioBusSends[(int32)EBusSendType::PostEffect].Reset();
 
 			SourceInfo.SourceEffectChainId = INDEX_NONE;
 
@@ -288,6 +288,23 @@ namespace Audio
 			MaxChannelsSupportedBySpatializationPlugin = MixerDevice->MaxChannelsSupportedBySpatializationPlugin;
 		}
 
+		// Spam command queue with nops.
+		static FAutoConsoleCommand SpamNopsCmd(
+			TEXT("au.SpamCommandQueue"),
+			TEXT(""),
+			FConsoleCommandDelegate::CreateLambda([this]() 
+			{				
+				struct FSpamPayload
+				{
+					uint8 JunkBytes[1024];
+				} Payload;
+				for (int32 i = 0; i < 65536; ++i)
+				{
+					AudioMixerThreadCommand([Payload] {});
+				}
+			})
+		);
+
 		bInitialized = true;
 		bPumpQueue = false;
 	}
@@ -374,7 +391,12 @@ namespace Audio
 
 		AUDIO_MIXER_CHECK(SourceId < NumTotalSources);
 		AUDIO_MIXER_CHECK(bInitialized);
-		AUDIO_MIXER_CHECK(MixerSources[SourceId] != nullptr);
+
+		if (MixerSources[SourceId] == nullptr)
+		{
+			UE_LOG(LogAudioMixer, Warning, TEXT("Ignoring double release of SourceId: %i"), SourceId);
+			return;
+		}
 
 		AUDIO_MIXER_DEBUG_LOG(SourceId, TEXT("Is releasing"));
 		
@@ -387,39 +409,44 @@ namespace Audio
 		}
 #endif
 		// Remove from list of active bus or source ids depending on what type of source this is
-		if (SourceInfo.BusId != INDEX_NONE)
+		if (SourceInfo.AudioBusId != INDEX_NONE)
 		{
 			// Remove this bus from the registry of bus instances
-			FMixerBus* Bus = Buses.Find(SourceInfo.BusId);
-			AUDIO_MIXER_CHECK(Bus);
-
-			// Remove this source from the list of bus instances.
-			if (Bus->RemoveInstanceId(SourceId))
+			TSharedPtr<FMixerAudioBus> AudioBusPtr = AudioBuses.FindRef(SourceInfo.AudioBusId);
+			if (ensure(AudioBusPtr.IsValid()))
 			{
-				Buses.Remove(SourceInfo.BusId);
+				// If this audio bus was automatically created via source bus playback, this this audio bus can be removed
+				if (AudioBusPtr->RemoveInstanceId(SourceId))
+				{
+					// Only automatic buses will be getting removed here. Otherwise they need to be manually removed from the source manager.
+					ensure(AudioBusPtr->IsAutomatic());
+					AudioBuses.Remove(SourceInfo.AudioBusId);
+				}
 			}
 		}
 
 		// Remove this source's send list from the bus data registry
-		for (int32 BusSendType = 0; BusSendType < (int32)EBusSendType::Count; ++BusSendType)
+		for (int32 AudioBusSendType = 0; AudioBusSendType < (int32)EBusSendType::Count; ++AudioBusSendType)
 		{
-			for (uint32 BusId : SourceInfo.BusSends[BusSendType])
+			for (uint32 AudioBusId : SourceInfo.AudioBusSends[AudioBusSendType])
 			{
 				// we should have a bus registration entry still since the send hasn't been cleaned up yet
-				FMixerBus* Bus = Buses.Find(BusId);
-				AUDIO_MIXER_CHECK(Bus);
-
-				if (Bus->RemoveBusSend((EBusSendType)BusSendType, SourceId))
+				TSharedPtr<FMixerAudioBus> AudioBusPtr = AudioBuses.FindRef(AudioBusId);
+				if (ensure(AudioBusPtr.IsValid()))
 				{
-					Buses.Remove(BusId);
+					if (AudioBusPtr->RemoveSend((EBusSendType)AudioBusSendType, SourceId))
+					{
+						ensure(AudioBusPtr->IsAutomatic());
+						AudioBuses.Remove(AudioBusId);
+					}
 				}
 			}
 
-			SourceInfo.BusSends[BusSendType].Reset();
+			SourceInfo.AudioBusSends[AudioBusSendType].Reset();
 		}
 
-		SourceInfo.BusId = INDEX_NONE;
-		SourceInfo.BusDurationFrames = INDEX_NONE;
+		SourceInfo.AudioBusId = INDEX_NONE;
+		SourceInfo.SourceBusDurationFrames = INDEX_NONE;
 
 		// Free the mixer source buffer data
 		if (SourceInfo.MixerSourceBuffer.IsValid())
@@ -445,6 +472,7 @@ namespace Audio
 		if (SourceInfo.bUseHRTFSpatializer)
 		{
 			AUDIO_MIXER_CHECK(bUsingSpatializationPlugin);
+			LLM_SCOPE(ELLMTag::AudioMixerPlugins);
 			SpatializationPlugin->OnReleaseSource(SourceId);
 		}
 
@@ -599,9 +627,9 @@ namespace Audio
 		return NumActiveSources;
 	}
 
-	int32 FMixerSourceManager::GetNumActiveBuses() const
+	int32 FMixerSourceManager::GetNumActiveAudioBuses() const
 	{
-		return Buses.Num();
+		return AudioBuses.Num();
 	}
 
 	void FMixerSourceManager::InitSource(const int32 SourceId, const FMixerSourceVoiceInitParams& InitParams)
@@ -670,6 +698,7 @@ namespace Audio
 			if (InitParams.bUseHRTFSpatialization)
 			{
 				AUDIO_MIXER_CHECK(bUsingSpatializationPlugin);
+				LLM_SCOPE(ELLMTag::AudioMixerPlugins);
 				SpatializationPlugin->OnInitSource(SourceId, InitParams.AudioComponentUserID, InitParams.SpatializationPluginSettings);
 			}
 
@@ -701,6 +730,7 @@ namespace Audio
 				InitData.SampleRate = MixerDevice->SampleRate;
 				InitData.NumSourceChannels = InitParams.NumInputChannels;
 				InitData.AudioClock = MixerDevice->GetAudioTime();
+				InitData.AudioDeviceId = MixerDevice->DeviceID;
 
 				SourceInfo.SourceEffectChainId = InitParams.SourceEffectChainId;
 				BuildSourceEffectChain(SourceId, InitData, InitParams.SourceEffectChain);
@@ -708,34 +738,33 @@ namespace Audio
 				// Whether or not to output to bus only
 				SourceInfo.bOutputToBusOnly = InitParams.bOutputToBusOnly;
 
-				// If this is a bus, add this source id to the list of active bus ids
-				if (InitParams.BusId != INDEX_NONE)
+				// If this is going to be a source bus, add this source id to the list of active bus ids
+				if (InitParams.AudioBusId != INDEX_NONE)
 				{
 					// Setting this BusId will flag this source as a bus. It doesn't try to generate 
 					// audio in the normal way but instead will render in a second stage, after normal source rendering.
-					SourceInfo.BusId = InitParams.BusId;
+					SourceInfo.AudioBusId = InitParams.AudioBusId;
 
-					// Bus duration allows us to stop a bus after a given time
-					if (InitParams.BusDuration != 0.0f)
+					// Source bus duration allows us to stop a bus after a given time
+					if (InitParams.SourceBusDuration != 0.0f)
 					{
-						SourceInfo.BusDurationFrames = InitParams.BusDuration * MixerDevice->GetSampleRate();
+						SourceInfo.SourceBusDurationFrames = InitParams.SourceBusDuration * MixerDevice->GetSampleRate();
 					}
 
 					// Register this bus as an instance
-					FMixerBus* Bus = Buses.Find(InitParams.BusId);
-					if (Bus)
+					TSharedPtr<FMixerAudioBus> AudioBusPtr = AudioBuses.FindRef(SourceInfo.AudioBusId);
+					if (AudioBusPtr.IsValid())
 					{
 						// If this bus is already registered, add this as a source id
-						Bus->AddInstanceId(SourceId, InitParams.NumInputChannels);
+						AudioBusPtr->AddInstanceId(SourceId, InitParams.NumInputChannels);
 					}
 					else
 					{
-						// If the bus is not registered, make a new entry
-						FMixerBus NewBusData(this, InitParams.NumInputChannels, NumOutputFrames);
+						// If the bus is not registered, make a new entry. This will default to an automatic audio bus until explicitly made manual later.
+						TSharedPtr<FMixerAudioBus> NewAudioBus = TSharedPtr<FMixerAudioBus>(new FMixerAudioBus(this, true, InitParams.NumInputChannels));
+						NewAudioBus->AddInstanceId(SourceId, InitParams.NumInputChannels);
 
-						NewBusData.AddInstanceId(SourceId, InitParams.NumInputChannels);
-
-						Buses.Add(InitParams.BusId, NewBusData);
+						AudioBuses.Add(InitParams.AudioBusId, NewAudioBus);
 					}
 				}
 
@@ -743,34 +772,34 @@ namespace Audio
 				// Note: buses can also send their audio to other buses.
 				for (int32 BusSendType = 0; BusSendType < (int32)EBusSendType::Count; ++BusSendType)
 				{
-					for (const FMixerBusSend& BusSend : InitParams.BusSends[BusSendType])
+					for (const FInitAudioBusSend& AudioBusSend : InitParams.AudioBusSends[BusSendType])
 					{
 						// New struct to map which source (SourceId) is sending to the bus
-						FBusSend NewBusSend;
-						NewBusSend.SourceId = SourceId;
-						NewBusSend.SendLevel = BusSend.SendLevel;
+						FAudioBusSend NewAudioBusSend;
+						NewAudioBusSend.SourceId = SourceId;
+						NewAudioBusSend.SendLevel = AudioBusSend.SendLevel;
 
 						// Get existing BusId and add the send, or create new bus registration
-						FMixerBus* Bus = Buses.Find(BusSend.BusId);
-						if (Bus)
+						TSharedPtr<FMixerAudioBus> AudioBusPtr = AudioBuses.FindRef(AudioBusSend.AudioBusId);
+						if (AudioBusPtr.IsValid())
 						{
-							Bus->AddBusSend((EBusSendType)BusSendType, NewBusSend);
+							AudioBusPtr->AddSend((EBusSendType)BusSendType, NewAudioBusSend);
 						}
 						else
 						{
-							// If the bus is not registered, make a new entry
-							FMixerBus NewBusData(this, InitParams.NumInputChannels, NumOutputFrames);
+							// If the bus is not registered, make a new entry. This will default to an automatic audio bus until explicitly made manual later.
+							TSharedPtr<FMixerAudioBus> NewAudioBus(new FMixerAudioBus(this, true, InitParams.NumInputChannels));
 
 							// Add a send to it. This will not have a bus instance id (i.e. won't output audio), but 
 							// we register the send anyway in the event that this bus does play, we'll know to send this
 							// source's audio to it.
-							NewBusData.AddBusSend((EBusSendType)BusSendType, NewBusSend);
+							NewAudioBus->AddSend((EBusSendType)BusSendType, NewAudioBusSend);
 
-							Buses.Add(BusSend.BusId, NewBusData);
+							AudioBuses.Add(AudioBusSend.AudioBusId, NewAudioBus);
 						}
 
 						// Store on this source, which buses its sending its audio to
-						SourceInfo.BusSends[BusSendType].Add(BusSend.BusId);
+						SourceInfo.AudioBusSends[BusSendType].Add(AudioBusSend.AudioBusId);
 					}
 				}
 			}
@@ -892,6 +921,97 @@ namespace Audio
 		});
 	}
 
+	void FMixerSourceManager::StartAudioBus(uint32 InAudioBusId, int32 InNumChannels, bool bInIsAutomatic)
+	{
+		if (AudioBusIds_AudioThread.Contains(InAudioBusId))
+		{
+			return;
+		}
+
+		AudioBusIds_AudioThread.Add(InAudioBusId);
+
+		AudioMixerThreadCommand([this, InAudioBusId, InNumChannels, bInIsAutomatic]()
+		{
+			// If this audio bus id already exists, set it to not be automatic and return it
+			TSharedPtr<FMixerAudioBus> AudioBusPtr = AudioBuses.FindRef(InAudioBusId);
+			if (AudioBusPtr.IsValid())
+			{
+				// If this audio bus already existed, make sure the num channels lines up
+				ensure(AudioBusPtr->GetNumChannels() == InNumChannels);
+				AudioBusPtr->SetAutomatic(bInIsAutomatic);
+			}
+			else
+			{
+				// If the bus is not registered, make a new entry.
+				TSharedPtr<FMixerAudioBus> NewBusData(new FMixerAudioBus(this, bInIsAutomatic, InNumChannels));
+
+				AudioBuses.Add(InAudioBusId, NewBusData);
+			}
+
+			//  Now add any existing playing sources to this audio bus as sends if they exist
+			for (FSourceInfo& SourceInfo : SourceInfos)
+			{
+				if (SourceInfo.AudioBusId == InAudioBusId)
+				{
+					SourceInfo.bIsPlaying = false;
+					SourceInfo.bIsPaused = false;
+					SourceInfo.bIsActive = false;
+					SourceInfo.bIsStopping = false;
+				}
+			}
+		});
+	}
+
+	void FMixerSourceManager::StopAudioBus(uint32 InAudioBusId)
+	{
+		if (!AudioBusIds_AudioThread.Contains(InAudioBusId))
+		{
+			return;
+		}
+
+		AudioBusIds_AudioThread.Remove(InAudioBusId);
+
+		AudioMixerThreadCommand([this, InAudioBusId]()
+		{
+			TSharedPtr<FMixerAudioBus>* AudioBusPtr = AudioBuses.Find(InAudioBusId);
+			if (AudioBusPtr)
+			{
+				if (!(*AudioBusPtr)->IsAutomatic())
+				{
+					// Immediately stop all sources which were source buses
+					for (FSourceInfo& SourceInfo : SourceInfos)
+					{
+						if (SourceInfo.AudioBusId == InAudioBusId)
+						{
+							SourceInfo.bIsPlaying = false;
+							SourceInfo.bIsPaused = false;
+							SourceInfo.bIsActive = false;
+							SourceInfo.bIsStopping = false;
+						}
+					}
+					AudioBuses.Remove(InAudioBusId);
+				}
+			}
+		});
+	}
+
+	bool FMixerSourceManager::IsAudioBusActive(uint32 InAudioBusId)
+	{
+		AUDIO_MIXER_CHECK_GAME_THREAD(MixerDevice);
+		return AudioBusIds_AudioThread.Contains(InAudioBusId);
+	}
+
+	FPatchOutputStrongPtr FMixerSourceManager::AddPatchForAudioBus(uint32 InAudioBusId, float PatchGain)
+	{
+		AUDIO_MIXER_CHECK_AUDIO_PLAT_THREAD(MixerDevice);
+		TSharedPtr<FMixerAudioBus>* AudioBusPtr = AudioBuses.Find(InAudioBusId);
+		if (AudioBusPtr)
+		{
+			return (*AudioBusPtr)->AddNewPatch(NumOutputFrames * 2, PatchGain);
+		}
+		return nullptr;
+	}
+
 	void FMixerSourceManager::Play(const int32 SourceId)
 	{
 		AUDIO_MIXER_CHECK(SourceId < NumTotalSources);
@@ -1214,45 +1334,49 @@ namespace Audio
 		});
 	}
 
-	void FMixerSourceManager::SetBusSendInfo(const int32 SourceId, EBusSendType InBusSendType, FMixerBusSend& InBusSend)
+	void FMixerSourceManager::SetBusSendInfo(const int32 SourceId, EBusSendType InAudioBusSendType, uint32 AudioBusId, float BusSendLevel)
 	{
 		AUDIO_MIXER_CHECK(SourceId < NumTotalSources);
 		AUDIO_MIXER_CHECK(GameThreadInfo.bIsBusy[SourceId]);
 		AUDIO_MIXER_CHECK_GAME_THREAD(MixerDevice);
 
-		AudioMixerThreadCommand([this, SourceId, InBusSendType, InBusSend]()
+		AudioMixerThreadCommand([this, SourceId, InAudioBusSendType, AudioBusId, BusSendLevel]()
 		{
 			// Create mapping of source id to bus send level
-			FBusSend BusSend{ SourceId, InBusSend.SendLevel };
+			FAudioBusSend BusSend;
+			BusSend.SendLevel = BusSendLevel;
+			BusSend.SourceId = SourceId;
+			BusSend.AudioBusId = AudioBusId;
+
 			FSourceInfo& SourceInfo = SourceInfos[SourceId];
 
 			// Retrieve the bus we want to send audio to
-			FMixerBus* Bus = Buses.Find(InBusSend.BusId);
+			TSharedPtr<FMixerAudioBus>* AudioBusPtr = AudioBuses.Find(AudioBusId);
 
 			// If we already have a bus, we update the amount of audio we want to send to it
-			if (Bus)
+			if (AudioBusPtr)
 			{
-				Bus->AddBusSend(InBusSendType, BusSend);			
+				(*AudioBusPtr)->AddSend(InAudioBusSendType, BusSend);
 			}
 			else
 			{
 				// If the bus is not registered, make a new entry on the send
-				FMixerBus NewBusData(this, SourceInfo.NumInputChannels, NumOutputFrames);
+				TSharedPtr<FMixerAudioBus> NewBusData(new FMixerAudioBus(this, true, SourceInfo.NumInputChannels));
 
 				// Add a send to it. This will not have a bus instance id (i.e. won't output audio), but 
 				// we register the send anyway in the event that this bus does play, we'll know to send this
 				// source's audio to it.
-				NewBusData.AddBusSend(InBusSendType, BusSend);
+				NewBusData->AddSend(InAudioBusSendType, BusSend);
 
-				Buses.Add(InBusSend.BusId, NewBusData);
+				AudioBuses.Add(AudioBusId, NewBusData);
 			}
 
 			// Check to see if we need to create new bus data. If we are not playing a bus with this id, then we
 			// need to create a slot for it such that when a bus does play, it'll start rendering audio from this source
 			bool bExisted = false;
-			for (uint32 BusId : SourceInfo.BusSends[(int32)InBusSendType])
+			for (uint32 BusId : SourceInfo.AudioBusSends[(int32)InAudioBusSendType])
 			{
-				if (BusId == InBusSend.BusId)
+				if (BusId == AudioBusId)
 				{
 					bExisted = true;
 					break;
@@ -1261,7 +1385,7 @@ namespace Audio
 
 			if (!bExisted)
 			{
-				SourceInfo.BusSends[(int32)InBusSendType].Add(InBusSend.BusId);
+				SourceInfo.AudioBusSends[(int32)InAudioBusSendType].Add(AudioBusId);
 			}
 		});
 	}
@@ -1434,8 +1558,8 @@ namespace Audio
 				continue;
 			}
 
-			const bool bIsBus = SourceInfo.BusId != INDEX_NONE;
-			if ((bGenerateBuses && !bIsBus) || (!bGenerateBuses && bIsBus))
+			const bool bIsSourceBus = SourceInfo.AudioBusId != INDEX_NONE;
+			if ((bGenerateBuses && !bIsSourceBus) || (!bGenerateBuses && bIsSourceBus))
 			{
 				continue;
 			}
@@ -1457,27 +1581,30 @@ namespace Audio
 
 			// if this is a bus, we just want to copy the bus audio to this source's output audio
 			// Note we need to copy this since bus instances may have different audio via dynamic source effects, etc.
-			if (bIsBus)
+			if (bIsSourceBus)
 			{
-				// Get the source's rendered bus data
-				const FMixerBus* Bus = Buses.Find(SourceInfo.BusId);
-				const float* RESTRICT BusBufferPtr = Bus->GetCurrentBusBuffer();
-
-				int32 NumFramesPlayed = NumOutputFrames;
-				if (SourceInfo.BusDurationFrames != INDEX_NONE)
+				// Get the source's rendered and mixed audio bus data
+				const TSharedPtr<FMixerAudioBus>* AudioBusPtr = AudioBuses.Find(SourceInfo.AudioBusId);
+				if (ensure(AudioBusPtr != nullptr))
 				{
-					// If we're now finishing, only copy over the real data
-					if ((SourceInfo.NumFramesPlayed + NumOutputFrames) >= SourceInfo.BusDurationFrames)
+					const float* RESTRICT BusBufferPtr = (*AudioBusPtr)->GetCurrentBusBuffer();
+
+					int32 NumFramesPlayed = NumOutputFrames;
+					if (SourceInfo.SourceBusDurationFrames != INDEX_NONE)
 					{
-						NumFramesPlayed = SourceInfo.BusDurationFrames - SourceInfo.NumFramesPlayed;
-						SourceInfo.bIsLastBuffer = true;				
+						// If we're now finishing, only copy over the real data
+						if ((SourceInfo.NumFramesPlayed + NumOutputFrames) >= SourceInfo.SourceBusDurationFrames)
+						{
+							NumFramesPlayed = SourceInfo.SourceBusDurationFrames - SourceInfo.NumFramesPlayed;
+							SourceInfo.bIsLastBuffer = true;
+						}
 					}
-				}
 
-				SourceInfo.NumFramesPlayed += NumFramesPlayed;
+					SourceInfo.NumFramesPlayed += NumFramesPlayed;
 
-				// Simply copy into the pre distance attenuation buffer ptr
-				FMemory::Memcpy(PreDistanceAttenBufferPtr, BusBufferPtr, sizeof(float) * NumFramesPlayed * SourceInfo.NumInputChannels);
+					// Simply copy into the pre distance attenuation buffer ptr
+					FMemory::Memcpy(PreDistanceAttenBufferPtr, BusBufferPtr, sizeof(float) * NumFramesPlayed * SourceInfo.NumInputChannels);
+				}
 			}
 			else
 			{
@@ -1571,20 +1698,20 @@ namespace Audio
 	void FMixerSourceManager::ComputeBuses()
 	{
 		// Loop through the bus registry and mix source audio
-		for (auto& Entry : Buses)
+		for (auto& Entry : AudioBuses)
 		{
-			FMixerBus& Bus = Entry.Value;
-			Bus.MixBuffer();
+			TSharedPtr<FMixerAudioBus>& AudioBus = Entry.Value;
+			AudioBus->MixBuffer();
 		}
 	}
 
 	void FMixerSourceManager::UpdateBuses()
 	{
 		// Update the bus states post mixing. This flips the current/previous buffer indices.
-		for (auto& Entry : Buses)
+		for (auto& Entry : AudioBuses)
 		{
-			FMixerBus& Bus = Entry.Value;
-			Bus.Update();
+			TSharedPtr<FMixerAudioBus>& AudioBus = Entry.Value;
+			AudioBus->Update();
 		}
 	}
 
@@ -1704,7 +1831,10 @@ namespace Audio
 				SourceInfo.AudioPluginOutputData.AudioBuffer.AddZeroed(2 * NumOutputFrames);
 			}
 
-			SpatializationPlugin->ProcessAudio(AudioPluginInputData, SourceInfo.AudioPluginOutputData);
+			{
+				LLM_SCOPE(ELLMTag::AudioMixerPlugins);
+				SpatializationPlugin->ProcessAudio(AudioPluginInputData, SourceInfo.AudioPluginOutputData);
+			}
 
 			// If this is an external send, we treat this source audio as if it was still a mono source
 			// This will allow it to traditionally pan in the ComputeOutputBuffers function and be
@@ -2039,8 +2169,8 @@ namespace Audio
 				continue;
 			}
 
-			const bool bIsBus = SourceInfo.BusId != INDEX_NONE;
-			if ((bGenerateBuses && !bIsBus) || (!bGenerateBuses && bIsBus))
+			const bool bIsSourceBus = SourceInfo.AudioBusId != INDEX_NONE;
+			if ((bGenerateBuses && !bIsSourceBus) || (!bGenerateBuses && bIsSourceBus))
 			{
 				continue;
 			}
@@ -2049,7 +2179,7 @@ namespace Audio
 
 			// Copy and store the current state of the pre-distance attenuation buffer before we feed it through our source effects
 			// This is used by pre-effect sends
-			if (SourceInfo.BusSends[(int32)EBusSendType::PreEffect].Num() > 0)
+			if (SourceInfo.AudioBusSends[(int32)EBusSendType::PreEffect].Num() > 0)
 			{
 				SourceInfo.PreEffectBuffer.Reset();
 				SourceInfo.PreEffectBuffer.Reserve(SourceInfo.PreDistanceAttenuationBuffer.Num());
@@ -2238,8 +2368,8 @@ namespace Audio
 
 			// If we're in generate buses mode and not a bus, or vice versa, or if we're set to only output audio to buses.
 			// If set to output buses, no need to do any panning for the source. The buses will do the panning.
-			const bool bIsBus = SourceInfo.BusId != INDEX_NONE;
-			if ((bGenerateBuses && !bIsBus) || (!bGenerateBuses && bIsBus) || SourceInfo.bOutputToBusOnly)
+			const bool bIsSourceBus = SourceInfo.AudioBusId != INDEX_NONE;
+			if ((bGenerateBuses && !bIsSourceBus) || (!bGenerateBuses && bIsSourceBus) || SourceInfo.bOutputToBusOnly)
 			{
 				continue;
 			}
@@ -2292,7 +2422,7 @@ namespace Audio
 	void FMixerSourceManager::GenerateSourceAudio(const bool bGenerateBuses)
 	{
 		// If there are no buses, don't need to do anything here
-		if (bGenerateBuses && !Buses.Num())
+		if (bGenerateBuses && !AudioBuses.Num())
 		{
 			return;
 		}
@@ -2418,6 +2548,7 @@ namespace Audio
 			FSoundEffectSourceInitData InitData;
 			InitData.AudioClock = MixerDevice->GetAudioClock();
 			InitData.SampleRate = MixerDevice->SampleRate;
+			InitData.AudioDeviceId = MixerDevice->DeviceID;
 
 			for (int32 SourceId = 0; SourceId < NumTotalSources; ++SourceId)
 			{
@@ -2481,11 +2612,23 @@ namespace Audio
 		return SourceInfos[SourceId].PreEffectBuffer.GetData();
 	}
 
-	const float* FMixerSourceManager::GetPreviousBusBuffer(const int32 SourceId) const
+	const float* FMixerSourceManager::GetPreviousSourceBusBuffer(const int32 SourceId) const
 	{
-		const uint32 BusId = SourceInfos[SourceId].BusId;
-		const FMixerBus* MixerBus = Buses.Find(BusId);
-		return MixerBus->GetPreviousBusBuffer();
+		if (SourceId < SourceInfos.Num())
+		{
+			return GetPreviousAudioBusBuffer(SourceInfos[SourceId].AudioBusId);
+		}
+		return nullptr;
+	}
+
+	const float* FMixerSourceManager::GetPreviousAudioBusBuffer(const int32 AudioBusId) const
+	{
+		const TSharedPtr<FMixerAudioBus>* AudioBusPtr = AudioBuses.Find(AudioBusId);
+		if (AudioBusPtr)
+		{
+			return (*AudioBusPtr)->GetPreviousBusBuffer();
+		}
+		return nullptr;
 	}
 
 	int32 FMixerSourceManager::GetNumChannels(const int32 SourceId) const
@@ -2493,9 +2636,9 @@ namespace Audio
 		return SourceInfos[SourceId].NumInputChannels;
 	}
 
-	bool FMixerSourceManager::IsBus(const int32 SourceId) const
+	bool FMixerSourceManager::IsSourceBus(const int32 SourceId) const
 	{
-		return SourceInfos[SourceId].BusId != INDEX_NONE;
+		return SourceInfos[SourceId].AudioBusId != INDEX_NONE;
 	}
 
 	void FMixerSourceManager::ComputeNextBlockOfSamples()
@@ -2534,6 +2677,7 @@ namespace Audio
 		if (bUsingSpatializationPlugin)
 		{
 			AUDIO_MIXER_CHECK(SpatializationPlugin.IsValid());
+			LLM_SCOPE(ELLMTag::AudioMixerPlugins);
 			SpatializationPlugin->OnAllSourcesProcessed();
 		}
 
@@ -2581,6 +2725,20 @@ namespace Audio
 
 		// Add the function to the command queue:
 		int32 AudioThreadCommandIndex = !RenderThreadCommandBufferIndex.GetValue();
+		
+#if !NO_LOGGING
+		static uint32 WarnSize = 1024 * 1024;
+		SIZE_T Size = CommandBuffers[AudioThreadCommandIndex].SourceCommandQueue.GetAllocatedSize();
+		if (Size > WarnSize )
+		{		
+			SIZE_T Num = CommandBuffers[AudioThreadCommandIndex].SourceCommandQueue.Num();
+			// NOTE: Although not really and error we want this to show up in shipping builds.
+			UE_LOG(LogAudioMixer, Error, TEXT("Command Queue has grown to %uk bytes, containing %d cmds, last pump was %fms ago."), 
+				Size >> 10, Num, FPlatformTime::ToMilliseconds64(FPlatformTime::Cycles64() - LastPumpTimeInCycles));
+			WarnSize *= 2;
+		}
+#endif //!NO_LOGGING
+
 		CommandBuffers[AudioThreadCommandIndex].SourceCommandQueue.Add(MoveTemp(InFunction));
 		NumCommands.Increment();
 	}
@@ -2608,6 +2766,7 @@ namespace Audio
 			NumCommands.Decrement();
 		}
 
+		LastPumpTimeInCycles = FPlatformTime::Cycles64();
 		Commands.SourceCommandQueue.Reset();
 
 		if (FPlatformProcess::SupportsMultithreading())
diff --git a/Engine/Source/Runtime/AudioMixer/Private/AudioMixerSourceManager.h b/Engine/Source/Runtime/AudioMixer/Private/AudioMixerSourceManager.h
--- a/Engine/Source/Runtime/AudioMixer/Private/AudioMixerSourceManager.h
+++ b/Engine/Source/Runtime/AudioMixer/Private/AudioMixerSourceManager.h
@@ -35,15 +35,10 @@ namespace Audio
 		AlignedFloatBuffer AudioData;
 
 		/** How many times this buffer will loop. */
-		int32 LoopCount;
+		int32 LoopCount = 0;
 
 		/** If this buffer is from real-time decoding and needs to make callbacks for more data. */
-		uint32 bRealTimeBuffer : 1;
-
-		FMixerSourceVoiceBuffer()
-		{
-			FMemory::Memzero(this, sizeof(*this));
-		}
+		bool bRealTimeBuffer = false;
 	};
 
 	typedef TSharedPtr<FMixerSubmix, ESPMode::ThreadSafe> FMixerSubmixPtr;
@@ -72,79 +67,53 @@ namespace Audio
 		FMixerSubmixWeakPtr Submix;
 
 		// The amount of audio that is to be mixed into this submix
-		float SendLevel;
+		float SendLevel = 0.0f;
 
 		// Whather or not this is the primary send (i.e. first in the send chain)
-		bool bIsMainSend;
+		bool bIsMainSend = false;
 
 		// If this is a soundfield submix, this is a pointer to the submix's Soundfield Factory.
 		// If this is nullptr, the submix is not a soundfield submix.
-		ISoundfieldFactory* SoundfieldFactory;
+		ISoundfieldFactory* SoundfieldFactory = nullptr;
 	};
 
 	// Struct holding mappings of bus ids (unique ids) to send level
-	struct FMixerBusSend
+	struct FInitAudioBusSend
 	{
-		uint32 BusId;
-		float SendLevel;
+		uint32 AudioBusId = INDEX_NONE;
+		float SendLevel = 0.0f;
 	};
 
 	struct FMixerSourceVoiceInitParams
 	{
-		TSharedPtr<FMixerSourceBuffer> MixerSourceBuffer;
-		ISourceListener* SourceListener;
+		TSharedPtr<FMixerSourceBuffer, ESPMode::ThreadSafe> MixerSourceBuffer = nullptr;
+		ISourceListener* SourceListener = nullptr;
 		TArray<FMixerSourceSubmixSend> SubmixSends;
-		TArray<FMixerBusSend> BusSends[(int32)EBusSendType::Count];
-		uint32 BusId;
-		float BusDuration;
-		uint32 SourceEffectChainId;
+		TArray<FInitAudioBusSend> AudioBusSends[(int32)EBusSendType::Count];
+		uint32 AudioBusId = INDEX_NONE;
+		float SourceBusDuration = 0.0f;
+		uint32 SourceEffectChainId = INDEX_NONE;
 		TArray<FSourceEffectChainEntry> SourceEffectChain;
-		FMixerSourceVoice* SourceVoice;
-		int32 NumInputChannels;
-		int32 NumInputFrames;
-		float EnvelopeFollowerAttackTime;
-		float EnvelopeFollowerReleaseTime;
+		FMixerSourceVoice* SourceVoice = nullptr;
+		int32 NumInputChannels = 0;
+		int32 NumInputFrames = 0;
+		float EnvelopeFollowerAttackTime = 10.0f;
+		float EnvelopeFollowerReleaseTime = 100.0f;
 		FString DebugName;
-		USpatializationPluginSourceSettingsBase* SpatializationPluginSettings;
-		UOcclusionPluginSourceSettingsBase* OcclusionPluginSettings;
-		UReverbPluginSourceSettingsBase* ReverbPluginSettings;
-		USoundModulationPluginSourceSettingsBase* ModulationPluginSettings;
+		USpatializationPluginSourceSettingsBase* SpatializationPluginSettings = nullptr;
+		UOcclusionPluginSourceSettingsBase* OcclusionPluginSettings = nullptr;
+		UReverbPluginSourceSettingsBase* ReverbPluginSettings = nullptr;
+		USoundModulationPluginSourceSettingsBase* ModulationPluginSettings = nullptr;
 		FName AudioComponentUserID;
-		uint64 AudioComponentID;
-		uint8 bPlayEffectChainTails : 1;
-		uint8 bUseHRTFSpatialization : 1;
-		uint8 bIsExternalSend : 1;
-		uint8 bIsDebugMode : 1;
-		uint8 bOutputToBusOnly : 1;
-		uint8 bIsVorbis : 1;
-		uint8 bIsAmbisonics : 1;
-		uint8 bIsSeeking : 1;
-
-		FMixerSourceVoiceInitParams()
-			: MixerSourceBuffer(nullptr)
-			, SourceListener(nullptr)
-			, BusId(INDEX_NONE)
-			, BusDuration(0.0f)
-			, SourceEffectChainId(INDEX_NONE)
-			, SourceVoice(nullptr)
-			, NumInputChannels(0)
-			, NumInputFrames(0)
-			, EnvelopeFollowerAttackTime(10.0f)
-			, EnvelopeFollowerReleaseTime(100.0f)
-			, SpatializationPluginSettings(nullptr)
-			, OcclusionPluginSettings(nullptr)
-			, ReverbPluginSettings(nullptr)
-			, ModulationPluginSettings(nullptr)
-			, AudioComponentID(0)
-			, bPlayEffectChainTails(false)
-			, bUseHRTFSpatialization(false)
-			, bIsExternalSend(false)
-			, bIsDebugMode(false)
-			, bOutputToBusOnly(false)
-			, bIsVorbis(false)
-			, bIsAmbisonics(false)
-			, bIsSeeking(false)
-		{}
+		uint64 AudioComponentID = 0;
+		bool bPlayEffectChainTails = false;
+		bool bUseHRTFSpatialization = false;
+		bool bIsExternalSend = false;
+		bool bIsDebugMode  = false;
+		bool bOutputToBusOnly = false;
+		bool bIsVorbis = false;
+		bool bIsAmbisonics = false;
+		bool bIsSeeking = false;
 	};
 
 	struct FSourceChannelMap
@@ -154,8 +123,8 @@ namespace Audio
 
 		// This is the number of bytes the gain array is using:
 		// (Number of input channels * number of output channels) * sizeof float.
-		int32 CopySize;
-		bool bIsInit;
+		int32 CopySize = 0;
+		bool bIsInit = false;
 
 		FSourceChannelMap(int32 InNumInChannels, int32 InNumOutChannels) 
 			: CopySize(InNumInChannels * InNumOutChannels * sizeof(float))
@@ -202,10 +171,10 @@ namespace Audio
 	struct FSourceManagerInitParams
 	{
 		// Total number of sources to use in the source manager
-		int32 NumSources;
+		int32 NumSources = 0;
 
 		// Number of worker threads to use for the source manager.
-		int32 NumSourceWorkers;
+		int32 NumSourceWorkers = 0;
 	};
 
 	class FMixerSourceManager
@@ -219,11 +188,17 @@ namespace Audio
 
 		bool GetFreeSourceId(int32& OutSourceId);
 		int32 GetNumActiveSources() const;
-		int32 GetNumActiveBuses() const;
+		int32 GetNumActiveAudioBuses() const;
 
 		void ReleaseSourceId(const int32 SourceId);
 		void InitSource(const int32 SourceId, const FMixerSourceVoiceInitParams& InitParams);
 
+		// Creates an audio bus manually. Returns an audio bus Id.
+		void StartAudioBus(uint32 InAudioBusId, int32 InNumChannels, bool bInIsAutomatic);
+		void StopAudioBus(uint32 InAudioBusId);
+		bool IsAudioBusActive(uint32 InAudioBusId);
+		FPatchOutputStrongPtr AddPatchForAudioBus(uint32 InAudioBusId, float PatchGain);
+
 		void Play(const int32 SourceId);
 		void Stop(const int32 SourceId);
 		void StopFade(const int32 SourceId, const int32 NumFrames);
@@ -256,7 +231,8 @@ namespace Audio
 		const FQuat GetListenerRotation(const int32 SourceId) const;
 
 		void SetSubmixSendInfo(const int32 SourceId, const FMixerSourceSubmixSend& SubmixSend);
-		void SetBusSendInfo(const int32 SourceId, EBusSendType InBusSendType, FMixerBusSend& BusSend);
+
+		void SetBusSendInfo(const int32 SourceId, EBusSendType InAudioBusSendType, uint32 AudiobusId, float BusSendLevel);
 
 		void UpdateDeviceChannelCount(const int32 InNumOutputChannels);
 
@@ -265,10 +241,11 @@ namespace Audio
 
 		const float* GetPreDistanceAttenuationBuffer(const int32 SourceId) const;
 		const float* GetPreEffectBuffer(const int32 SourceId) const;
-		const float* GetPreviousBusBuffer(const int32 SourceId) const;
+		const float* GetPreviousSourceBusBuffer(const int32 SourceId) const;
+		const float* GetPreviousAudioBusBuffer(const int32 AudioBusId) const;
 		int32 GetNumChannels(const int32 SourceId) const;
 		int32 GetNumOutputFrames() const { return NumOutputFrames; }
-		bool IsBus(const int32 SourceId) const;
+		bool IsSourceBus(const int32 SourceId) const;
 		void PumpCommandQueue();
 		void UpdatePendingReleaseData(bool bForceWait = false);
 		void FlushCommandQueue(bool bPumpCommandQueue = false);
@@ -451,11 +428,11 @@ namespace Audio
 			~FSourceInfo() {}
 
 			// Object which handles source buffer decoding
-			TSharedPtr<FMixerSourceBuffer> MixerSourceBuffer;
+			TSharedPtr<FMixerSourceBuffer, ESPMode::ThreadSafe> MixerSourceBuffer;
 			ISourceListener* SourceListener;
 
 			// Data used for rendering sources
-			TSharedPtr<FMixerSourceVoiceBuffer> CurrentPCMBuffer;
+			TSharedPtr<FMixerSourceVoiceBuffer, ESPMode::ThreadSafe> CurrentPCMBuffer;
 			int32 CurrentAudioChunkNumFrames;
 
 			// The post-attenuation source buffer, used to send audio to submixes
@@ -475,14 +452,14 @@ namespace Audio
 
 			TArray<FMixerSourceSubmixSend> SubmixSends;
 
-			// What bus Id this source is, if it is a bus. This is INDEX_NONE for sources which are not buses.
-			uint32 BusId;
+			// What audio bus Id this source is sonfiying, if it is a source bus. This is INDEX_NONE for sources which are not source buses.
+			uint32 AudioBusId;
 
-			// Number of samples to count for bus
-			int64 BusDurationFrames;
+			// Number of samples to count for source bus
+			int64 SourceBusDurationFrames;
 
 			// What buses this source is sending its audio to. Used to remove this source from the bus send list.
-			TArray<uint32> BusSends[(int32)EBusSendType::Count];
+			TArray<uint32> AudioBusSends[(int32)EBusSendType::Count];
 
 			// Interpolated source params
 			FParam PitchSourceParam;
@@ -585,14 +562,15 @@ namespace Audio
 		// sources.
 		TArray<FSourceDownmixData> DownmixDataArray;
 
-		// Map of bus object Id's to bus data. 
-		TMap<uint32, FMixerBus> Buses;
+		// Map of bus object Id's to audio bus data. 
+		TMap<uint32, TSharedPtr<FMixerAudioBus>> AudioBuses;
+		TArray<uint32> AudioBusIds_AudioThread;
 
 		// Async task workers for processing sources in parallel
 		TArray<FAsyncTask<FAudioMixerSourceWorker>*> SourceWorkers;
 
 		// Array of task data waiting to finished. Processed on audio render thread.
-		TArray<TSharedPtr<FMixerSourceBuffer>> PendingSourceBuffers;
+		TArray<TSharedPtr<FMixerSourceBuffer, ESPMode::ThreadSafe>> PendingSourceBuffers;
 
 		// General information about sources in source manager accessible from game thread
 		struct FGameThreadInfo
@@ -619,6 +597,7 @@ namespace Audio
 
 		// Set to true when the audio source manager should pump the command queue
 		FThreadSafeBool bPumpQueue;
+		uint64 LastPumpTimeInCycles = 0;
 
 		friend class FMixerSourceVoice;
 	};
diff --git a/Engine/Source/Runtime/AudioMixer/Private/Components/SynthComponent.cpp b/Engine/Source/Runtime/AudioMixer/Private/Components/SynthComponent.cpp
--- a/Engine/Source/Runtime/AudioMixer/Private/Components/SynthComponent.cpp
+++ b/Engine/Source/Runtime/AudioMixer/Private/Components/SynthComponent.cpp
@@ -219,7 +219,7 @@ void USynthComponent::Initialize(int32 SampleRateOverride)
 
 		if (!Synth)
 		{
-			Synth = NewObject<USynthSound>(this, TEXT("Synth"));
+			Synth = NewObject<USynthSound>(this, FName(TEXT("TimeSynth"), this->GetUniqueID()));
 		}
 
 		// Copy sound base data to the sound
