Undo //UE4/Release-4.25/Engine/Source/Runtime/AudioMixer/Private/... changelist 12423730
CIS failure (2019 strict)

diff --git a/Engine/Source/Runtime/AudioMixer/Private/AudioMixerSource.h b/Engine/Source/Runtime/AudioMixer/Private/AudioMixerSource.h
--- a/Engine/Source/Runtime/AudioMixer/Private/AudioMixerSource.h
+++ b/Engine/Source/Runtime/AudioMixer/Private/AudioMixerSource.h
@@ -61,9 +61,6 @@ namespace Audio
 
 	private:
 
-		/** Initializes the bus sends. */
-		void InitBusSends(FWaveInstance* InWaveInstance, FMixerSourceVoiceInitParams& InitParams);
-
 		/** Frees any resources for this sound source. */
 		void FreeResources();
 
@@ -119,7 +116,7 @@ namespace Audio
 
 		FMixerDevice* MixerDevice;
 		FMixerBuffer* MixerBuffer;
-		TSharedPtr<FMixerSourceBuffer, ESPMode::ThreadSafe> MixerSourceBuffer;
+		TSharedPtr<FMixerSourceBuffer> MixerSourceBuffer;
 		FMixerSourceVoice* MixerSourceVoice;
 
 		// This holds data copied from FSoundSourceBusSendInfo when a new sound starts playing
diff --git a/Engine/Source/Runtime/AudioMixer/Private/AudioMixerSourceBuffer.cpp b/Engine/Source/Runtime/AudioMixer/Private/AudioMixerSourceBuffer.cpp
--- a/Engine/Source/Runtime/AudioMixer/Private/AudioMixerSourceBuffer.cpp
+++ b/Engine/Source/Runtime/AudioMixer/Private/AudioMixerSourceBuffer.cpp
@@ -56,7 +56,7 @@ namespace Audio
 		return CurrentSample >= NumSamples;
 	}
 
-	TSharedPtr<FMixerSourceBuffer, ESPMode::ThreadSafe> FMixerSourceBuffer::Create(FMixerBuffer& InBuffer, USoundWave& InWave, ELoopingMode InLoopingMode, bool bInIsSeeking, bool bInForceSyncDecode)
+	TSharedPtr<FMixerSourceBuffer> FMixerSourceBuffer::Create(FMixerBuffer& InBuffer, USoundWave& InWave, ELoopingMode InLoopingMode, bool bInIsSeeking, bool bInForceSyncDecode)
 	{
 		LLM_SCOPE(ELLMTag::AudioMixer);
 
@@ -72,7 +72,7 @@ namespace Audio
 			return nullptr;
 		}
 
-		TSharedPtr<FMixerSourceBuffer, ESPMode::ThreadSafe> NewSourceBuffer = MakeShareable(new FMixerSourceBuffer(InBuffer, InWave, InLoopingMode, bInIsSeeking, bInForceSyncDecode));
+		TSharedPtr<FMixerSourceBuffer> NewSourceBuffer = MakeShareable(new FMixerSourceBuffer(InBuffer, InWave, InLoopingMode, bInIsSeeking, bInForceSyncDecode));
 		return NewSourceBuffer;
 	}
 
@@ -92,7 +92,7 @@ namespace Audio
 		, bIsSeeking(bInIsSeeking)
 		, bLoopCallback(false)
 		, bProcedural(InWave.bProcedural)
-		, bIsBus(InWave.bIsSourceBus)
+		, bIsBus(InWave.bIsBus)
 		, bForceSyncDecode(bInForceSyncDecode)
 	{
 		InWave.AddPlayingSource(this);
@@ -100,7 +100,7 @@ namespace Audio
 		const uint32 TotalSamples = MONO_PCM_BUFFER_SAMPLES * NumChannels;
 		for (int32 BufferIndex = 0; BufferIndex < Audio::MAX_BUFFERS_QUEUED; ++BufferIndex)
 		{
-			SourceVoiceBuffers.Add(MakeShared<FMixerSourceVoiceBuffer, ESPMode::ThreadSafe>());
+			SourceVoiceBuffers.Add(MakeShared<FMixerSourceVoiceBuffer>());
 
 			// Prepare the memory to fit the max number of samples
 			SourceVoiceBuffers[BufferIndex]->AudioData.Reset(TotalSamples);
@@ -113,14 +113,29 @@ namespace Audio
 	{
 		// GC methods may get called from the game thread during the destructor
 		// These methods will trylock and early exit if we have this lock
-		FScopeLock Lock(&SoundWaveCritSec);
+		FScopeLock Lock(&DtorCritSec);
+
+		// Make sure we have completed our async realtime task before deleting the decompression state
+		if (AsyncRealtimeAudioTask)
+		{
+			AsyncRealtimeAudioTask->CancelTask();
+			delete AsyncRealtimeAudioTask;
+			AsyncRealtimeAudioTask = nullptr;
+		}
 
-		// OnEndGenerate calls EnsureTaskFinishes,
-		// which will make sure we have completed our async realtime task before deleting the decompression state
 		OnEndGenerate();
 
 		// Clean up decompression state after things have been finished using it
-		DeleteDecoder();
+		if (DecompressionState)
+		{
+			if (BufferType == EBufferType::Streaming)
+			{
+				IStreamingManager::Get().GetAudioStreamingManager().RemoveDecoder(DecompressionState);
+			}
+
+			delete DecompressionState;
+			DecompressionState = nullptr;
+		}
 
 		if (SoundWave)
 		{
@@ -160,18 +175,18 @@ namespace Audio
 
 		switch (BufferType)
 		{
-		case EBufferType::PCM:
-		case EBufferType::PCMPreview:
-			SubmitInitialPCMBuffers();
-			break;
+			case EBufferType::PCM:
+			case EBufferType::PCMPreview:
+				SubmitInitialPCMBuffers();
+				break;
 
-		case EBufferType::PCMRealTime:
-		case EBufferType::Streaming:
-			SubmitInitialRealtimeBuffers();
-			break;
+			case EBufferType::PCMRealTime:
+			case EBufferType::Streaming:
+				SubmitInitialRealtimeBuffers();
+				break;
 
-		case EBufferType::Invalid:
-			break;
+			case EBufferType::Invalid:
+				break;
 		}
 
 		return true;
@@ -179,9 +194,7 @@ namespace Audio
 
 	void FMixerSourceBuffer::OnBufferEnd()
 	{
-		FScopeTryLock Lock(&SoundWaveCritSec);
-
-		if (!Lock.IsLocked() || (NumBuffersQeueued == 0 && bBufferFinished) || (bProcedural && !SoundWave))
+		if ((NumBuffersQeueued == 0 && bBufferFinished) || (bProcedural && !SoundWave))
 		{
 			return;
 		}
@@ -189,26 +202,9 @@ namespace Audio
 		ProcessRealTimeSource();
 	}
 
-	int32 FMixerSourceBuffer::GetNumBuffersQueued() const
+	TSharedPtr<FMixerSourceVoiceBuffer> FMixerSourceBuffer::GetNextBuffer()
 	{
-		FScopeTryLock Lock(&SoundWaveCritSec);
-		if (Lock.IsLocked())
-		{
-			return NumBuffersQeueued;
-		}
-
-		return 0;
-	}
-
-	TSharedPtr<FMixerSourceVoiceBuffer, ESPMode::ThreadSafe> FMixerSourceBuffer::GetNextBuffer()
-	{
-		FScopeTryLock Lock(&SoundWaveCritSec);
-		if(!Lock.IsLocked())
-		{
-			return nullptr;
-		}
-
-		TSharedPtr<FMixerSourceVoiceBuffer, ESPMode::ThreadSafe> NewBufferPtr;
+		TSharedPtr<FMixerSourceVoiceBuffer> NewBufferPtr;
 		BufferQueue.Dequeue(NewBufferPtr);
 		--NumBuffersQeueued;
 		return NewBufferPtr;
@@ -312,7 +308,7 @@ namespace Audio
 		{
 			ProcessRealTimeSource();
 		}
-			}
+	}
 
 	bool FMixerSourceBuffer::ReadMoreRealtimeData(ICompressedAudioInfo* InDecoder, const int32 BufferIndex, EBufferReadMode BufferReadMode)
 	{
@@ -322,24 +318,17 @@ namespace Audio
 
 		if (bProcedural)
 		{
-			FScopeTryLock Lock(&SoundWaveCritSec);
-			
-			if (Lock.IsLocked() && SoundWave)
-			{
-				check(SoundWave && SoundWave->bProcedural);
-				FProceduralAudioTaskData NewTaskData;
-				NewTaskData.ProceduralSoundWave = SoundWave;
-				NewTaskData.AudioData = SourceVoiceBuffers[BufferIndex]->AudioData.GetData();
-				NewTaskData.NumSamples = MaxSamples;
-				NewTaskData.NumChannels = NumChannels;
-				check(!AsyncRealtimeAudioTask);
-				AsyncRealtimeAudioTask = CreateAudioTask(NewTaskData);
-
-				// Procedural sound waves never loop
-				return false;
-			}
-
-			return true;
+			check(SoundWave && SoundWave->bProcedural);
+			FProceduralAudioTaskData NewTaskData;
+			NewTaskData.ProceduralSoundWave = SoundWave;
+			NewTaskData.AudioData = SourceVoiceBuffers[BufferIndex]->AudioData.GetData();
+			NewTaskData.NumSamples = MaxSamples;
+			NewTaskData.NumChannels = NumChannels;
+			check(!AsyncRealtimeAudioTask);
+			AsyncRealtimeAudioTask = CreateAudioTask(NewTaskData);
+
+			// Procedural sound waves never loop
+			return false;
 		}
 		else if (BufferType != EBufferType::PCMRealTime && BufferType != EBufferType::Streaming)
 		{
@@ -362,7 +351,6 @@ namespace Audio
 		NewTaskData.NumPrecacheFrames = NumPrecacheFrames;
 		NewTaskData.bForceSyncDecode = bForceSyncDecode;
 
-		FScopeLock Lock(&DecodeTaskCritSec);
 		check(!AsyncRealtimeAudioTask);
 		AsyncRealtimeAudioTask = CreateAudioTask(NewTaskData);
 
@@ -401,7 +389,6 @@ namespace Audio
 
 	void FMixerSourceBuffer::ProcessRealTimeSource()
 	{
-		FScopeLock Lock(&DecodeTaskCritSec);
 		if (AsyncRealtimeAudioTask)
 		{
 			AsyncRealtimeAudioTask->EnsureCompletion();
@@ -464,40 +451,23 @@ namespace Audio
 		}
 	}
 
-	void FMixerSourceBuffer::SubmitBuffer(TSharedPtr<FMixerSourceVoiceBuffer, ESPMode::ThreadSafe> InSourceVoiceBuffer)
+	void FMixerSourceBuffer::SubmitBuffer(TSharedPtr<FMixerSourceVoiceBuffer> InSourceVoiceBuffer)
 	{
 		NumBuffersQeueued++;
 		BufferQueue.Enqueue(InSourceVoiceBuffer);
 	}
 
-	void FMixerSourceBuffer::DeleteDecoder()
-	{
-		// Clean up decompression state after things have been finished using it
-		if (DecompressionState)
-		{
-			if (BufferType == EBufferType::Streaming)
-			{
-				IStreamingManager::Get().GetAudioStreamingManager().RemoveDecoder(DecompressionState);
-			}
-
-			delete DecompressionState;
-			DecompressionState = nullptr;
-		}
-	}
-
 	bool FMixerSourceBuffer::OnBeginDestroy(USoundWave * /*Wave*/)
 	{
-		FScopeTryLock Lock(&SoundWaveCritSec);
+		FScopeTryLock Lock(&DtorCritSec);
 
 		// if we don't have the lock, it means we are in ~FMixerSourceBuffer() on another thread
 		if (Lock.IsLocked() && SoundWave)
 		{
-			EnsureAsyncTaskFinishes();
-			DeleteDecoder();
-			ClearWave();
+			SoundWave = nullptr;
 			return true;
 		}
-
+		
 		return false;
 	}
 
@@ -508,26 +478,22 @@ namespace Audio
 
 	void FMixerSourceBuffer::OnFinishDestroy(USoundWave * /*Wave*/)
 	{
-		EnsureAsyncTaskFinishes();
-		FScopeTryLock Lock(&SoundWaveCritSec);
+		FScopeTryLock Lock(&DtorCritSec);
 
 		// if we don't have the lock, it means we are in ~FMixerSourceBuffer() on another thread
 		if (Lock.IsLocked() && SoundWave)
 		{
-			DeleteDecoder();
-			ClearWave();
+			SoundWave = nullptr;
 		}
 	}
 
 	bool FMixerSourceBuffer::IsAsyncTaskInProgress() const
 	{ 
-		FScopeLock Lock(&DecodeTaskCritSec);
-		return AsyncRealtimeAudioTask != nullptr;
+		return AsyncRealtimeAudioTask != nullptr; 
 	}
 
 	bool FMixerSourceBuffer::IsAsyncTaskDone() const 
 	{
-		FScopeLock Lock(&DecodeTaskCritSec);
 		if (AsyncRealtimeAudioTask)
 		{
 			return AsyncRealtimeAudioTask->IsDone();
@@ -537,7 +503,6 @@ namespace Audio
 	
 	void FMixerSourceBuffer::EnsureAsyncTaskFinishes() 
 	{
-		FScopeLock Lock(&DecodeTaskCritSec);
 		if (AsyncRealtimeAudioTask) 
 		{ 
 			AsyncRealtimeAudioTask->CancelTask(); 
@@ -549,9 +514,7 @@ namespace Audio
 
 	void FMixerSourceBuffer::OnBeginGenerate()
 	{
-		FScopeTryLock Lock(&SoundWaveCritSec);
-
-		if (Lock.IsLocked() && SoundWave && bProcedural)
+		if (SoundWave && bProcedural)
 		{
 			check(SoundWave && SoundWave->bProcedural);
 			SoundWave->OnBeginGenerate();
@@ -562,13 +525,6 @@ namespace Audio
 	{		
 		// Make sure the async task finishes!
 		EnsureAsyncTaskFinishes();
-
-		FScopeTryLock Lock(&SoundWaveCritSec);
-		if (!Lock.IsLocked())
-		{
-			return;
-		}
-
 						
 		// Only need to call OnEndGenerate and access SoundWave here if we successfully initialized
 		if (SoundWave && bInitialized && bProcedural)
diff --git a/Engine/Source/Runtime/AudioMixer/Private/AudioMixerSourceBuffer.h b/Engine/Source/Runtime/AudioMixer/Private/AudioMixerSourceBuffer.h
--- a/Engine/Source/Runtime/AudioMixer/Private/AudioMixerSourceBuffer.h
+++ b/Engine/Source/Runtime/AudioMixer/Private/AudioMixerSourceBuffer.h
@@ -43,7 +43,7 @@ namespace Audio
 		AsynchronousSkipFirstFrame
 	};
 
-	using FMixerSourceBufferPtr = TSharedPtr<class FMixerSourceBuffer, ESPMode::ThreadSafe>;
+	using FMixerSourceBufferPtr = TSharedPtr<class FMixerSourceBuffer>;
 
 	/** Class which handles decoding audio for a particular source buffer. */
 	class FMixerSourceBuffer : public ISoundWaveClient
@@ -68,10 +68,10 @@ namespace Audio
 		void OnBufferEnd();
 
 		// Return the number of buffers enqueued on the mixer source buffer
-		int32 GetNumBuffersQueued() const;
-
+		int32 GetNumBuffersQueued() const { return NumBuffersQeueued; }
+		
 		// Returns the next enqueued buffer, returns nullptr if no buffers enqueued
-		TSharedPtr<FMixerSourceVoiceBuffer, ESPMode::ThreadSafe> GetNextBuffer();
+		TSharedPtr<FMixerSourceVoiceBuffer> GetNextBuffer();
 
 		// Returns if buffer looped
 		bool DidBufferLoop() const { return bLoopCallback; }
@@ -102,15 +102,14 @@ namespace Audio
 		void SubmitInitialRealtimeBuffers();
 		void SubmitRealTimeSourceData(const bool bLooped);
 		void ProcessRealTimeSource();
-		void SubmitBuffer(TSharedPtr<FMixerSourceVoiceBuffer, ESPMode::ThreadSafe> InSourceVoiceBuffer);
-		void DeleteDecoder();
+		void SubmitBuffer(TSharedPtr<FMixerSourceVoiceBuffer> InSourceVoiceBuffer);
 
 
 		int32 NumBuffersQeueued;
 		FRawPCMDataBuffer RawPCMDataBuffer;
 
-		TArray<TSharedPtr<FMixerSourceVoiceBuffer, ESPMode::ThreadSafe>> SourceVoiceBuffers;
-		TQueue<TSharedPtr<FMixerSourceVoiceBuffer, ESPMode::ThreadSafe>> BufferQueue;
+		TArray<TSharedPtr<FMixerSourceVoiceBuffer>> SourceVoiceBuffers;
+		TQueue<TSharedPtr<FMixerSourceVoiceBuffer>> BufferQueue;
 		int32 CurrentBuffer;
 		// SoundWaves are only set for procedural sound waves
 		USoundWave* SoundWave;
@@ -122,8 +121,7 @@ namespace Audio
 		int32 NumPrecacheFrames;
 		TArray<uint8> CachedRealtimeFirstBuffer;
 
-		mutable FCriticalSection SoundWaveCritSec;
-		mutable FCriticalSection DecodeTaskCritSec;
+		FCriticalSection DtorCritSec;
 
 		uint32 bInitialized : 1;
 		uint32 bBufferFinished : 1;
diff --git a/Engine/Source/Runtime/AudioMixer/Private/AudioMixerSourceManager.cpp b/Engine/Source/Runtime/AudioMixer/Private/AudioMixerSourceManager.cpp
--- a/Engine/Source/Runtime/AudioMixer/Private/AudioMixerSourceManager.cpp
+++ b/Engine/Source/Runtime/AudioMixer/Private/AudioMixerSourceManager.cpp
@@ -203,11 +203,11 @@ namespace Audio
 			SourceInfo.NumFramesPlayed = 0;
 			SourceInfo.StartTime = 0.0;
 			SourceInfo.SubmixSends.Reset();
-			SourceInfo.AudioBusId = INDEX_NONE;
-			SourceInfo.SourceBusDurationFrames = INDEX_NONE;
+			SourceInfo.BusId = INDEX_NONE;
+			SourceInfo.BusDurationFrames = INDEX_NONE;
 		
-			SourceInfo.AudioBusSends[(int32)EBusSendType::PreEffect].Reset();
-			SourceInfo.AudioBusSends[(int32)EBusSendType::PostEffect].Reset();
+			SourceInfo.BusSends[(int32)EBusSendType::PreEffect].Reset();
+			SourceInfo.BusSends[(int32)EBusSendType::PostEffect].Reset();
 
 			SourceInfo.SourceEffectChainId = INDEX_NONE;
 
@@ -288,23 +288,6 @@ namespace Audio
 			MaxChannelsSupportedBySpatializationPlugin = MixerDevice->MaxChannelsSupportedBySpatializationPlugin;
 		}
 
-		// Spam command queue with nops.
-		static FAutoConsoleCommand SpamNopsCmd(
-			TEXT("au.SpamCommandQueue"),
-			TEXT(""),
-			FConsoleCommandDelegate::CreateLambda([this]() 
-			{				
-				struct FSpamPayload
-				{
-					uint8 JunkBytes[1024];
-				} Payload;
-				for (int32 i = 0; i < 65536; ++i)
-				{
-					AudioMixerThreadCommand([Payload] {});
-				}
-			})
-		);
-
 		bInitialized = true;
 		bPumpQueue = false;
 	}
@@ -391,12 +374,7 @@ namespace Audio
 
 		AUDIO_MIXER_CHECK(SourceId < NumTotalSources);
 		AUDIO_MIXER_CHECK(bInitialized);
-
-		if (MixerSources[SourceId] == nullptr)
-		{
-			UE_LOG(LogAudioMixer, Warning, TEXT("Ignoring double release of SourceId: %i"), SourceId);
-			return;
-		}
+		AUDIO_MIXER_CHECK(MixerSources[SourceId] != nullptr);
 
 		AUDIO_MIXER_DEBUG_LOG(SourceId, TEXT("Is releasing"));
 		
@@ -409,44 +387,39 @@ namespace Audio
 		}
 #endif
 		// Remove from list of active bus or source ids depending on what type of source this is
-		if (SourceInfo.AudioBusId != INDEX_NONE)
+		if (SourceInfo.BusId != INDEX_NONE)
 		{
 			// Remove this bus from the registry of bus instances
-			TSharedPtr<FMixerAudioBus> AudioBusPtr = AudioBuses.FindRef(SourceInfo.AudioBusId);
-			if (ensure(AudioBusPtr.IsValid()))
+			FMixerBus* Bus = Buses.Find(SourceInfo.BusId);
+			AUDIO_MIXER_CHECK(Bus);
+
+			// Remove this source from the list of bus instances.
+			if (Bus->RemoveInstanceId(SourceId))
 			{
-				// If this audio bus was automatically created via source bus playback, this this audio bus can be removed
-				if (AudioBusPtr->RemoveInstanceId(SourceId))
-				{
-					// Only automatic buses will be getting removed here. Otherwise they need to be manually removed from the source manager.
-					ensure(AudioBusPtr->IsAutomatic());
-					AudioBuses.Remove(SourceInfo.AudioBusId);
-				}
+				Buses.Remove(SourceInfo.BusId);
 			}
 		}
 
 		// Remove this source's send list from the bus data registry
-		for (int32 AudioBusSendType = 0; AudioBusSendType < (int32)EBusSendType::Count; ++AudioBusSendType)
+		for (int32 BusSendType = 0; BusSendType < (int32)EBusSendType::Count; ++BusSendType)
 		{
-			for (uint32 AudioBusId : SourceInfo.AudioBusSends[AudioBusSendType])
+			for (uint32 BusId : SourceInfo.BusSends[BusSendType])
 			{
 				// we should have a bus registration entry still since the send hasn't been cleaned up yet
-				TSharedPtr<FMixerAudioBus> AudioBusPtr = AudioBuses.FindRef(AudioBusId);
-				if (ensure(AudioBusPtr.IsValid()))
+				FMixerBus* Bus = Buses.Find(BusId);
+				AUDIO_MIXER_CHECK(Bus);
+
+				if (Bus->RemoveBusSend((EBusSendType)BusSendType, SourceId))
 				{
-					if (AudioBusPtr->RemoveSend((EBusSendType)AudioBusSendType, SourceId))
-					{
-						ensure(AudioBusPtr->IsAutomatic());
-						AudioBuses.Remove(AudioBusId);
-					}
+					Buses.Remove(BusId);
 				}
 			}
 
-			SourceInfo.AudioBusSends[AudioBusSendType].Reset();
+			SourceInfo.BusSends[BusSendType].Reset();
 		}
 
-		SourceInfo.AudioBusId = INDEX_NONE;
-		SourceInfo.SourceBusDurationFrames = INDEX_NONE;
+		SourceInfo.BusId = INDEX_NONE;
+		SourceInfo.BusDurationFrames = INDEX_NONE;
 
 		// Free the mixer source buffer data
 		if (SourceInfo.MixerSourceBuffer.IsValid())
@@ -472,7 +445,6 @@ namespace Audio
 		if (SourceInfo.bUseHRTFSpatializer)
 		{
 			AUDIO_MIXER_CHECK(bUsingSpatializationPlugin);
-			LLM_SCOPE(ELLMTag::AudioMixerPlugins);
 			SpatializationPlugin->OnReleaseSource(SourceId);
 		}
 
@@ -627,9 +599,9 @@ namespace Audio
 		return NumActiveSources;
 	}
 
-	int32 FMixerSourceManager::GetNumActiveAudioBuses() const
+	int32 FMixerSourceManager::GetNumActiveBuses() const
 	{
-		return AudioBuses.Num();
+		return Buses.Num();
 	}
 
 	void FMixerSourceManager::InitSource(const int32 SourceId, const FMixerSourceVoiceInitParams& InitParams)
@@ -698,7 +670,6 @@ namespace Audio
 			if (InitParams.bUseHRTFSpatialization)
 			{
 				AUDIO_MIXER_CHECK(bUsingSpatializationPlugin);
-				LLM_SCOPE(ELLMTag::AudioMixerPlugins);
 				SpatializationPlugin->OnInitSource(SourceId, InitParams.AudioComponentUserID, InitParams.SpatializationPluginSettings);
 			}
 
@@ -730,7 +701,6 @@ namespace Audio
 				InitData.SampleRate = MixerDevice->SampleRate;
 				InitData.NumSourceChannels = InitParams.NumInputChannels;
 				InitData.AudioClock = MixerDevice->GetAudioTime();
-				InitData.AudioDeviceId = MixerDevice->DeviceID;
 
 				SourceInfo.SourceEffectChainId = InitParams.SourceEffectChainId;
 				BuildSourceEffectChain(SourceId, InitData, InitParams.SourceEffectChain);
@@ -738,33 +708,34 @@ namespace Audio
 				// Whether or not to output to bus only
 				SourceInfo.bOutputToBusOnly = InitParams.bOutputToBusOnly;
 
-				// If this is going to be a source bus, add this source id to the list of active bus ids
-				if (InitParams.AudioBusId != INDEX_NONE)
+				// If this is a bus, add this source id to the list of active bus ids
+				if (InitParams.BusId != INDEX_NONE)
 				{
 					// Setting this BusId will flag this source as a bus. It doesn't try to generate 
 					// audio in the normal way but instead will render in a second stage, after normal source rendering.
-					SourceInfo.AudioBusId = InitParams.AudioBusId;
+					SourceInfo.BusId = InitParams.BusId;
 
-					// Source bus duration allows us to stop a bus after a given time
-					if (InitParams.SourceBusDuration != 0.0f)
+					// Bus duration allows us to stop a bus after a given time
+					if (InitParams.BusDuration != 0.0f)
 					{
-						SourceInfo.SourceBusDurationFrames = InitParams.SourceBusDuration * MixerDevice->GetSampleRate();
+						SourceInfo.BusDurationFrames = InitParams.BusDuration * MixerDevice->GetSampleRate();
 					}
 
 					// Register this bus as an instance
-					TSharedPtr<FMixerAudioBus> AudioBusPtr = AudioBuses.FindRef(SourceInfo.AudioBusId);
-					if (AudioBusPtr.IsValid())
+					FMixerBus* Bus = Buses.Find(InitParams.BusId);
+					if (Bus)
 					{
 						// If this bus is already registered, add this as a source id
-						AudioBusPtr->AddInstanceId(SourceId, InitParams.NumInputChannels);
+						Bus->AddInstanceId(SourceId, InitParams.NumInputChannels);
 					}
 					else
 					{
-						// If the bus is not registered, make a new entry. This will default to an automatic audio bus until explicitly made manual later.
-						TSharedPtr<FMixerAudioBus> NewAudioBus = TSharedPtr<FMixerAudioBus>(new FMixerAudioBus(this, true, InitParams.NumInputChannels));
-						NewAudioBus->AddInstanceId(SourceId, InitParams.NumInputChannels);
+						// If the bus is not registered, make a new entry
+						FMixerBus NewBusData(this, InitParams.NumInputChannels, NumOutputFrames);
 
-						AudioBuses.Add(InitParams.AudioBusId, NewAudioBus);
+						NewBusData.AddInstanceId(SourceId, InitParams.NumInputChannels);
+
+						Buses.Add(InitParams.BusId, NewBusData);
 					}
 				}
 
@@ -772,34 +743,34 @@ namespace Audio
 				// Note: buses can also send their audio to other buses.
 				for (int32 BusSendType = 0; BusSendType < (int32)EBusSendType::Count; ++BusSendType)
 				{
-					for (const FInitAudioBusSend& AudioBusSend : InitParams.AudioBusSends[BusSendType])
+					for (const FMixerBusSend& BusSend : InitParams.BusSends[BusSendType])
 					{
 						// New struct to map which source (SourceId) is sending to the bus
-						FAudioBusSend NewAudioBusSend;
-						NewAudioBusSend.SourceId = SourceId;
-						NewAudioBusSend.SendLevel = AudioBusSend.SendLevel;
+						FBusSend NewBusSend;
+						NewBusSend.SourceId = SourceId;
+						NewBusSend.SendLevel = BusSend.SendLevel;
 
 						// Get existing BusId and add the send, or create new bus registration
-						TSharedPtr<FMixerAudioBus> AudioBusPtr = AudioBuses.FindRef(AudioBusSend.AudioBusId);
-						if (AudioBusPtr.IsValid())
+						FMixerBus* Bus = Buses.Find(BusSend.BusId);
+						if (Bus)
 						{
-							AudioBusPtr->AddSend((EBusSendType)BusSendType, NewAudioBusSend);
+							Bus->AddBusSend((EBusSendType)BusSendType, NewBusSend);
 						}
 						else
 						{
-							// If the bus is not registered, make a new entry. This will default to an automatic audio bus until explicitly made manual later.
-							TSharedPtr<FMixerAudioBus> NewAudioBus(new FMixerAudioBus(this, true, InitParams.NumInputChannels));
+							// If the bus is not registered, make a new entry
+							FMixerBus NewBusData(this, InitParams.NumInputChannels, NumOutputFrames);
 
 							// Add a send to it. This will not have a bus instance id (i.e. won't output audio), but 
 							// we register the send anyway in the event that this bus does play, we'll know to send this
 							// source's audio to it.
-							NewAudioBus->AddSend((EBusSendType)BusSendType, NewAudioBusSend);
+							NewBusData.AddBusSend((EBusSendType)BusSendType, NewBusSend);
 
-							AudioBuses.Add(AudioBusSend.AudioBusId, NewAudioBus);
+							Buses.Add(BusSend.BusId, NewBusData);
 						}
 
 						// Store on this source, which buses its sending its audio to
-						SourceInfo.AudioBusSends[BusSendType].Add(AudioBusSend.AudioBusId);
+						SourceInfo.BusSends[BusSendType].Add(BusSend.BusId);
 					}
 				}
 			}
@@ -921,97 +892,6 @@ namespace Audio
 		});
 	}
 
-	void FMixerSourceManager::StartAudioBus(uint32 InAudioBusId, int32 InNumChannels, bool bInIsAutomatic)
-	{
-		if (AudioBusIds_AudioThread.Contains(InAudioBusId))
-		{
-			return;
-		}
-
-		AudioBusIds_AudioThread.Add(InAudioBusId);
-
-		AudioMixerThreadCommand([this, InAudioBusId, InNumChannels, bInIsAutomatic]()
-		{
-			// If this audio bus id already exists, set it to not be automatic and return it
-			TSharedPtr<FMixerAudioBus> AudioBusPtr = AudioBuses.FindRef(InAudioBusId);
-			if (AudioBusPtr.IsValid())
-			{
-				// If this audio bus already existed, make sure the num channels lines up
-				ensure(AudioBusPtr->GetNumChannels() == InNumChannels);
-				AudioBusPtr->SetAutomatic(bInIsAutomatic);
-			}
-			else
-			{
-				// If the bus is not registered, make a new entry.
-				TSharedPtr<FMixerAudioBus> NewBusData(new FMixerAudioBus(this, bInIsAutomatic, InNumChannels));
-
-				AudioBuses.Add(InAudioBusId, NewBusData);
-			}
-
-			//  Now add any existing playing sources to this audio bus as sends if they exist
-			for (FSourceInfo& SourceInfo : SourceInfos)
-			{
-				if (SourceInfo.AudioBusId == InAudioBusId)
-				{
-					SourceInfo.bIsPlaying = false;
-					SourceInfo.bIsPaused = false;
-					SourceInfo.bIsActive = false;
-					SourceInfo.bIsStopping = false;
-				}
-			}
-		});
-	}
-
-	void FMixerSourceManager::StopAudioBus(uint32 InAudioBusId)
-	{
-		if (!AudioBusIds_AudioThread.Contains(InAudioBusId))
-		{
-			return;
-		}
-
-		AudioBusIds_AudioThread.Remove(InAudioBusId);
-
-		AudioMixerThreadCommand([this, InAudioBusId]()
-		{
-			TSharedPtr<FMixerAudioBus>* AudioBusPtr = AudioBuses.Find(InAudioBusId);
-			if (AudioBusPtr)
-			{
-				if (!(*AudioBusPtr)->IsAutomatic())
-				{
-					// Immediately stop all sources which were source buses
-					for (FSourceInfo& SourceInfo : SourceInfos)
-					{
-						if (SourceInfo.AudioBusId == InAudioBusId)
-						{
-							SourceInfo.bIsPlaying = false;
-							SourceInfo.bIsPaused = false;
-							SourceInfo.bIsActive = false;
-							SourceInfo.bIsStopping = false;
-						}
-					}
-					AudioBuses.Remove(InAudioBusId);
-				}
-			}
-		});
-	}
-
-	bool FMixerSourceManager::IsAudioBusActive(uint32 InAudioBusId)
-	{
-		AUDIO_MIXER_CHECK_GAME_THREAD(MixerDevice);
-		return AudioBusIds_AudioThread.Contains(InAudioBusId);
-	}
-
-	FPatchOutputStrongPtr FMixerSourceManager::AddPatchForAudioBus(uint32 InAudioBusId, float PatchGain)
-	{
-		AUDIO_MIXER_CHECK_AUDIO_PLAT_THREAD(MixerDevice);
-		TSharedPtr<FMixerAudioBus>* AudioBusPtr = AudioBuses.Find(InAudioBusId);
-		if (AudioBusPtr)
-		{
-			return (*AudioBusPtr)->AddNewPatch(NumOutputFrames * 2, PatchGain);
-		}
-		return nullptr;
-	}
-
 	void FMixerSourceManager::Play(const int32 SourceId)
 	{
 		AUDIO_MIXER_CHECK(SourceId < NumTotalSources);
@@ -1334,49 +1214,45 @@ namespace Audio
 		});
 	}
 
-	void FMixerSourceManager::SetBusSendInfo(const int32 SourceId, EBusSendType InAudioBusSendType, uint32 AudioBusId, float BusSendLevel)
+	void FMixerSourceManager::SetBusSendInfo(const int32 SourceId, EBusSendType InBusSendType, FMixerBusSend& InBusSend)
 	{
 		AUDIO_MIXER_CHECK(SourceId < NumTotalSources);
 		AUDIO_MIXER_CHECK(GameThreadInfo.bIsBusy[SourceId]);
 		AUDIO_MIXER_CHECK_GAME_THREAD(MixerDevice);
 
-		AudioMixerThreadCommand([this, SourceId, InAudioBusSendType, AudioBusId, BusSendLevel]()
+		AudioMixerThreadCommand([this, SourceId, InBusSendType, InBusSend]()
 		{
 			// Create mapping of source id to bus send level
-			FAudioBusSend BusSend;
-			BusSend.SendLevel = BusSendLevel;
-			BusSend.SourceId = SourceId;
-			BusSend.AudioBusId = AudioBusId;
-
+			FBusSend BusSend{ SourceId, InBusSend.SendLevel };
 			FSourceInfo& SourceInfo = SourceInfos[SourceId];
 
 			// Retrieve the bus we want to send audio to
-			TSharedPtr<FMixerAudioBus>* AudioBusPtr = AudioBuses.Find(AudioBusId);
+			FMixerBus* Bus = Buses.Find(InBusSend.BusId);
 
 			// If we already have a bus, we update the amount of audio we want to send to it
-			if (AudioBusPtr)
+			if (Bus)
 			{
-				(*AudioBusPtr)->AddSend(InAudioBusSendType, BusSend);
+				Bus->AddBusSend(InBusSendType, BusSend);			
 			}
 			else
 			{
 				// If the bus is not registered, make a new entry on the send
-				TSharedPtr<FMixerAudioBus> NewBusData(new FMixerAudioBus(this, true, SourceInfo.NumInputChannels));
+				FMixerBus NewBusData(this, SourceInfo.NumInputChannels, NumOutputFrames);
 
 				// Add a send to it. This will not have a bus instance id (i.e. won't output audio), but 
 				// we register the send anyway in the event that this bus does play, we'll know to send this
 				// source's audio to it.
-				NewBusData->AddSend(InAudioBusSendType, BusSend);
+				NewBusData.AddBusSend(InBusSendType, BusSend);
 
-				AudioBuses.Add(AudioBusId, NewBusData);
+				Buses.Add(InBusSend.BusId, NewBusData);
 			}
 
 			// Check to see if we need to create new bus data. If we are not playing a bus with this id, then we
 			// need to create a slot for it such that when a bus does play, it'll start rendering audio from this source
 			bool bExisted = false;
-			for (uint32 BusId : SourceInfo.AudioBusSends[(int32)InAudioBusSendType])
+			for (uint32 BusId : SourceInfo.BusSends[(int32)InBusSendType])
 			{
-				if (BusId == AudioBusId)
+				if (BusId == InBusSend.BusId)
 				{
 					bExisted = true;
 					break;
@@ -1385,7 +1261,7 @@ namespace Audio
 
 			if (!bExisted)
 			{
-				SourceInfo.AudioBusSends[(int32)InAudioBusSendType].Add(AudioBusId);
+				SourceInfo.BusSends[(int32)InBusSendType].Add(InBusSend.BusId);
 			}
 		});
 	}
@@ -1558,8 +1434,8 @@ namespace Audio
 				continue;
 			}
 
-			const bool bIsSourceBus = SourceInfo.AudioBusId != INDEX_NONE;
-			if ((bGenerateBuses && !bIsSourceBus) || (!bGenerateBuses && bIsSourceBus))
+			const bool bIsBus = SourceInfo.BusId != INDEX_NONE;
+			if ((bGenerateBuses && !bIsBus) || (!bGenerateBuses && bIsBus))
 			{
 				continue;
 			}
@@ -1581,30 +1457,27 @@ namespace Audio
 
 			// if this is a bus, we just want to copy the bus audio to this source's output audio
 			// Note we need to copy this since bus instances may have different audio via dynamic source effects, etc.
-			if (bIsSourceBus)
+			if (bIsBus)
 			{
-				// Get the source's rendered and mixed audio bus data
-				const TSharedPtr<FMixerAudioBus>* AudioBusPtr = AudioBuses.Find(SourceInfo.AudioBusId);
-				if (ensure(AudioBusPtr != nullptr))
-				{
-					const float* RESTRICT BusBufferPtr = (*AudioBusPtr)->GetCurrentBusBuffer();
+				// Get the source's rendered bus data
+				const FMixerBus* Bus = Buses.Find(SourceInfo.BusId);
+				const float* RESTRICT BusBufferPtr = Bus->GetCurrentBusBuffer();
 
-					int32 NumFramesPlayed = NumOutputFrames;
-					if (SourceInfo.SourceBusDurationFrames != INDEX_NONE)
+				int32 NumFramesPlayed = NumOutputFrames;
+				if (SourceInfo.BusDurationFrames != INDEX_NONE)
+				{
+					// If we're now finishing, only copy over the real data
+					if ((SourceInfo.NumFramesPlayed + NumOutputFrames) >= SourceInfo.BusDurationFrames)
 					{
-						// If we're now finishing, only copy over the real data
-						if ((SourceInfo.NumFramesPlayed + NumOutputFrames) >= SourceInfo.SourceBusDurationFrames)
-						{
-							NumFramesPlayed = SourceInfo.SourceBusDurationFrames - SourceInfo.NumFramesPlayed;
-							SourceInfo.bIsLastBuffer = true;
-						}
+						NumFramesPlayed = SourceInfo.BusDurationFrames - SourceInfo.NumFramesPlayed;
+						SourceInfo.bIsLastBuffer = true;				
 					}
+				}
 
-					SourceInfo.NumFramesPlayed += NumFramesPlayed;
+				SourceInfo.NumFramesPlayed += NumFramesPlayed;
 
-					// Simply copy into the pre distance attenuation buffer ptr
-					FMemory::Memcpy(PreDistanceAttenBufferPtr, BusBufferPtr, sizeof(float) * NumFramesPlayed * SourceInfo.NumInputChannels);
-				}
+				// Simply copy into the pre distance attenuation buffer ptr
+				FMemory::Memcpy(PreDistanceAttenBufferPtr, BusBufferPtr, sizeof(float) * NumFramesPlayed * SourceInfo.NumInputChannels);
 			}
 			else
 			{
@@ -1698,20 +1571,20 @@ namespace Audio
 	void FMixerSourceManager::ComputeBuses()
 	{
 		// Loop through the bus registry and mix source audio
-		for (auto& Entry : AudioBuses)
+		for (auto& Entry : Buses)
 		{
-			TSharedPtr<FMixerAudioBus>& AudioBus = Entry.Value;
-			AudioBus->MixBuffer();
+			FMixerBus& Bus = Entry.Value;
+			Bus.MixBuffer();
 		}
 	}
 
 	void FMixerSourceManager::UpdateBuses()
 	{
 		// Update the bus states post mixing. This flips the current/previous buffer indices.
-		for (auto& Entry : AudioBuses)
+		for (auto& Entry : Buses)
 		{
-			TSharedPtr<FMixerAudioBus>& AudioBus = Entry.Value;
-			AudioBus->Update();
+			FMixerBus& Bus = Entry.Value;
+			Bus.Update();
 		}
 	}
 
@@ -1831,10 +1704,7 @@ namespace Audio
 				SourceInfo.AudioPluginOutputData.AudioBuffer.AddZeroed(2 * NumOutputFrames);
 			}
 
-			{
-				LLM_SCOPE(ELLMTag::AudioMixerPlugins);
-				SpatializationPlugin->ProcessAudio(AudioPluginInputData, SourceInfo.AudioPluginOutputData);
-			}
+			SpatializationPlugin->ProcessAudio(AudioPluginInputData, SourceInfo.AudioPluginOutputData);
 
 			// If this is an external send, we treat this source audio as if it was still a mono source
 			// This will allow it to traditionally pan in the ComputeOutputBuffers function and be
@@ -2169,8 +2039,8 @@ namespace Audio
 				continue;
 			}
 
-			const bool bIsSourceBus = SourceInfo.AudioBusId != INDEX_NONE;
-			if ((bGenerateBuses && !bIsSourceBus) || (!bGenerateBuses && bIsSourceBus))
+			const bool bIsBus = SourceInfo.BusId != INDEX_NONE;
+			if ((bGenerateBuses && !bIsBus) || (!bGenerateBuses && bIsBus))
 			{
 				continue;
 			}
@@ -2179,7 +2049,7 @@ namespace Audio
 
 			// Copy and store the current state of the pre-distance attenuation buffer before we feed it through our source effects
 			// This is used by pre-effect sends
-			if (SourceInfo.AudioBusSends[(int32)EBusSendType::PreEffect].Num() > 0)
+			if (SourceInfo.BusSends[(int32)EBusSendType::PreEffect].Num() > 0)
 			{
 				SourceInfo.PreEffectBuffer.Reset();
 				SourceInfo.PreEffectBuffer.Reserve(SourceInfo.PreDistanceAttenuationBuffer.Num());
@@ -2368,8 +2238,8 @@ namespace Audio
 
 			// If we're in generate buses mode and not a bus, or vice versa, or if we're set to only output audio to buses.
 			// If set to output buses, no need to do any panning for the source. The buses will do the panning.
-			const bool bIsSourceBus = SourceInfo.AudioBusId != INDEX_NONE;
-			if ((bGenerateBuses && !bIsSourceBus) || (!bGenerateBuses && bIsSourceBus) || SourceInfo.bOutputToBusOnly)
+			const bool bIsBus = SourceInfo.BusId != INDEX_NONE;
+			if ((bGenerateBuses && !bIsBus) || (!bGenerateBuses && bIsBus) || SourceInfo.bOutputToBusOnly)
 			{
 				continue;
 			}
@@ -2422,7 +2292,7 @@ namespace Audio
 	void FMixerSourceManager::GenerateSourceAudio(const bool bGenerateBuses)
 	{
 		// If there are no buses, don't need to do anything here
-		if (bGenerateBuses && !AudioBuses.Num())
+		if (bGenerateBuses && !Buses.Num())
 		{
 			return;
 		}
@@ -2548,7 +2418,6 @@ namespace Audio
 			FSoundEffectSourceInitData InitData;
 			InitData.AudioClock = MixerDevice->GetAudioClock();
 			InitData.SampleRate = MixerDevice->SampleRate;
-			InitData.AudioDeviceId = MixerDevice->DeviceID;
 
 			for (int32 SourceId = 0; SourceId < NumTotalSources; ++SourceId)
 			{
@@ -2612,23 +2481,11 @@ namespace Audio
 		return SourceInfos[SourceId].PreEffectBuffer.GetData();
 	}
 
-	const float* FMixerSourceManager::GetPreviousSourceBusBuffer(const int32 SourceId) const
+	const float* FMixerSourceManager::GetPreviousBusBuffer(const int32 SourceId) const
 	{
-		if (SourceId < SourceInfos.Num())
-		{
-			return GetPreviousAudioBusBuffer(SourceInfos[SourceId].AudioBusId);
-		}
-		return nullptr;
-	}
-
-	const float* FMixerSourceManager::GetPreviousAudioBusBuffer(const int32 AudioBusId) const
-	{
-		const TSharedPtr<FMixerAudioBus>* AudioBusPtr = AudioBuses.Find(AudioBusId);
-		if (AudioBusPtr)
-		{
-			return (*AudioBusPtr)->GetPreviousBusBuffer();
-		}
-		return nullptr;
+		const uint32 BusId = SourceInfos[SourceId].BusId;
+		const FMixerBus* MixerBus = Buses.Find(BusId);
+		return MixerBus->GetPreviousBusBuffer();
 	}
 
 	int32 FMixerSourceManager::GetNumChannels(const int32 SourceId) const
@@ -2636,9 +2493,9 @@ namespace Audio
 		return SourceInfos[SourceId].NumInputChannels;
 	}
 
-	bool FMixerSourceManager::IsSourceBus(const int32 SourceId) const
+	bool FMixerSourceManager::IsBus(const int32 SourceId) const
 	{
-		return SourceInfos[SourceId].AudioBusId != INDEX_NONE;
+		return SourceInfos[SourceId].BusId != INDEX_NONE;
 	}
 
 	void FMixerSourceManager::ComputeNextBlockOfSamples()
@@ -2677,7 +2534,6 @@ namespace Audio
 		if (bUsingSpatializationPlugin)
 		{
 			AUDIO_MIXER_CHECK(SpatializationPlugin.IsValid());
-			LLM_SCOPE(ELLMTag::AudioMixerPlugins);
 			SpatializationPlugin->OnAllSourcesProcessed();
 		}
 
@@ -2725,20 +2581,6 @@ namespace Audio
 
 		// Add the function to the command queue:
 		int32 AudioThreadCommandIndex = !RenderThreadCommandBufferIndex.GetValue();
-		
-#if !NO_LOGGING
-		static uint32 WarnSize = 1024 * 1024;
-		SIZE_T Size = CommandBuffers[AudioThreadCommandIndex].SourceCommandQueue.GetAllocatedSize();
-		if (Size > WarnSize )
-		{		
-			SIZE_T Num = CommandBuffers[AudioThreadCommandIndex].SourceCommandQueue.Num();
-			// NOTE: Although not really and error we want this to show up in shipping builds.
-			UE_LOG(LogAudioMixer, Error, TEXT("Command Queue has grown to %uk bytes, containing %d cmds, last pump was %fms ago."), 
-				Size >> 10, Num, FPlatformTime::ToMilliseconds64(FPlatformTime::Cycles64() - LastPumpTimeInCycles));
-			WarnSize *= 2;
-		}
-#endif //!NO_LOGGING
-
 		CommandBuffers[AudioThreadCommandIndex].SourceCommandQueue.Add(MoveTemp(InFunction));
 		NumCommands.Increment();
 	}
@@ -2766,7 +2608,6 @@ namespace Audio
 			NumCommands.Decrement();
 		}
 
-		LastPumpTimeInCycles = FPlatformTime::Cycles64();
 		Commands.SourceCommandQueue.Reset();
 
 		if (FPlatformProcess::SupportsMultithreading())
diff --git a/Engine/Source/Runtime/AudioMixer/Private/AudioMixerSourceManager.h b/Engine/Source/Runtime/AudioMixer/Private/AudioMixerSourceManager.h
--- a/Engine/Source/Runtime/AudioMixer/Private/AudioMixerSourceManager.h
+++ b/Engine/Source/Runtime/AudioMixer/Private/AudioMixerSourceManager.h
@@ -35,10 +35,15 @@ namespace Audio
 		AlignedFloatBuffer AudioData;
 
 		/** How many times this buffer will loop. */
-		int32 LoopCount = 0;
+		int32 LoopCount;
 
 		/** If this buffer is from real-time decoding and needs to make callbacks for more data. */
-		bool bRealTimeBuffer = false;
+		uint32 bRealTimeBuffer : 1;
+
+		FMixerSourceVoiceBuffer()
+		{
+			FMemory::Memzero(this, sizeof(*this));
+		}
 	};
 
 	typedef TSharedPtr<FMixerSubmix, ESPMode::ThreadSafe> FMixerSubmixPtr;
@@ -67,53 +72,79 @@ namespace Audio
 		FMixerSubmixWeakPtr Submix;
 
 		// The amount of audio that is to be mixed into this submix
-		float SendLevel = 0.0f;
+		float SendLevel;
 
 		// Whather or not this is the primary send (i.e. first in the send chain)
-		bool bIsMainSend = false;
+		bool bIsMainSend;
 
 		// If this is a soundfield submix, this is a pointer to the submix's Soundfield Factory.
 		// If this is nullptr, the submix is not a soundfield submix.
-		ISoundfieldFactory* SoundfieldFactory = nullptr;
+		ISoundfieldFactory* SoundfieldFactory;
 	};
 
 	// Struct holding mappings of bus ids (unique ids) to send level
-	struct FInitAudioBusSend
+	struct FMixerBusSend
 	{
-		uint32 AudioBusId = INDEX_NONE;
-		float SendLevel = 0.0f;
+		uint32 BusId;
+		float SendLevel;
 	};
 
 	struct FMixerSourceVoiceInitParams
 	{
-		TSharedPtr<FMixerSourceBuffer, ESPMode::ThreadSafe> MixerSourceBuffer = nullptr;
-		ISourceListener* SourceListener = nullptr;
+		TSharedPtr<FMixerSourceBuffer> MixerSourceBuffer;
+		ISourceListener* SourceListener;
 		TArray<FMixerSourceSubmixSend> SubmixSends;
-		TArray<FInitAudioBusSend> AudioBusSends[(int32)EBusSendType::Count];
-		uint32 AudioBusId = INDEX_NONE;
-		float SourceBusDuration = 0.0f;
-		uint32 SourceEffectChainId = INDEX_NONE;
+		TArray<FMixerBusSend> BusSends[(int32)EBusSendType::Count];
+		uint32 BusId;
+		float BusDuration;
+		uint32 SourceEffectChainId;
 		TArray<FSourceEffectChainEntry> SourceEffectChain;
-		FMixerSourceVoice* SourceVoice = nullptr;
-		int32 NumInputChannels = 0;
-		int32 NumInputFrames = 0;
-		float EnvelopeFollowerAttackTime = 10.0f;
-		float EnvelopeFollowerReleaseTime = 100.0f;
+		FMixerSourceVoice* SourceVoice;
+		int32 NumInputChannels;
+		int32 NumInputFrames;
+		float EnvelopeFollowerAttackTime;
+		float EnvelopeFollowerReleaseTime;
 		FString DebugName;
-		USpatializationPluginSourceSettingsBase* SpatializationPluginSettings = nullptr;
-		UOcclusionPluginSourceSettingsBase* OcclusionPluginSettings = nullptr;
-		UReverbPluginSourceSettingsBase* ReverbPluginSettings = nullptr;
-		USoundModulationPluginSourceSettingsBase* ModulationPluginSettings = nullptr;
+		USpatializationPluginSourceSettingsBase* SpatializationPluginSettings;
+		UOcclusionPluginSourceSettingsBase* OcclusionPluginSettings;
+		UReverbPluginSourceSettingsBase* ReverbPluginSettings;
+		USoundModulationPluginSourceSettingsBase* ModulationPluginSettings;
 		FName AudioComponentUserID;
-		uint64 AudioComponentID = 0;
-		bool bPlayEffectChainTails = false;
-		bool bUseHRTFSpatialization = false;
-		bool bIsExternalSend = false;
-		bool bIsDebugMode  = false;
-		bool bOutputToBusOnly = false;
-		bool bIsVorbis = false;
-		bool bIsAmbisonics = false;
-		bool bIsSeeking = false;
+		uint64 AudioComponentID;
+		uint8 bPlayEffectChainTails : 1;
+		uint8 bUseHRTFSpatialization : 1;
+		uint8 bIsExternalSend : 1;
+		uint8 bIsDebugMode : 1;
+		uint8 bOutputToBusOnly : 1;
+		uint8 bIsVorbis : 1;
+		uint8 bIsAmbisonics : 1;
+		uint8 bIsSeeking : 1;
+
+		FMixerSourceVoiceInitParams()
+			: MixerSourceBuffer(nullptr)
+			, SourceListener(nullptr)
+			, BusId(INDEX_NONE)
+			, BusDuration(0.0f)
+			, SourceEffectChainId(INDEX_NONE)
+			, SourceVoice(nullptr)
+			, NumInputChannels(0)
+			, NumInputFrames(0)
+			, EnvelopeFollowerAttackTime(10.0f)
+			, EnvelopeFollowerReleaseTime(100.0f)
+			, SpatializationPluginSettings(nullptr)
+			, OcclusionPluginSettings(nullptr)
+			, ReverbPluginSettings(nullptr)
+			, ModulationPluginSettings(nullptr)
+			, AudioComponentID(0)
+			, bPlayEffectChainTails(false)
+			, bUseHRTFSpatialization(false)
+			, bIsExternalSend(false)
+			, bIsDebugMode(false)
+			, bOutputToBusOnly(false)
+			, bIsVorbis(false)
+			, bIsAmbisonics(false)
+			, bIsSeeking(false)
+		{}
 	};
 
 	struct FSourceChannelMap
@@ -123,8 +154,8 @@ namespace Audio
 
 		// This is the number of bytes the gain array is using:
 		// (Number of input channels * number of output channels) * sizeof float.
-		int32 CopySize = 0;
-		bool bIsInit = false;
+		int32 CopySize;
+		bool bIsInit;
 
 		FSourceChannelMap(int32 InNumInChannels, int32 InNumOutChannels) 
 			: CopySize(InNumInChannels * InNumOutChannels * sizeof(float))
@@ -171,10 +202,10 @@ namespace Audio
 	struct FSourceManagerInitParams
 	{
 		// Total number of sources to use in the source manager
-		int32 NumSources = 0;
+		int32 NumSources;
 
 		// Number of worker threads to use for the source manager.
-		int32 NumSourceWorkers = 0;
+		int32 NumSourceWorkers;
 	};
 
 	class FMixerSourceManager
@@ -188,17 +219,11 @@ namespace Audio
 
 		bool GetFreeSourceId(int32& OutSourceId);
 		int32 GetNumActiveSources() const;
-		int32 GetNumActiveAudioBuses() const;
+		int32 GetNumActiveBuses() const;
 
 		void ReleaseSourceId(const int32 SourceId);
 		void InitSource(const int32 SourceId, const FMixerSourceVoiceInitParams& InitParams);
 
-		// Creates an audio bus manually. Returns an audio bus Id.
-		void StartAudioBus(uint32 InAudioBusId, int32 InNumChannels, bool bInIsAutomatic);
-		void StopAudioBus(uint32 InAudioBusId);
-		bool IsAudioBusActive(uint32 InAudioBusId);
-		FPatchOutputStrongPtr AddPatchForAudioBus(uint32 InAudioBusId, float PatchGain);
-
 		void Play(const int32 SourceId);
 		void Stop(const int32 SourceId);
 		void StopFade(const int32 SourceId, const int32 NumFrames);
@@ -231,8 +256,7 @@ namespace Audio
 		const FQuat GetListenerRotation(const int32 SourceId) const;
 
 		void SetSubmixSendInfo(const int32 SourceId, const FMixerSourceSubmixSend& SubmixSend);
-
-		void SetBusSendInfo(const int32 SourceId, EBusSendType InAudioBusSendType, uint32 AudiobusId, float BusSendLevel);
+		void SetBusSendInfo(const int32 SourceId, EBusSendType InBusSendType, FMixerBusSend& BusSend);
 
 		void UpdateDeviceChannelCount(const int32 InNumOutputChannels);
 
@@ -241,11 +265,10 @@ namespace Audio
 
 		const float* GetPreDistanceAttenuationBuffer(const int32 SourceId) const;
 		const float* GetPreEffectBuffer(const int32 SourceId) const;
-		const float* GetPreviousSourceBusBuffer(const int32 SourceId) const;
-		const float* GetPreviousAudioBusBuffer(const int32 AudioBusId) const;
+		const float* GetPreviousBusBuffer(const int32 SourceId) const;
 		int32 GetNumChannels(const int32 SourceId) const;
 		int32 GetNumOutputFrames() const { return NumOutputFrames; }
-		bool IsSourceBus(const int32 SourceId) const;
+		bool IsBus(const int32 SourceId) const;
 		void PumpCommandQueue();
 		void UpdatePendingReleaseData(bool bForceWait = false);
 		void FlushCommandQueue(bool bPumpCommandQueue = false);
@@ -428,11 +451,11 @@ namespace Audio
 			~FSourceInfo() {}
 
 			// Object which handles source buffer decoding
-			TSharedPtr<FMixerSourceBuffer, ESPMode::ThreadSafe> MixerSourceBuffer;
+			TSharedPtr<FMixerSourceBuffer> MixerSourceBuffer;
 			ISourceListener* SourceListener;
 
 			// Data used for rendering sources
-			TSharedPtr<FMixerSourceVoiceBuffer, ESPMode::ThreadSafe> CurrentPCMBuffer;
+			TSharedPtr<FMixerSourceVoiceBuffer> CurrentPCMBuffer;
 			int32 CurrentAudioChunkNumFrames;
 
 			// The post-attenuation source buffer, used to send audio to submixes
@@ -452,14 +475,14 @@ namespace Audio
 
 			TArray<FMixerSourceSubmixSend> SubmixSends;
 
-			// What audio bus Id this source is sonfiying, if it is a source bus. This is INDEX_NONE for sources which are not source buses.
-			uint32 AudioBusId;
+			// What bus Id this source is, if it is a bus. This is INDEX_NONE for sources which are not buses.
+			uint32 BusId;
 
-			// Number of samples to count for source bus
-			int64 SourceBusDurationFrames;
+			// Number of samples to count for bus
+			int64 BusDurationFrames;
 
 			// What buses this source is sending its audio to. Used to remove this source from the bus send list.
-			TArray<uint32> AudioBusSends[(int32)EBusSendType::Count];
+			TArray<uint32> BusSends[(int32)EBusSendType::Count];
 
 			// Interpolated source params
 			FParam PitchSourceParam;
@@ -562,15 +585,14 @@ namespace Audio
 		// sources.
 		TArray<FSourceDownmixData> DownmixDataArray;
 
-		// Map of bus object Id's to audio bus data. 
-		TMap<uint32, TSharedPtr<FMixerAudioBus>> AudioBuses;
-		TArray<uint32> AudioBusIds_AudioThread;
+		// Map of bus object Id's to bus data. 
+		TMap<uint32, FMixerBus> Buses;
 
 		// Async task workers for processing sources in parallel
 		TArray<FAsyncTask<FAudioMixerSourceWorker>*> SourceWorkers;
 
 		// Array of task data waiting to finished. Processed on audio render thread.
-		TArray<TSharedPtr<FMixerSourceBuffer, ESPMode::ThreadSafe>> PendingSourceBuffers;
+		TArray<TSharedPtr<FMixerSourceBuffer>> PendingSourceBuffers;
 
 		// General information about sources in source manager accessible from game thread
 		struct FGameThreadInfo
@@ -597,7 +619,6 @@ namespace Audio
 
 		// Set to true when the audio source manager should pump the command queue
 		FThreadSafeBool bPumpQueue;
-		uint64 LastPumpTimeInCycles = 0;
 
 		friend class FMixerSourceVoice;
 	};
diff --git a/Engine/Source/Runtime/AudioMixer/Private/Components/SynthComponent.cpp b/Engine/Source/Runtime/AudioMixer/Private/Components/SynthComponent.cpp
--- a/Engine/Source/Runtime/AudioMixer/Private/Components/SynthComponent.cpp
+++ b/Engine/Source/Runtime/AudioMixer/Private/Components/SynthComponent.cpp
@@ -219,7 +219,7 @@ void USynthComponent::Initialize(int32 SampleRateOverride)
 
 		if (!Synth)
 		{
-			Synth = NewObject<USynthSound>(this, FName(TEXT("TimeSynth"), this->GetUniqueID()));
+			Synth = NewObject<USynthSound>(this, TEXT("Synth"));
 		}
 
 		// Copy sound base data to the sound
