Added stat scopes for FAsyncDecodeWorker.

#rnx

diff --git a/Engine/Source/Runtime/AudioMixer/Private/AudioMixerSourceDecode.cpp b/Engine/Source/Runtime/AudioMixer/Private/AudioMixerSourceDecode.cpp
--- a/Engine/Source/Runtime/AudioMixer/Private/AudioMixerSourceDecode.cpp
+++ b/Engine/Source/Runtime/AudioMixer/Private/AudioMixerSourceDecode.cpp
@@ -55,6 +55,7 @@ public:
 		{
 			case EAudioTaskType::Procedural:
 			{
+				QUICK_SCOPE_CYCLE_COUNTER(STAT_FAsyncDecodeWorker_Procedural);
 				// Make sure we've been flagged as active
 				if (!ProceduralTaskData.ProceduralSoundWave || !ProceduralTaskData.ProceduralSoundWave->IsGeneratingAudio())
 				{
@@ -99,12 +100,14 @@ public:
 
 			case EAudioTaskType::Header:
 			{
+				QUICK_SCOPE_CYCLE_COUNTER(STAT_FAsyncDecodeWorker_Header);
 				HeaderParseAudioData.MixerBuffer->ReadCompressedInfo(HeaderParseAudioData.SoundWave);
 			}
 			break;
 
 			case EAudioTaskType::Decode:
 			{
+				QUICK_SCOPE_CYCLE_COUNTER(STAT_FAsyncDecodeWorker_Decode);
 				int32 NumChannels = DecodeTaskData.NumChannels;
 				int32 ByteSize = NumChannels * DecodeTaskData.NumFramesToDecode * sizeof(int16);
 
