Fixing envelope follower delegate.

diff --git a/Engine/Source/Runtime/AudioMixer/Private/AudioMixerSubmix.cpp b/Engine/Source/Runtime/AudioMixer/Private/AudioMixerSubmix.cpp
--- a/Engine/Source/Runtime/AudioMixer/Private/AudioMixerSubmix.cpp
+++ b/Engine/Source/Runtime/AudioMixer/Private/AudioMixerSubmix.cpp
@@ -986,8 +986,8 @@ namespace Audio
 		// Perform any envelope following if we're told to do so
 		if (bIsEnvelopeFollowing)
 		{
-			const int32 OutBufferSamples = OutAudioBuffer.Num();
-			const float* OutAudioBufferPtr = OutAudioBuffer.GetData();
+			const int32 BufferSamples = InputBuffer.Num();
+			const float* AudioBufferPtr = InputBuffer.GetData();
 
 
 			// Perform envelope following per channel
@@ -1000,9 +1000,9 @@ namespace Audio
 				FEnvelopeFollower& EnvFollower = EnvelopeFollowers[ChannelIndex];
 
 				// Track the last sample
-				for (int32 SampleIndex = ChannelIndex; SampleIndex < OutBufferSamples; SampleIndex += NumChannels)
+				for (int32 SampleIndex = ChannelIndex; SampleIndex < BufferSamples; SampleIndex += NumChannels)
 				{
-					const float SampleValue = OutAudioBufferPtr[SampleIndex];
+					const float SampleValue = AudioBufferPtr[SampleIndex];
 					EnvFollower.ProcessAudio(SampleValue);
 				}
 
