Fix submix issue

diff --git a/Engine/Source/Runtime/AudioMixer/Private/AudioMixerSubmix.cpp b/Engine/Source/Runtime/AudioMixer/Private/AudioMixerSubmix.cpp
--- a/Engine/Source/Runtime/AudioMixer/Private/AudioMixerSubmix.cpp
+++ b/Engine/Source/Runtime/AudioMixer/Private/AudioMixerSubmix.cpp
@@ -1012,19 +1012,49 @@ namespace Audio
 			EnvelopeNumChannels = NumChannels;
 		}
 
-		// Now apply the output volume
-		if (!FMath::IsNearlyEqual(TargetOutputVolume, CurrentOutputVolume) || !FMath::IsNearlyEqual(CurrentOutputVolume, 1.0f))
+		if (!FMath::IsNearlyEqual(TargetWetLevel, CurrentWetLevel) || !FMath::IsNearlyEqual(CurrentWetLevel, 1.0f))
 		{
-			// If we've already set the output volume, only need to multiply by constant
-			if (FMath::IsNearlyEqual(TargetOutputVolume, CurrentOutputVolume))
+			// If we've already set the volume, only need to multiply by constant
+			if (FMath::IsNearlyEqual(TargetWetLevel, CurrentWetLevel))
 			{
-				Audio::MultiplyBufferByConstantInPlace(InputBuffer, TargetOutputVolume);
+				Audio::MultiplyBufferByConstantInPlace(InputBuffer, TargetWetLevel);
 			}
 			else
 			{
 				// To avoid popping, we do a fade on the buffer to the target volume
-				Audio::FadeBufferFast(InputBuffer, CurrentOutputVolume, TargetOutputVolume);
-				CurrentOutputVolume = TargetOutputVolume;
+				Audio::FadeBufferFast(InputBuffer, CurrentWetLevel, TargetWetLevel);
+				CurrentWetLevel = TargetWetLevel;
+			}
+		}
+
+		// Check to see if need to mix together the dry and wet buffers
+		if (DryChannelBuffer.Num())
+		{
+			Audio::MixInBufferFast(DryChannelBuffer, InputBuffer);
+		}
+
+
+		if (bApplyOutputVolumeScale)
+		{
+			const float TargetVolumeProduct = TargetOutputVolume * InitializedOutputVolume;
+			const float OutputVolumeProduct = OutputVolume * InitializedOutputVolume;
+
+			// If we've already set the volume, only need to multiply by constant
+			if (FMath::IsNearlyEqual(TargetVolumeProduct, OutputVolumeProduct))
+			{
+				Audio::MultiplyBufferByConstantInPlace(InputBuffer, OutputVolumeProduct);
+			}
+			else
+			{
+				// To avoid popping, we do a fade on the buffer to the target volume
+				Audio::FadeBufferFast(InputBuffer, OutputVolumeProduct, TargetVolumeProduct);
+				OutputVolume = TargetOutputVolume;
+
+				// No longer need to multiply the output buffer if we're now at 1.0
+				if (FMath::IsNearlyEqual(OutputVolume * InitializedOutputVolume, 1.0f))
+				{
+					bApplyOutputVolumeScale = false;
+				}
 			}
 		}
 
