Fixing audio mixer submix.

Removing accidentally added code from bad merge.

diff --git a/Engine/Source/Runtime/AudioMixer/Private/AudioMixerSubmix.cpp b/Engine/Source/Runtime/AudioMixer/Private/AudioMixerSubmix.cpp
--- a/Engine/Source/Runtime/AudioMixer/Private/AudioMixerSubmix.cpp
+++ b/Engine/Source/Runtime/AudioMixer/Private/AudioMixerSubmix.cpp
@@ -940,30 +940,7 @@ namespace Audio
 		{
 			FMemory::Memzero((void*)BufferPtr, sizeof(float) * NumSamples);
 		}
-
-
-		// Don't necessarily need to do this if the user isn't using this feature
-		if (!FMath::IsNearlyEqual(TargetWetLevel, CurrentWetLevel) || !FMath::IsNearlyEqual(CurrentWetLevel, 1.0f))
-		{
-			// If we've already set the volume, only need to multiply by constant
-			if (FMath::IsNearlyEqual(TargetWetLevel, CurrentWetLevel))
-			{
-				Audio::MultiplyBufferByConstantInPlace(InputBuffer, TargetWetLevel);
-			}
-			else
-			{
-				// To avoid popping, we do a fade on the buffer to the target volume
-				Audio::FadeBufferFast(InputBuffer, CurrentWetLevel, TargetWetLevel);
-				CurrentWetLevel = TargetWetLevel;
-			}
-		}
-
-		// Check to see if need to mix together the dry and wet buffers
-		if (DryChannelBuffer.Num())
-		{
-			Audio::MixInBufferFast(DryChannelBuffer, InputBuffer);
-		}
-		
+	
 		// If we are recording, Add out buffer to the RecordingData buffer:
 		{
 			FScopeLock ScopedLock(&RecordingCriticalSection);
@@ -989,7 +966,6 @@ namespace Audio
 			const int32 BufferSamples = InputBuffer.Num();
 			const float* AudioBufferPtr = InputBuffer.GetData();
 
-
 			// Perform envelope following per channel
 			FScopeLock EnvelopeScopeLock(&EnvelopeCriticalSection);
 			FMemory::Memset(EnvelopeValues, sizeof(float) * AUDIO_MIXER_MAX_OUTPUT_CHANNELS);
@@ -1012,28 +988,6 @@ namespace Audio
 			EnvelopeNumChannels = NumChannels;
 		}
 
-		if (!FMath::IsNearlyEqual(TargetWetLevel, CurrentWetLevel) || !FMath::IsNearlyEqual(CurrentWetLevel, 1.0f))
-		{
-			// If we've already set the volume, only need to multiply by constant
-			if (FMath::IsNearlyEqual(TargetWetLevel, CurrentWetLevel))
-			{
-				Audio::MultiplyBufferByConstantInPlace(InputBuffer, TargetWetLevel);
-			}
-			else
-			{
-				// To avoid popping, we do a fade on the buffer to the target volume
-				Audio::FadeBufferFast(InputBuffer, CurrentWetLevel, TargetWetLevel);
-				CurrentWetLevel = TargetWetLevel;
-			}
-		}
-
-		// Check to see if need to mix together the dry and wet buffers
-		if (DryChannelBuffer.Num())
-		{
-			Audio::MixInBufferFast(DryChannelBuffer, InputBuffer);
-		}
-
-
 		if (bApplyOutputVolumeScale)
 		{
 			const float TargetVolumeProduct = TargetOutputVolume * InitializedOutputVolume;
