Fixing compile error

diff --git a/Engine/Source/Runtime/AudioMixer/Private/AudioMixerSubmix.cpp b/Engine/Source/Runtime/AudioMixer/Private/AudioMixerSubmix.cpp
--- a/Engine/Source/Runtime/AudioMixer/Private/AudioMixerSubmix.cpp
+++ b/Engine/Source/Runtime/AudioMixer/Private/AudioMixerSubmix.cpp
@@ -988,27 +988,19 @@ namespace Audio
 			EnvelopeNumChannels = NumChannels;
 		}
 
-		if (bApplyOutputVolumeScale)
+		// Now apply the output volume
+		if (!FMath::IsNearlyEqual(TargetOutputVolume, CurrentOutputVolume) || !FMath::IsNearlyEqual(CurrentOutputVolume, 1.0f))
 		{
-			const float TargetVolumeProduct = TargetOutputVolume * InitializedOutputVolume;
-			const float OutputVolumeProduct = OutputVolume * InitializedOutputVolume;
-
-			// If we've already set the volume, only need to multiply by constant
-			if (FMath::IsNearlyEqual(TargetVolumeProduct, OutputVolumeProduct))
+			// If we've already set the output volume, only need to multiply by constant
+			if (FMath::IsNearlyEqual(TargetOutputVolume, CurrentOutputVolume))
 			{
-				Audio::MultiplyBufferByConstantInPlace(InputBuffer, OutputVolumeProduct);
+				Audio::MultiplyBufferByConstantInPlace(InputBuffer, TargetOutputVolume);
 			}
 			else
 			{
 				// To avoid popping, we do a fade on the buffer to the target volume
-				Audio::FadeBufferFast(InputBuffer, OutputVolumeProduct, TargetVolumeProduct);
-				OutputVolume = TargetOutputVolume;
-
-				// No longer need to multiply the output buffer if we're now at 1.0
-				if (FMath::IsNearlyEqual(OutputVolume * InitializedOutputVolume, 1.0f))
-				{
-					bApplyOutputVolumeScale = false;
-				}
+				Audio::FadeBufferFast(InputBuffer, CurrentOutputVolume, TargetOutputVolume);
+				CurrentOutputVolume = TargetOutputVolume;
 			}
 		}
 
