- Fix for thread safety access violation in DynamicsProcessor by removing USoundSubmix reference from audio render thread OnProcessAudio call
- Fix for SoundEffectPreset was not being set on sound effect instance prior to init in CreateInstance factory call
- Fix for registering a submix listener in DynamicsProcessor while device is initializing was not getting called due to new handle system not including device as acquirable until initialized
- Fix for registering a submix buffer listener for a submix that hasn't been loaded yet failed to register the listener (repro'ed by registering buffer listener while initializing a reference to a submix in an effect, not in the submix graph)
#rnx

diff --git a/Engine/Source/Runtime/AudioMixer/Classes/SubmixEffects/AudioMixerSubmixEffectDynamicsProcessor.h b/Engine/Source/Runtime/AudioMixer/Classes/SubmixEffects/AudioMixerSubmixEffectDynamicsProcessor.h
--- a/Engine/Source/Runtime/AudioMixer/Classes/SubmixEffects/AudioMixerSubmixEffectDynamicsProcessor.h
+++ b/Engine/Source/Runtime/AudioMixer/Classes/SubmixEffects/AudioMixerSubmixEffectDynamicsProcessor.h
@@ -2,6 +2,7 @@
 #pragma once
 
 #include "AudioDevice.h"
+#include "Delegates/IDelegateInstance.h"
 #include "DSP/DynamicsProcessor.h"
 #include "Sound/SoundEffectSubmix.h"
 #include "Sound/SoundSubmix.h"
@@ -180,15 +181,20 @@ public:
 
 	virtual ~FSubmixEffectDynamicsProcessor();
 
-	// Called on an audio effect at initialization on main thread before audio processing begins.
-	virtual void Init(const FSoundEffectSubmixInitData& InSampleRate) override;
+	// Gets the effect's deviceId that owns it
+	Audio::FDeviceId GetDeviceId() const;
 
-	// Process the input block of audio. Called on audio thread.
+	// Called on an audio effect at initialization on audio thread before audio processing begins.
+	virtual void Init(const FSoundEffectSubmixInitData& InInitData) override;
+
+	// Process the input block of audio. Called on audio render thread.
 	virtual void OnProcessAudio(const FSoundEffectSubmixInputData& InData, FSoundEffectSubmixOutputData& OutData) override;
 
 	// Called when an audio effect preset is changed
 	virtual void OnPresetChanged() override;
 
+	void SetExternalSubmix(USoundSubmix* InSoundSubmix);
+
 protected:
 	// ISubmixBufferListener interface
 	virtual void OnNewSubmixBuffer(
@@ -199,16 +205,21 @@ protected:
 		const int32			InSampleRate,
 		double				InAudioClock) override;
 
+	void OnNewDeviceCreated(Audio::FDeviceId InDeviceId);
+	
 	TWeakObjectPtr<USoundSubmix> ExternalSubmix;
 	Audio::AlignedFloatBuffer AudioExternal;
 
+	FThreadSafeBool bUseExternalSubmix;
+
 	TArray<float> AudioKeyFrame;
 	TArray<float> AudioInputFrame;
 	TArray<float> AudioOutputFrame;
 
 	Audio::FDeviceId DeviceId;
-
 	Audio::FDynamicsProcessor DynamicsProcessor;
+
+	FDelegateHandle DeviceCreatedHandle;
 };
 
 UCLASS(ClassGroup = AudioSourceEffect, meta = (BlueprintSpawnableComponent))
@@ -217,13 +228,17 @@ class AUDIOMIXER_API USubmixEffectDynamicsProcessorPreset : public USoundEffectS
 	GENERATED_BODY()
 
 public:
-
 	EFFECT_PRESET_METHODS(SubmixEffectDynamicsProcessor)
 
-	void Serialize(FStructuredArchive::FRecord Record) override;
+	virtual void OnInit() override;
+
+	virtual void Serialize(FStructuredArchive::FRecord Record) override;
+
+	UFUNCTION(BlueprintCallable, Category = "Audio|Effects")
+	void SetExternalSubmix(USoundSubmix* Submix);
 
 	UFUNCTION(BlueprintCallable, Category = "Audio|Effects")
-	void SetSettings(const FSubmixEffectDynamicsProcessorSettings& InSettings);
+	void SetSettings(const FSubmixEffectDynamicsProcessorSettings& Settings);
 
 	UPROPERTY(EditAnywhere, BlueprintReadWrite, Category = SubmixEffectPreset, meta = (ShowOnlyInnerProperties))
 	FSubmixEffectDynamicsProcessorSettings Settings;
diff --git a/Engine/Source/Runtime/AudioMixer/Private/AudioMixerDevice.cpp b/Engine/Source/Runtime/AudioMixer/Private/AudioMixerDevice.cpp
--- a/Engine/Source/Runtime/AudioMixer/Private/AudioMixerDevice.cpp
+++ b/Engine/Source/Runtime/AudioMixer/Private/AudioMixerDevice.cpp
@@ -1947,81 +1947,69 @@ namespace Audio
 
 	void FMixerDevice::RegisterSubmixBufferListener(ISubmixBufferListener* InSubmixBufferListener, USoundSubmix* InSubmix)
 	{
-		if (!IsInAudioThread())
+		DECLARE_CYCLE_STAT(TEXT("FAudioThreadTask.RegisterSubmixBufferListener"), STAT_RegisterSubmixBufferListener, STATGROUP_AudioThreadCommands);
+
+		const bool bUseMaster = InSubmix == nullptr;
+		const TWeakObjectPtr<USoundSubmix> SubmixPtr(InSubmix);
+
+		auto RegisterLambda = [this, InSubmixBufferListener, bUseMaster, SubmixPtr]()
 		{
-			DECLARE_CYCLE_STAT(TEXT("FAudioThreadTask.RegisterSubmixBufferListener"), STAT_RegisterSubmixBufferListener, STATGROUP_AudioThreadCommands);
-			auto AudioThreadCommand = [this, InSubmixBufferListener, InSubmix]()
+			CSV_SCOPED_TIMING_STAT(Audio, RegisterSubmixBufferListener);
+
+			FMixerSubmixPtr FoundSubmix = bUseMaster
+				? GetMasterSubmix().Pin()
+				: GetSubmixInstance(SubmixPtr.Get()).Pin();
+
+			// Attempt to register submix if instance not found and is not master (i.e. default) submix
+			if (!bUseMaster && !FoundSubmix.IsValid() && SubmixPtr.IsValid())
 			{
-				CSV_SCOPED_TIMING_STAT(Audio, StopSpectrumAnalysis);
-				RegisterSubmixBufferListener(InSubmixBufferListener, InSubmix);
-			};
+				RegisterSoundSubmix(SubmixPtr.Get(), true /* bInit */);
+				FoundSubmix = GetSubmixInstance(SubmixPtr.Get()).Pin();
+			}
 
-			if (IsInGameThread())
+			if (FoundSubmix.IsValid())
 			{
-				FAudioThread::RunCommandOnAudioThread(AudioThreadCommand, GET_STATID(STAT_RegisterSubmixBufferListener));
+				FoundSubmix->RegisterBufferListener(InSubmixBufferListener);
 			}
 			else
 			{
-				AsyncTask(ENamedThreads::GameThread, [AudioThreadCommand]() {
-					FAudioThread::RunCommandOnAudioThread(AudioThreadCommand, GET_STATID(STAT_RegisterSubmixBufferListener));
-				});
+				UE_LOG(LogAudioMixer, Warning, TEXT("Submix buffer listener not registered. Submix not loaded."));
 			}
-			return;
-		}
-
-		FMixerSubmixPtr FoundSubmix = GetSubmixInstance(InSubmix).Pin();
-		if (FoundSubmix.IsValid())
-		{
-			return FoundSubmix->RegisterBufferListener(InSubmixBufferListener);
-		}
-		else
-		{
-			FMixerSubmixWeakPtr MasterSubmix = GetMasterSubmix();
-			FMixerSubmixPtr MasterSubmixPtr = MasterSubmix.Pin();
-			check(MasterSubmixPtr.IsValid());
+		};
 
-			return MasterSubmixPtr->RegisterBufferListener(InSubmixBufferListener);
-		}
+		IsInAudioThread()
+			? RegisterLambda()
+			: AsyncTask(ENamedThreads::AudioThread, MoveTemp(RegisterLambda));
 	}
 
 	void FMixerDevice::UnregisterSubmixBufferListener(ISubmixBufferListener* InSubmixBufferListener, USoundSubmix* InSubmix)
 	{
-		if (!IsInAudioThread())
+		DECLARE_CYCLE_STAT(TEXT("FAudioThreadTask.UnregisterSubmixBufferListener"), STAT_UnregisterSubmixBufferListener, STATGROUP_AudioThreadCommands);
+
+		const bool bUseMaster = InSubmix == nullptr;
+		const TWeakObjectPtr<USoundSubmix> SubmixPtr(InSubmix);
+
+		auto UnregisterLambda = [this, InSubmixBufferListener, bUseMaster, SubmixPtr]()
 		{
-			DECLARE_CYCLE_STAT(TEXT("FAudioThreadTask.UnregisterSubmixBufferListener"), STAT_UnregisterSubmixBufferListener, STATGROUP_AudioThreadCommands);
+			CSV_SCOPED_TIMING_STAT(Audio, UnregisterSubmixBufferListener);
 
-			auto AudioThreadCommand = [this, InSubmixBufferListener, InSubmix]()
-			{
-				CSV_SCOPED_TIMING_STAT(Audio, UnregisterSubmixBufferListener);
-				UnregisterSubmixBufferListener(InSubmixBufferListener, InSubmix);
-			};
+			FMixerSubmixPtr FoundSubmix = bUseMaster
+				? GetMasterSubmix().Pin()
+				: GetSubmixInstance(SubmixPtr.Get()).Pin();
 
-			if (IsInGameThread())
+			if (FoundSubmix.IsValid())
 			{
-				FAudioThread::RunCommandOnAudioThread(AudioThreadCommand, GET_STATID(STAT_UnregisterSubmixBufferListener));
+				FoundSubmix->UnregisterBufferListener(InSubmixBufferListener);
 			}
 			else
 			{
-				AsyncTask(ENamedThreads::GameThread, [AudioThreadCommand]() {
-					FAudioThread::RunCommandOnAudioThread(AudioThreadCommand, GET_STATID(STAT_UnregisterSubmixBufferListener));
-				});
+				UE_LOG(LogAudioMixer, Display, TEXT("Submix buffer listener not unregistered. Submix not loaded."));
 			}
+		};
 
-			return;
-		}
-
-		FMixerSubmixPtr FoundSubmix = GetSubmixInstance(InSubmix).Pin();
-		if (FoundSubmix.IsValid())
-		{
-			return FoundSubmix->UnregisterBufferListener(InSubmixBufferListener);
-		}
-		else
-		{
-			FMixerSubmixPtr MasterSubmixPtr = GetMasterSubmix().Pin();
-			check(MasterSubmixPtr.IsValid());
-
-			return MasterSubmixPtr->UnregisterBufferListener(InSubmixBufferListener);
-		}
+		IsInAudioThread()
+			? UnregisterLambda()
+			: AsyncTask(ENamedThreads::AudioThread, MoveTemp(UnregisterLambda));
 	}
 
 	void FMixerDevice::StartAudioBus(uint32 InAudioBusId, int32 InNumChannels, bool bInIsAutomatic)
diff --git a/Engine/Source/Runtime/AudioMixer/Private/Effects/AudioMixerSubmixEffectDynamicsProcessor.cpp b/Engine/Source/Runtime/AudioMixer/Private/Effects/AudioMixerSubmixEffectDynamicsProcessor.cpp
--- a/Engine/Source/Runtime/AudioMixer/Private/Effects/AudioMixerSubmixEffectDynamicsProcessor.cpp
+++ b/Engine/Source/Runtime/AudioMixer/Private/Effects/AudioMixerSubmixEffectDynamicsProcessor.cpp
@@ -9,11 +9,6 @@
 // Link to "Audio" profiling category
 CSV_DECLARE_CATEGORY_MODULE_EXTERN(AUDIOMIXERCORE_API, Audio);
 
-namespace
-{
-	// Use a full multichannel frame as a scratch processing buffer
-	static const int32 ProcessorScratchNumChannels = 8;
-} // namespace <>
 
 FSubmixEffectDynamicsProcessor::FSubmixEffectDynamicsProcessor()
 	: DeviceId(INDEX_NONE)
@@ -22,23 +17,33 @@ FSubmixEffectDynamicsProcessor::FSubmixEffectDynamicsProcessor()
 
 FSubmixEffectDynamicsProcessor::~FSubmixEffectDynamicsProcessor()
 {
-	if (ExternalSubmix.IsValid())
+	FAudioDeviceManagerDelegates::OnAudioDeviceCreated.Remove(DeviceCreatedHandle);
+
+	if (FAudioDeviceManager* DeviceManager = FAudioDeviceManager::Get())
 	{
-		Audio::FMixerDevice* MixerDevice = nullptr;
-		if (FAudioDeviceManager* DeviceManager = FAudioDeviceManager::Get())
+		FAudioDeviceHandle DeviceHandle = DeviceManager->GetAudioDevice(DeviceId);
+		if (DeviceHandle.IsValid())
 		{
-			MixerDevice = static_cast<Audio::FMixerDevice*>(DeviceManager->GetAudioDeviceRaw(DeviceId));
-		}
+			Audio::FMixerDevice* MixerDevice = static_cast<Audio::FMixerDevice*>(DeviceHandle.GetAudioDevice());
+			check(MixerDevice);
 
-		if (MixerDevice)
-		{
-			MixerDevice->UnregisterSubmixBufferListener(this, ExternalSubmix.Get());
+			if (ExternalSubmix.IsValid())
+			{
+				MixerDevice->UnregisterSubmixBufferListener(this, ExternalSubmix.Get());
+			}
 		}
 	}
 }
 
+Audio::FDeviceId FSubmixEffectDynamicsProcessor::GetDeviceId() const
+{
+	return DeviceId;
+}
+
 void FSubmixEffectDynamicsProcessor::Init(const FSoundEffectSubmixInitData& InitData)
 {
+	static const int32 ProcessorScratchNumChannels = 8;
+
 	DynamicsProcessor.Init(InitData.SampleRate, ProcessorScratchNumChannels);
 
 	AudioKeyFrame.Reset();
@@ -51,6 +56,11 @@ void FSubmixEffectDynamicsProcessor::Init(const FSoundEffectSubmixInitData& Init
 	AudioOutputFrame.AddZeroed(ProcessorScratchNumChannels);
 
 	DeviceId = InitData.DeviceID;
+
+	if (USubmixEffectDynamicsProcessorPreset* ProcPreset = Cast<USubmixEffectDynamicsProcessorPreset>(Preset.Get()))
+	{
+		SetExternalSubmix(ProcPreset->Settings.ExternalSubmix);
+	}
 }
 
 void FSubmixEffectDynamicsProcessor::OnNewSubmixBuffer(
@@ -126,30 +136,6 @@ void FSubmixEffectDynamicsProcessor::OnPresetChanged()
 
 	static_assert(static_cast<int32>(ESubmixEffectDynamicsChannelLinkMode::Count) == static_cast<int32>(Audio::EDynamicsProcessorChannelLinkMode::Count), "Enumerations must match");
 	DynamicsProcessor.SetChannelLinkMode(static_cast<Audio::EDynamicsProcessorChannelLinkMode>(Settings.LinkMode));
-
-	Audio::FMixerDevice* MixerDevice = nullptr;
-	if (FAudioDeviceManager* DeviceManager = FAudioDeviceManager::Get())
-	{
-		MixerDevice = static_cast<Audio::FMixerDevice*>(DeviceManager->GetAudioDeviceRaw(DeviceId));
-	}
-
-	if (MixerDevice)
-	{
-		if (ExternalSubmix.Get() != Settings.ExternalSubmix)
-		{
-			if (ExternalSubmix.IsValid())
-			{
-				MixerDevice->UnregisterSubmixBufferListener(this, ExternalSubmix.Get());
-			}
-
-			ExternalSubmix = Settings.ExternalSubmix;
-
-			if (ExternalSubmix.IsValid())
-			{
-				MixerDevice->RegisterSubmixBufferListener(this, ExternalSubmix.Get());
-			}
-		}
-	}
 }
 
 void FSubmixEffectDynamicsProcessor::OnProcessAudio(const FSoundEffectSubmixInputData& InData, FSoundEffectSubmixOutputData& OutData)
@@ -167,7 +153,7 @@ void FSubmixEffectDynamicsProcessor::OnProcessAudio(const FSoundEffectSubmixInpu
 		{
 			const int32 SampleIndex = SampleIndexOfFrame + Channel;
 			AudioInputFrame[Channel] = InBuffer[SampleIndex];
-			if (ExternalSubmix.IsValid())
+			if (bUseExternalSubmix)
 			{
 				AudioKeyFrame[Channel] = AudioExternal.Num() > SampleIndex ? AudioExternal[SampleIndex] : 0;
 			}
@@ -189,6 +175,61 @@ void FSubmixEffectDynamicsProcessor::OnProcessAudio(const FSoundEffectSubmixInpu
 	}
 }
 
+void FSubmixEffectDynamicsProcessor::SetExternalSubmix(USoundSubmix* InSoundSubmix)
+{
+	FAudioDeviceManager* DeviceManager = FAudioDeviceManager::Get();
+	if (!DeviceManager)
+	{
+		return;
+	}
+
+	FAudioDeviceHandle DeviceHandle = DeviceManager->GetAudioDevice(DeviceId);
+	if (DeviceHandle.IsValid())
+	{
+		Audio::FMixerDevice* MixerDevice = static_cast<Audio::FMixerDevice*>(DeviceHandle.GetAudioDevice());
+		check(MixerDevice);
+
+		if (ExternalSubmix.Get() != InSoundSubmix)
+		{
+			if (ExternalSubmix.IsValid())
+			{
+				MixerDevice->UnregisterSubmixBufferListener(this, ExternalSubmix.Get());
+			}
+
+			ExternalSubmix = InSoundSubmix;
+
+			bUseExternalSubmix = ExternalSubmix.IsValid();
+			if (bUseExternalSubmix)
+			{
+				MixerDevice->RegisterSubmixBufferListener(this, ExternalSubmix.Get());
+			}
+		}
+
+	}
+	else
+	{
+		DeviceCreatedHandle = FAudioDeviceManagerDelegates::OnAudioDeviceCreated.AddRaw(this, &FSubmixEffectDynamicsProcessor::OnNewDeviceCreated);
+		bUseExternalSubmix = true;
+	}
+}
+
+void FSubmixEffectDynamicsProcessor::OnNewDeviceCreated(Audio::FDeviceId InDeviceId)
+{
+	if (InDeviceId == DeviceId)
+	{
+		FAudioDeviceManagerDelegates::OnAudioDeviceCreated.Remove(DeviceCreatedHandle);
+
+		GET_EFFECT_SETTINGS(SubmixEffectDynamicsProcessor);
+		SetExternalSubmix(Settings.ExternalSubmix);
+
+	}
+}
+
+void USubmixEffectDynamicsProcessorPreset::OnInit()
+{
+	SetExternalSubmix(Settings.ExternalSubmix);
+}
+
 void USubmixEffectDynamicsProcessorPreset::Serialize(FStructuredArchive::FRecord Record)
 {
 	FArchive& UnderlyingArchive = Record.GetUnderlyingArchive();
@@ -204,7 +245,16 @@ void USubmixEffectDynamicsProcessorPreset::Serialize(FStructuredArchive::FRecord
 	Super::Serialize(Record);
 }
 
+void USubmixEffectDynamicsProcessorPreset::SetExternalSubmix(USoundSubmix* InSubmix)
+{
+	IterateEffects<FSubmixEffectDynamicsProcessor>([this, InSubmix](FSubmixEffectDynamicsProcessor& Instance)
+	{
+		Instance.SetExternalSubmix(InSubmix);
+	});
+}
+
 void USubmixEffectDynamicsProcessorPreset::SetSettings(const FSubmixEffectDynamicsProcessorSettings& InSettings)
 {
 	UpdateSettings(InSettings);
+	SetExternalSubmix(InSettings.ExternalSubmix);
 }
diff --git a/Engine/Source/Runtime/Engine/Classes/Sound/SoundEffectBase.h b/Engine/Source/Runtime/Engine/Classes/Sound/SoundEffectBase.h
--- a/Engine/Source/Runtime/Engine/Classes/Sound/SoundEffectBase.h
+++ b/Engine/Source/Runtime/Engine/Classes/Sound/SoundEffectBase.h
@@ -68,7 +68,7 @@ public:
 	virtual ~FSoundEffectBase() = default;
 
 	/** Called when the sound effect's preset changed. */
-	virtual void OnPresetChanged() {};
+	virtual void OnPresetChanged() { }
 
 	/** Returns if the submix is active or bypassing audio. */
 	bool IsActive() const;
@@ -90,9 +90,6 @@ public:
 protected:
 	FSoundEffectBase();
 
-	/** Removes the instance from the preset. */
-	void ClearPreset();
-
 	/** Pumps messages awaiting execution on the audio render thread */
 	void PumpPendingMessages();
 
@@ -108,8 +105,13 @@ protected:
 	// Effect command queue
 	TQueue<TUniqueFunction<void()>> CommandQueue;
 
+private:
+	/** Removes the instance from the preset. */
+	void ClearPreset();
+
 	// Allow preset to re-register when editor update is requested
-	// and create effects using the templated Create call
+	// and create effects using the templated Create call, as well
+	// as clear preset.
 	friend class USoundEffectPreset;
 };
 
diff --git a/Engine/Source/Runtime/Engine/Classes/Sound/SoundEffectPreset.h b/Engine/Source/Runtime/Engine/Classes/Sound/SoundEffectPreset.h
--- a/Engine/Source/Runtime/Engine/Classes/Sound/SoundEffectPreset.h
+++ b/Engine/Source/Runtime/Engine/Classes/Sound/SoundEffectPreset.h
@@ -91,10 +91,14 @@ public:
 	static TSharedPtr<TSoundEffectType, ESPMode::ThreadSafe> CreateInstance(const TInitData& InInitData, USoundEffectPreset& InOutPreset)
 	{
 		TSoundEffectType* NewEffect = static_cast<TSoundEffectType*>(InOutPreset.CreateNewEffect());
+		NewEffect->Preset = &InOutPreset;
 		NewEffect->Init(InInitData);
 
 		TSharedPtr<TSoundEffectType, ESPMode::ThreadSafe> NewEffectPtr(NewEffect);
-		RegisterInstance(InOutPreset, StaticCastSharedPtr<FSoundEffectBase, TSoundEffectType, ESPMode::ThreadSafe>(NewEffectPtr));
+
+		TSoundEffectPtr SoundEffectPtr = StaticCastSharedPtr<FSoundEffectBase, TSoundEffectType, ESPMode::ThreadSafe>(NewEffectPtr);
+		InOutPreset.AddEffectInstance(SoundEffectPtr);
+
 		return NewEffectPtr;
 	}
 
diff --git a/Engine/Source/Runtime/Engine/Private/SoundEffectBase.cpp b/Engine/Source/Runtime/Engine/Private/SoundEffectBase.cpp
--- a/Engine/Source/Runtime/Engine/Private/SoundEffectBase.cpp
+++ b/Engine/Source/Runtime/Engine/Private/SoundEffectBase.cpp
@@ -6,7 +6,7 @@
 
 
 FSoundEffectBase::FSoundEffectBase()
-	: bChanged(false)
+	: bChanged(true)
 	, bIsRunning(false)
 	, bIsActive(false)
 {}
