Fixing crashes in audio spectral analysis

diff --git a/Engine/Source/Runtime/AudioMixer/Private/AudioMixerSubmix.cpp b/Engine/Source/Runtime/AudioMixer/Private/AudioMixerSubmix.cpp
--- a/Engine/Source/Runtime/AudioMixer/Private/AudioMixerSubmix.cpp
+++ b/Engine/Source/Runtime/AudioMixer/Private/AudioMixerSubmix.cpp
@@ -49,6 +49,7 @@ namespace Audio
 		, NumSubmixEffects(0)
 		, bIsRecording(false)
 		, bIsBackgroundMuted(false)
+		, bIsSpectrumAnalyzing(false)
 		, OwningSubmixObject(nullptr)
 	{
 		EnvelopeFollowers.Reset();
@@ -1666,10 +1667,12 @@ namespace Audio
 					{
 						Audio::FSpectrumAnalyzerScopeLock AnalyzerLock(SpectrumAnalyzer.Get());
 
-						check(DelegateInfo.SpectrumBandExtractor.IsValid());
-						ISpectrumBandExtractor* Extractor = DelegateInfo.SpectrumBandExtractor.Get();
+						if (ensure(DelegateInfo.SpectrumBandExtractor.IsValid()))
+						{
+							ISpectrumBandExtractor* Extractor = DelegateInfo.SpectrumBandExtractor.Get();
 
-						SpectrumAnalyzer->GetBands(*Extractor, SpectralResults);
+							SpectrumAnalyzer->GetBands(*Extractor, SpectralResults);
+						}
 					}
 
 					// Feed the results through the band envelope followers
diff --git a/Engine/Source/Runtime/SignalProcessing/Private/SpectrumAnalyzer.cpp b/Engine/Source/Runtime/SignalProcessing/Private/SpectrumAnalyzer.cpp
--- a/Engine/Source/Runtime/SignalProcessing/Private/SpectrumAnalyzer.cpp
+++ b/Engine/Source/Runtime/SignalProcessing/Private/SpectrumAnalyzer.cpp
@@ -120,7 +120,7 @@ namespace Audio
 
 						default:
 							// Should not get here. Means that a window type is not covered.
-							check(false);
+							UE_LOG(LogSignalProcessing, Warning, TEXT("Invalid window type enum encountered, %d"), (int32)InSettings.WindowType);
 							WindowScale = 1.f / FloatFFTSize;
 							break;
 
diff --git a/Engine/Source/Runtime/SignalProcessing/Public/DSP/SpectrumAnalyzer.h b/Engine/Source/Runtime/SignalProcessing/Public/DSP/SpectrumAnalyzer.h
--- a/Engine/Source/Runtime/SignalProcessing/Public/DSP/SpectrumAnalyzer.h
+++ b/Engine/Source/Runtime/SignalProcessing/Public/DSP/SpectrumAnalyzer.h
@@ -82,6 +82,12 @@ namespace Audio
 		/** Window used when perform FFT */
 		EWindowType WindowType;
 
+		FSpectrumBandExtractorSettings()
+		:	SampleRate(48000.f)
+		,	FFTSize(1024)
+		,	FFTScaling(EFFTScaling::None)
+		,	WindowType(EWindowType::None)
+		{}
 
 		/** Compare whether two settings structures are equal. */
 		bool operator==(const FSpectrumBandExtractorSettings& Other) const
