Move invalid assertion to conditional to avoid crash when instances are destroyed during shutdown
#rnx

diff --git a/Engine/Source/Runtime/AudioMixer/Private/AudioMixerDevice.cpp b/Engine/Source/Runtime/AudioMixer/Private/AudioMixerDevice.cpp
--- a/Engine/Source/Runtime/AudioMixer/Private/AudioMixerDevice.cpp
+++ b/Engine/Source/Runtime/AudioMixer/Private/AudioMixerDevice.cpp
@@ -453,19 +453,19 @@ namespace Audio
 		{
 			if (SoundSubmix)
 			{
-				// Retrieve the submix instance and the envelope data
+				// Retrieve the submix instance and the envelope data and broadcast on the audio thread.
 				Audio::FMixerSubmixWeakPtr SubmixPtr = GetSubmixInstance(SoundSubmix);
-				check(SubmixPtr.IsValid());
-
-				// On the audio thread, do the broadcast.
-				FAudioThread::RunCommandOnGameThread([this, SubmixPtr]()
+				if (SubmixPtr.IsValid())
 				{
-					Audio::FMixerSubmixPtr ThisSubmixPtr = SubmixPtr.Pin();
-					if (ThisSubmixPtr.IsValid())
-					{
-						ThisSubmixPtr->BroadcastDelegates();
-					}
-				});
+					FAudioThread::RunCommandOnGameThread([this, SubmixPtr]()
+						{
+							Audio::FMixerSubmixPtr ThisSubmixPtr = SubmixPtr.Pin();
+							if (ThisSubmixPtr.IsValid())
+							{
+								ThisSubmixPtr->BroadcastDelegates();
+							}
+						});
+				}
 			}
 		}
 
