Adding ability to remove an existing spectral analysis delegate
Adding convenience functions for creating arrays of band settings.

diff --git a/Engine/Source/Runtime/AudioMixer/Private/AudioMixerBlueprintLibrary.cpp b/Engine/Source/Runtime/AudioMixer/Private/AudioMixerBlueprintLibrary.cpp
--- a/Engine/Source/Runtime/AudioMixer/Private/AudioMixerBlueprintLibrary.cpp
+++ b/Engine/Source/Runtime/AudioMixer/Private/AudioMixerBlueprintLibrary.cpp
@@ -5,6 +5,7 @@
 #include "AudioDevice.h"
 #include "AudioMixerDevice.h"
 #include "CoreMinimal.h"
+#include "DSP/ConstantQ.h"
 #include "DSP/SpectrumAnalyzer.h"
 #include "ContentStreaming.h"
 #include "AudioCompressionSettingsUtils.h"
@@ -287,6 +288,112 @@ void UAudioMixerBlueprintLibrary::StopAnalyzingOutput(const UObject* WorldContex
 	}
 }
 
+TArray<FSoundSubmixSpectralAnalysisBandSettings> UAudioMixerBlueprintLibrary::MakeMusicalSpectralAnalysisBandSettings(int32 InNumNotes, EMusicalNoteName InStartingMusicalNote, int32 InStartingOctave, int32 InAttackTimeMsec, int32 InReleaseTimeMsec)
+{
+	// Make values sane.
+	InNumNotes = FMath::Clamp(InNumNotes, 0, 10000);
+	InStartingOctave = FMath::Clamp(InStartingOctave, -1, 10);
+	InAttackTimeMsec = FMath::Clamp(InAttackTimeMsec, 0, 10000);
+	InReleaseTimeMsec = FMath::Clamp(InReleaseTimeMsec, 0, 10000);
+
+	// Some assumptions here on what constitutes "music". 12 notes, equal temperament.
+	const float BandsPerOctave = 12.f;
+	// This QFactor makes the bandwidth equal to the difference in frequency between adjacent notes.
+	const float QFactor = 1.f / (FMath::Pow(2.f, 1.f / BandsPerOctave) - 1.f);
+
+	// Base note index off of A4 which we know to be 440 hz.
+	// Make note relative to A
+	int32 NoteIndex = static_cast<int32>(InStartingMusicalNote) - static_cast<int32>(EMusicalNoteName::A);
+	// Make relative to 4th octave of A
+	NoteIndex += 12 * (InStartingOctave - 4);
+
+	const float StartingFrequency = 440.f * FMath::Pow(2.f, static_cast<float>(NoteIndex) / 12.f);
+
+
+	TArray<FSoundSubmixSpectralAnalysisBandSettings> BandSettingsArray;
+	for (int32 i = 0; i < InNumNotes; i++)
+	{
+		FSoundSubmixSpectralAnalysisBandSettings BandSettings;
+
+		BandSettings.BandFrequency = Audio::FPseudoConstantQ::GetConstantQCenterFrequency(i, StartingFrequency, BandsPerOctave);
+
+		BandSettings.QFactor = QFactor;
+		BandSettings.AttackTimeMsec = InAttackTimeMsec;
+		BandSettings.ReleaseTimeMsec = InReleaseTimeMsec;
+
+		BandSettingsArray.Add(BandSettings);
+	}
+
+	return BandSettingsArray;
+}
+
+TArray<FSoundSubmixSpectralAnalysisBandSettings> UAudioMixerBlueprintLibrary::MakeFullSpectrumSpectralAnalysisBandSettings(int32 InNumBands, float InMinimumFrequency, float InMaximumFrequency, int32 InAttackTimeMsec, int32 InReleaseTimeMsec)
+{
+	// Make inputs sane.
+	InNumBands = FMath::Clamp(InNumBands, 0, 10000);
+	InMinimumFrequency = FMath::Clamp(InMinimumFrequency, 20.0f, 20000.0f);
+	InMaximumFrequency = FMath::Clamp(InMaximumFrequency, InMinimumFrequency, 20000.0f);
+	InAttackTimeMsec = FMath::Clamp(InAttackTimeMsec, 0, 10000);
+	InReleaseTimeMsec = FMath::Clamp(InReleaseTimeMsec, 0, 10000);
+
+	// Calculate CQT settings needed to space bands.
+	const float NumOctaves = FMath::Loge(InMaximumFrequency / InMinimumFrequency) / FMath::Loge(2.f);
+	const float BandsPerOctave = static_cast<float>(InNumBands) / FMath::Max(NumOctaves, 0.01f);
+	const float QFactor = 1.f / (FMath::Pow(2.f, 1.f / FMath::Max(BandsPerOctave, 0.01f)) - 1.f);
+
+	TArray<FSoundSubmixSpectralAnalysisBandSettings> BandSettingsArray;
+	for (int32 i = 0; i < InNumBands; i++)
+	{
+		FSoundSubmixSpectralAnalysisBandSettings BandSettings;
+
+		BandSettings.BandFrequency = Audio::FPseudoConstantQ::GetConstantQCenterFrequency(i, InMinimumFrequency, BandsPerOctave);
+
+		BandSettings.QFactor = QFactor;
+		BandSettings.AttackTimeMsec = InAttackTimeMsec;
+		BandSettings.ReleaseTimeMsec = InReleaseTimeMsec;
+
+		BandSettingsArray.Add(BandSettings);
+	}
+
+	return BandSettingsArray;
+}
+
+TArray<FSoundSubmixSpectralAnalysisBandSettings> UAudioMixerBlueprintLibrary::MakePresetSpectralAnalysisBandSettings(EAudioSpectrumBandPresetType InBandPresetType, int32 InNumBands, int32 InAttackTimeMsec, int32 InReleaseTimeMsec)
+{
+	float MinimumFrequency = 20.f;
+	float MaximumFrequency = 20000.f;
+
+	// Likely all these are debatable. What we are shooting for is the most active frequency
+	// ranges, so when an instrument plays a significant amount of spectral energy from that
+	// instrument will show up in the frequency range. 
+	switch (InBandPresetType)
+	{
+		case EAudioSpectrumBandPresetType::KickDrum:
+			MinimumFrequency = 40.f;
+			MaximumFrequency = 100.f;
+			break;
+
+		case EAudioSpectrumBandPresetType::SnareDrum:
+			MinimumFrequency = 150.f;
+			MaximumFrequency = 4500.f;
+			break;
+
+		case EAudioSpectrumBandPresetType::Voice:
+			MinimumFrequency = 300.f;
+			MaximumFrequency = 3000.f;
+			break;
+
+		case EAudioSpectrumBandPresetType::Cymbals:
+			MinimumFrequency = 6000.f;
+			MaximumFrequency = 16000.f;
+			break;
+
+		// More presets can be added. The possibilities are endless.
+	}
+
+	return MakeFullSpectrumSpectralAnalysisBandSettings(InNumBands, MinimumFrequency, MaximumFrequency, InAttackTimeMsec, InReleaseTimeMsec);
+}
+
 void UAudioMixerBlueprintLibrary::GetMagnitudeForFrequencies(const UObject* WorldContextObject, const TArray<float>& Frequencies, TArray<float>& Magnitudes, USoundSubmix* SubmixToAnalyze /*= nullptr*/)
 {
 	if (Audio::FMixerDevice* MixerDevice = GetAudioMixerDeviceFromWorldContext(WorldContextObject))
diff --git a/Engine/Source/Runtime/AudioMixer/Private/AudioMixerDevice.cpp b/Engine/Source/Runtime/AudioMixer/Private/AudioMixerDevice.cpp
--- a/Engine/Source/Runtime/AudioMixer/Private/AudioMixerDevice.cpp
+++ b/Engine/Source/Runtime/AudioMixer/Private/AudioMixerDevice.cpp
@@ -1919,29 +1919,59 @@ namespace Audio
 	{
 		if (!IsInAudioThread())
 		{
-			DECLARE_CYCLE_STAT(TEXT("FAudioThreadTask.AddFFTDelegate"), STAT_AddFFTDelegate, STATGROUP_AudioThreadCommands);
+			DECLARE_CYCLE_STAT(TEXT("FAudioThreadTask.AddSpectralAnalysisDelegate"), STAT_AddSpectralAnalysisDelegate, STATGROUP_AudioThreadCommands);
 
 			FAudioThread::RunCommandOnAudioThread([this, InSubmix, InBandSettings, OnSubmixSpectralAnalysisBP, UpdateRate]()
 			{
-				CSV_SCOPED_TIMING_STAT(Audio, AddEnvelopeFollowerDelegate);
+				CSV_SCOPED_TIMING_STAT(Audio, AddSpectralAnalysisDelegate);
 				AddSpectralAnalysisDelegate(InSubmix, InBandSettings, OnSubmixSpectralAnalysisBP, UpdateRate);
-			}, GET_STATID(STAT_AddFFTDelegate));
+			}, GET_STATID(STAT_AddSpectralAnalysisDelegate));
 			return;
 		}
 
-		// if we can find the submix here, record that submix. Otherwise, just record the master submix.
+		// get submix if it is available.
 		FMixerSubmixPtr FoundSubmix = GetSubmixInstance(InSubmix).Pin();
-		if (FoundSubmix.IsValid())
+
+		if (!FoundSubmix.IsValid())
+		{
+			// If can't find the submix isntance, use master submix.
+			FMixerSubmixWeakPtr MasterSubmix = GetMasterSubmix();
+			FoundSubmix = MasterSubmix.Pin();
+		}
+
+		if (ensure(FoundSubmix.IsValid()))
 		{
 			FoundSubmix->AddSpectralAnalysisDelegate(OnSubmixSpectralAnalysisBP, InBandSettings, UpdateRate);
 		}
-		else
+	}
+
+	void FMixerDevice::RemoveSpectralAnalysisDelegate(USoundSubmix* InSubmix, const FOnSubmixSpectralAnalysisBP& InDelegate)
+	{
+		if (!IsInAudioThread())
 		{
+			DECLARE_CYCLE_STAT(TEXT("FAudioThreadTask.RemoveSpectralAnalysisDelegate"), STAT_RemoveSpectralAnalysisDelegate, STATGROUP_AudioThreadCommands);
+
+			FAudioThread::RunCommandOnAudioThread([this, InSubmix, InDelegate]()
+			{
+				CSV_SCOPED_TIMING_STAT(Audio, RemoveSpectralAnalysisDelegate);
+				RemoveSpectralAnalysisDelegate(InSubmix, InDelegate);
+			}, GET_STATID(STAT_RemoveSpectralAnalysisDelegate));
+			return;
+		}
+
+		// get submix if it is available.
+		FMixerSubmixPtr FoundSubmix = GetSubmixInstance(InSubmix).Pin();
+
+		if (!FoundSubmix.IsValid())
+		{
+			// If can't find the submix isntance, use master submix.
 			FMixerSubmixWeakPtr MasterSubmix = GetMasterSubmix();
-			FMixerSubmixPtr MasterSubmixPtr = MasterSubmix.Pin();
-			check(MasterSubmixPtr.IsValid());
+			FoundSubmix = MasterSubmix.Pin();
+		}
 
-			MasterSubmixPtr->AddSpectralAnalysisDelegate(OnSubmixSpectralAnalysisBP, InBandSettings, UpdateRate);
+		if (ensure(FoundSubmix.IsValid()))
+		{
+			FoundSubmix->RemoveSpectralAnalysisDelegate(InDelegate);
 		}
 	}
 
diff --git a/Engine/Source/Runtime/AudioMixer/Private/AudioMixerSubmix.cpp b/Engine/Source/Runtime/AudioMixer/Private/AudioMixerSubmix.cpp
--- a/Engine/Source/Runtime/AudioMixer/Private/AudioMixerSubmix.cpp
+++ b/Engine/Source/Runtime/AudioMixer/Private/AudioMixerSubmix.cpp
@@ -1490,6 +1490,7 @@ namespace Audio
 		{
 			FSpectralAnalysisBandInfo NewBand;
 			NewBand.BandFrequency = BandSettings.BandFrequency;
+			NewBand.QFactor = BandSettings.QFactor;
 			NewBand.EnvelopeFollower.Init(UpdateRate, BandSettings.AttackTimeMsec, BandSettings.ReleaseTimeMsec);
 		
 			NewDelegateInfo.SpectralBands.Add(NewBand);
@@ -1499,6 +1500,21 @@ namespace Audio
 
 	}
 
+	void FMixerSubmix::RemoveSpectralAnalysisDelegate(const FOnSubmixSpectralAnalysisBP& OnSubmixSpectralAnalysisBP)
+	{
+		for (FSpectrumAnalysisDelegateInfo& Info : SpectralAnalysisDelegates)
+		{
+			if (Info.OnSubmixSpectralAnalysis.Contains(OnSubmixSpectralAnalysisBP))
+			{
+				Info.OnSubmixSpectralAnalysis.Remove(OnSubmixSpectralAnalysisBP);
+			}
+		}
+
+		SpectralAnalysisDelegates.RemoveAllSwap([](FSpectrumAnalysisDelegateInfo& Info) {
+			return !Info.OnSubmixSpectralAnalysis.IsBound();
+		});
+	}
+
 	void FMixerSubmix::StartSpectrumAnalysis(const FSpectrumAnalyzerSettings& InSettings)
 	{
 		bIsSpectrumAnalyzing = true;
@@ -1530,8 +1546,6 @@ namespace Audio
 				break;
 			}
 			
-			// TODO: @phil -- feels like maybe we can pass in the enum to AddBand vs 4 different Add Band functions?  Maybe not because of the Q-factor?
-
 			// Loop through each band and add them to the extractor
 			for (FSpectralAnalysisBandInfo& BandInfo : DelegateInfo.SpectralBands)
 			{
@@ -1547,8 +1561,7 @@ namespace Audio
 						DelegateInfo.SpectrumBandExtractor->AddQuadraticBand(BandInfo.BandFrequency, Metric);
 						break;
 					case FSpectrumAnalyzerSettings::EPeakInterpolationMethod::ConstantQ:
-						// TODO: @phil -- how to compute q-factor?? I assume some math around bands?
-						DelegateInfo.SpectrumBandExtractor->AddConstantQBand(BandInfo.BandFrequency, 20.0f, Metric);
+						DelegateInfo.SpectrumBandExtractor->AddConstantQBand(BandInfo.BandFrequency, BandInfo.QFactor, Metric);
 						break;
 				}
 			}
diff --git a/Engine/Source/Runtime/AudioMixer/Public/AudioMixerBlueprintLibrary.h b/Engine/Source/Runtime/AudioMixer/Public/AudioMixerBlueprintLibrary.h
--- a/Engine/Source/Runtime/AudioMixer/Public/AudioMixerBlueprintLibrary.h
+++ b/Engine/Source/Runtime/AudioMixer/Public/AudioMixerBlueprintLibrary.h
@@ -20,6 +20,23 @@ class USoundSubmix;
 */
 DECLARE_DYNAMIC_DELEGATE_TwoParams(FOnSoundLoadComplete, const class USoundWave*, LoadedSoundWave, const bool, WasCancelled);
 
+UENUM(BlueprintType)
+enum class EMusicalNoteName : uint8
+{
+	C  = 0,
+	Db = 1,
+	D  = 2,
+	Eb = 3,
+	E  = 4,
+	F  = 5,
+	Gb = 6,
+	G  = 7,
+	Ab = 8,
+	A  = 9,
+	Bb = 10,
+	B  = 11,
+};
+
 
 UCLASS(meta=(ScriptName="AudioMixerLibrary"))
 class AUDIOMIXER_API UAudioMixerBlueprintLibrary : public UBlueprintFunctionLibrary
@@ -84,6 +101,38 @@ public:
 	UFUNCTION(BlueprintCallable, Category = "Audio|Analysis", meta = (WorldContext = "WorldContextObject", AdvancedDisplay = 1))
 	static void StopAnalyzingOutput(const UObject* WorldContextObject, USoundSubmix* SubmixToStopAnalyzing = nullptr);
 
+	/** Make an array of musically spaced bands with ascending frequency.
+	 *
+	 *  @param InNumSemitones - The number of semitones to represent.
+	 *  @param InStartingMuiscalNote - The name of the first note in the array.
+	 *  @param InStartingOctave - The octave of the first note in the arrya.
+	 *  @param InAttackTimeMsec - The attack time (in milliseconds) to apply to each band's envelope tracker.
+	 *  @param InReleaseTimeMsec - The release time (in milliseconds) to apply to each band's envelope tracker.
+	 */
+	UFUNCTION(BlueprintPure, Category = "Audio|Analysis")
+	static TArray<FSoundSubmixSpectralAnalysisBandSettings> MakeMusicalSpectralAnalysisBandSettings(int32 InNumSemitones=60, EMusicalNoteName InStartingMusicalNote = EMusicalNoteName::C, int32 InStartingOctave = 2, int32 InAttackTimeMsec = 10, int32 InReleaseTimeMsec = 10);
+
+	/** Make an array of logarithmically spaced bands. 
+	 *
+	 *  @param InNumBands - The number of bands to used to represent the spectrum.
+	 *  @param InMinimumFrequency - The center frequency of the first band.
+	 *  @param InMaximumFrequency - The center frequency of the last band.
+	 *  @param InAttackTimeMsec - The attack time (in milliseconds) to apply to each band's envelope tracker.
+	 *  @param InReleaseTimeMsec - The release time (in milliseconds) to apply to each band's envelope tracker.
+	 */
+	UFUNCTION(BlueprintPure, Category = "Audio|Analysis")
+	static TArray<FSoundSubmixSpectralAnalysisBandSettings> MakeFullSpectrumSpectralAnalysisBandSettings(int32 InNumBands, float InMinimumFrequency=40.f, float InMaximumFrequency=16000.f, int32 InAttackTimeMsec = 10, int32 InReleaseTimeMsec = 10);
+
+	/** Make an array of bands which span the frequency range of a given EAudioSpectrumBandPresetType. 
+	 *
+	 *  @param InBandPresetType - The type audio content which the bands encompass.
+	 *  @param InNumBands - The number of bands used to represent the spectrum.
+	 *  @param InAttackTimeMsec - The attack time (in milliseconds) to apply to each band's envelope tracker.
+	 *  @param InReleaseTimeMsec - The release time (in milliseconds) to apply to each band's envelope tracker.
+	 */
+	UFUNCTION(BlueprintPure, Category = "Audio|Analysis")
+	static TArray<FSoundSubmixSpectralAnalysisBandSettings> MakePresetSpectralAnalysisBandSettings(EAudioSpectrumBandPresetType InBandPresetType, int32 InNumBands, int32 InAttackTimeMsec = 10, int32 InReleaseTimeMsec = 10);
+
 	/** Start spectrum analysis of the audio output. By leaving the Submix To Analyze blank, you can analyze the master output of the game. */
 	UFUNCTION(BlueprintCallable, Category = "Audio|Analysis", meta = (WorldContext = "WorldContextObject", AdvancedDisplay = 3))
 	static void GetMagnitudeForFrequencies(const UObject* WorldContextObject, const TArray<float>& Frequencies, TArray<float>& Magnitudes, USoundSubmix* SubmixToAnalyze = nullptr);
diff --git a/Engine/Source/Runtime/AudioMixer/Public/AudioMixerDevice.h b/Engine/Source/Runtime/AudioMixer/Public/AudioMixerDevice.h
--- a/Engine/Source/Runtime/AudioMixer/Public/AudioMixerDevice.h
+++ b/Engine/Source/Runtime/AudioMixer/Public/AudioMixerDevice.h
@@ -131,6 +131,7 @@ namespace Audio
 		virtual void GetMagnitudesForFrequencies(USoundSubmix* InSubmix, const TArray<float>& InFrequencies, TArray<float>& OutMagnitudes);
 		virtual void GetPhasesForFrequencies(USoundSubmix* InSubmix, const TArray<float>& InFrequencies, TArray<float>& OutPhases);
 		virtual void AddSpectralAnalysisDelegate(USoundSubmix* InSubmix, const TArray<FSoundSubmixSpectralAnalysisBandSettings>& InBandSettings, const FOnSubmixSpectralAnalysisBP& OnSubmixSpectralAnalysisBP, float UpdateRate) override;
+		virtual void RemoveSpectralAnalysisDelegate(USoundSubmix* InSubmixconst, const FOnSubmixSpectralAnalysisBP& OnSubmixSpectralAnalysisBP) override;
 
 		// Submix buffer listener callbacks
 		virtual void RegisterSubmixBufferListener(ISubmixBufferListener* InSubmixBufferListener, USoundSubmix* InSubmix = nullptr) override;
diff --git a/Engine/Source/Runtime/AudioMixer/Public/AudioMixerSubmix.h b/Engine/Source/Runtime/AudioMixer/Public/AudioMixerSubmix.h
--- a/Engine/Source/Runtime/AudioMixer/Public/AudioMixerSubmix.h
+++ b/Engine/Source/Runtime/AudioMixer/Public/AudioMixerSubmix.h
@@ -200,9 +200,12 @@ namespace Audio
 		// Terminates whatever FFT Analyzer is being used for this submix.
 		void StopSpectrumAnalysis();
 
-		// Adds an envelope follower delegate
+		// Adds an spectral analysis delegate
 		void AddSpectralAnalysisDelegate(const FOnSubmixSpectralAnalysisBP& OnSubmixSpectralAnalysisBP, const TArray<FSoundSubmixSpectralAnalysisBandSettings>& InBandSettings, float UpdateRate);
 
+		// Removes an existing spectral analysis delegate
+		void RemoveSpectralAnalysisDelegate(const FOnSubmixSpectralAnalysisBP& OnSubmixSpectralAnalysisBP);
+
 		// Gets the most recent magnitude values for each corresponding value in InFrequencies (in Hz).
 		// This requires StartSpectrumAnalysis to be called first.
 		void GetMagnitudeForFrequencies(const TArray<float>& InFrequencies, TArray<float>& OutMagnitudes);
@@ -488,6 +491,7 @@ namespace Audio
 		{
 			FEnvelopeFollower EnvelopeFollower;
 			float BandFrequency = 0.0f;
+			float QFactor = 1.f;
 		};
 
 		struct FSpectrumAnalysisDelegateInfo
diff --git a/Engine/Source/Runtime/Engine/Classes/Sound/SoundSubmix.h b/Engine/Source/Runtime/Engine/Classes/Sound/SoundSubmix.h
--- a/Engine/Source/Runtime/Engine/Classes/Sound/SoundSubmix.h
+++ b/Engine/Source/Runtime/Engine/Classes/Sound/SoundSubmix.h
@@ -37,6 +37,7 @@ DECLARE_DYNAMIC_MULTICAST_DELEGATE_OneParam(FOnSubmixEnvelope, const TArray<floa
 
 DECLARE_DYNAMIC_MULTICAST_DELEGATE_OneParam(FOnSubmixSpectralAnalysis, const TArray<float>&, Magnitudes);
 
+
 UENUM(BlueprintType)
 enum class EFFTSize : uint8
 {
@@ -297,7 +298,7 @@ public:
 	void AddEnvelopeFollowerDelegate(const UObject* WorldContextObject, const FOnSubmixEnvelopeBP& OnSubmixEnvelopeBP);
 
 	/**
-	 *	Adds a spectral envelope delegate to receive notifications when this submix has spectrum analysis enabled.
+	 *	Adds a spectral analysis delegate to receive notifications when this submix has spectrum analysis enabled.
 	 *	@param	InBandsettings					The frequency bands to analyze and their envelope-following settings.
 	 *  @param  OnSubmixSpectralAnalysisBP		Event to fire when new spectral data is available.
 	 *	@param	UpdateRate						How often to retrieve the data from the spectral analyzer and broadcast the event. Max is 30 times per second.
@@ -305,6 +306,13 @@ public:
 	UFUNCTION(BlueprintCallable, Category = "Audio|Spectrum", meta = (WorldContext = "WorldContextObject"))
 	void AddSpectralAnalysisDelegate(const UObject* WorldContextObject, const TArray<FSoundSubmixSpectralAnalysisBandSettings>& InBandSettings, const FOnSubmixSpectralAnalysisBP& OnSubmixSpectralAnalysisBP, float UpdateRate);
 
+	/**
+	 *	Remove a spectral analysis delegate.
+	 *  @param  OnSubmixSpectralAnalysisBP		The event delegate to remove.
+	 */
+	UFUNCTION(BlueprintCallable, Category = "Audio|Spectrum", meta = (WorldContext = "WorldContextObject"))
+	void RemoveSpectralAnalysisDelegate(const UObject* WorldContextObject, const FOnSubmixSpectralAnalysisBP& OnSubmixSpectralAnalysisBP);
+
 	/** Start spectrum analysis of the audio output. */
 	UFUNCTION(BlueprintCallable, Category = "Audio|Analysis", meta = (WorldContext = "WorldContextObject", AdvancedDisplay = 1))
 	void StartSpectralAnalysis(const UObject* WorldContextObject, EFFTSize FFTSize = EFFTSize::DefaultSize, EFFTPeakInterpolationMethod InterpolationMethod = EFFTPeakInterpolationMethod::Linear, EFFTWindowType WindowType = EFFTWindowType::Hann, float HopSize = 0, EAudioSpectrumType SpectrumType = EAudioSpectrumType::MagnitudeSpectrum);
diff --git a/Engine/Source/Runtime/Engine/Classes/Sound/SoundSubmixSend.h b/Engine/Source/Runtime/Engine/Classes/Sound/SoundSubmixSend.h
--- a/Engine/Source/Runtime/Engine/Classes/Sound/SoundSubmixSend.h
+++ b/Engine/Source/Runtime/Engine/Classes/Sound/SoundSubmixSend.h
@@ -13,6 +13,22 @@
 // Forward Declarations
 class USoundSubmixBase;
 
+UENUM(BlueprintType)
+enum class EAudioSpectrumBandPresetType: uint8
+{
+	/** Band which contains frequencies generally related to kick drums. */
+	KickDrum,
+
+	/** Band which contains frequencies generally related to snare drums. */
+	SnareDrum,
+
+	/** Band which contains frequencies generally related to vocals. */
+	Voice,
+
+	/** Band which contains frequencies generally related to cymbals. */
+	Cymbals	
+};
+
 USTRUCT(BlueprintType)
 struct ENGINE_API FSoundSubmixSpectralAnalysisBandSettings
 {
@@ -29,13 +45,16 @@ struct ENGINE_API FSoundSubmixSpectralAnalysisBandSettings
 	// The release time for the FFT band interpolation for delegate callback
 	UPROPERTY(EditAnywhere, BlueprintReadWrite, Category = SubmixSpectralAnalysis, meta = (ClampMin = "0.0", UIMin = "10.0", UIMax = "10000.0"))
 	int32 ReleaseTimeMsec = 500;
+
+	// The ratio of the bandwidth divided by the center frequency. Only used when the spectral analysis type is set to Constant Q.
+	UPROPERTY(EditAnywhere, BlueprintReadWrite, AdvancedDisplay, Category = SubmixSpectralAnalysis, meta = (ClampMin = "0.001", UIMin = "0.1", UIMax = "100.0"))
+	float QFactor = 10.0f;
 };
 
 DECLARE_DYNAMIC_DELEGATE_OneParam(FOnSubmixEnvelopeBP, const TArray<float>&, Envelope);
 
 DECLARE_DYNAMIC_DELEGATE_OneParam(FOnSubmixSpectralAnalysisBP, const TArray<float>&, Magnitude);
 
-
 UENUM(BlueprintType)
 enum class EAudioRecordingExportType : uint8
 {
diff --git a/Engine/Source/Runtime/Engine/Private/AudioDevice.cpp b/Engine/Source/Runtime/Engine/Private/AudioDevice.cpp
--- a/Engine/Source/Runtime/Engine/Private/AudioDevice.cpp
+++ b/Engine/Source/Runtime/Engine/Private/AudioDevice.cpp
@@ -4700,6 +4700,11 @@ void FAudioDevice::AddSpectralAnalysisDelegate(USoundSubmix* InSubmix, const TAr
 	UE_LOG(LogAudio, Error, TEXT("Spectrum analysis of submixes only works with the audio mixer. Please run using -audiomixer or set INI file to use submix recording."));
 }
 
+void FAudioDevice::RemoveSpectralAnalysisDelegate(USoundSubmix* InSubmix, const FOnSubmixSpectralAnalysisBP& OnSubmixSpectralAnalysisBP)
+{
+	UE_LOG(LogAudio, Error, TEXT("Spectrum analysis of submixes only works with the audio mixer. Please run using -audiomixer or set INI file to use submix recording."));
+}
+
 void FAudioDevice::AddVirtualLoop(const FAudioVirtualLoop& InVirtualLoop)
 {
 	FAudioVirtualLoop VirtualLoop = InVirtualLoop;
diff --git a/Engine/Source/Runtime/Engine/Private/SoundSubmix.cpp b/Engine/Source/Runtime/Engine/Private/SoundSubmix.cpp
--- a/Engine/Source/Runtime/Engine/Private/SoundSubmix.cpp
+++ b/Engine/Source/Runtime/Engine/Private/SoundSubmix.cpp
@@ -276,6 +276,23 @@ void USoundSubmix::AddSpectralAnalysisDelegate(const UObject* WorldContextObject
 	}
 }
 
+void USoundSubmix::RemoveSpectralAnalysisDelegate(const UObject* WorldContextObject, const FOnSubmixSpectralAnalysisBP& OnSubmixSpectralAnalysisBP)
+{
+	if (!GEngine)
+	{
+		return;
+	}
+
+	// Find device for this specific audio recording thing.
+	if (UWorld* ThisWorld = GEngine->GetWorldFromContextObject(WorldContextObject, EGetWorldErrorMode::LogAndReturnNull))
+	{
+		if (FAudioDevice* AudioDevice = ThisWorld->GetAudioDeviceRaw())
+		{
+			AudioDevice->RemoveSpectralAnalysisDelegate(this, OnSubmixSpectralAnalysisBP);
+		}
+	}
+}
+
 void USoundSubmix::StartSpectralAnalysis(const UObject* WorldContextObject, EFFTSize FFTSize, EFFTPeakInterpolationMethod InterpolationMethod, EFFTWindowType WindowType, float HopSize, EAudioSpectrumType SpectrumType)
 {
 	if (!GEngine)
diff --git a/Engine/Source/Runtime/Engine/Public/AudioDevice.h b/Engine/Source/Runtime/Engine/Public/AudioDevice.h
--- a/Engine/Source/Runtime/Engine/Public/AudioDevice.h
+++ b/Engine/Source/Runtime/Engine/Public/AudioDevice.h
@@ -1309,6 +1309,7 @@ public:
 	virtual void GetMagnitudesForFrequencies(USoundSubmix* InSubmix, const TArray<float>& InFrequencies, TArray<float>& OutMagnitudes);
 	virtual void GetPhasesForFrequencies(USoundSubmix* InSubmix, const TArray<float>& InFrequencies, TArray<float>& OutPhases);
 	virtual void AddSpectralAnalysisDelegate(USoundSubmix* InSubmix, const TArray<FSoundSubmixSpectralAnalysisBandSettings>& InBandSettings, const FOnSubmixSpectralAnalysisBP& OnSubmixSpectralAnalysisBP, float UpdateRate);
+	virtual void RemoveSpectralAnalysisDelegate(USoundSubmix* InSubmix, const FOnSubmixSpectralAnalysisBP& OnSubmixSpectralAnalysisBP);
 
 
 protected:
