- Fix regression from which memcopied audio bus buffer data with different channel counts
- Deprecate bAutoDeactivateWhenSilent and allow user to set virtualization setting on SourceBuses
#rnx

diff --git a/Engine/Source/Runtime/AudioMixer/Private/AudioMixerBus.cpp b/Engine/Source/Runtime/AudioMixer/Private/AudioMixerBus.cpp
--- a/Engine/Source/Runtime/AudioMixer/Private/AudioMixerBus.cpp
+++ b/Engine/Source/Runtime/AudioMixer/Private/AudioMixerBus.cpp
@@ -2,6 +2,7 @@
 
 #include "AudioMixerBus.h"
 #include "AudioMixerSourceManager.h"
+#include "DSP/BufferVectorOperations.h"
 
 namespace Audio
 {
@@ -180,6 +181,65 @@ namespace Audio
 		}
 	}
 
+	void FMixerAudioBus::CopyCurrentBuffer(Audio::AlignedFloatBuffer& OutBuffer, int32 InNumFrames, int32 InNumOutputChannels, EMonoChannelUpmixMethod InMixMethod) const
+	{
+		const float* RESTRICT CurrentBuffer = GetCurrentBusBuffer();
+
+		if (InNumOutputChannels == 2 && NumChannels == 1)
+		{
+			switch (InMixMethod)
+			{
+			case EMonoChannelUpmixMethod::EqualPower:
+			{
+				static constexpr float Sqrt2Over2Gains[2] = { 0.707106781f, 0.707106781f };
+				MixMonoTo2ChannelsFast(CurrentBuffer, OutBuffer.GetData(), InNumFrames, Sqrt2Over2Gains);
+			}
+			break;
+
+			case EMonoChannelUpmixMethod::Linear:
+			{
+				static constexpr float HalfGains[2] = { 0.5f, 0.5f };
+				MixMonoTo2ChannelsFast(CurrentBuffer, OutBuffer.GetData(), InNumFrames, HalfGains);
+			}
+			break;
+
+			case EMonoChannelUpmixMethod::FullVolume:
+			default:
+				MixMonoTo2ChannelsFast(CurrentBuffer, OutBuffer.GetData(), InNumFrames);
+				break;
+			}
+		}
+		else if (InNumOutputChannels == 1 && NumChannels == 2)
+		{
+			BufferSum2ChannelToMonoFast(CurrentBuffer, OutBuffer.GetData(), InNumFrames);
+			switch (InMixMethod)
+			{
+			case EMonoChannelUpmixMethod::EqualPower:
+			{
+				static constexpr float Sqrt2Gain = 1.41421356f;
+				MultiplyBufferByConstantInPlace(OutBuffer, Sqrt2Gain);
+			}
+			break;
+
+			case EMonoChannelUpmixMethod::Linear:
+			{
+				static constexpr float DoubleGain = 2.0f;
+				MultiplyBufferByConstantInPlace(OutBuffer, DoubleGain);
+			}
+			break;
+
+			case EMonoChannelUpmixMethod::FullVolume:
+			default:
+				break;
+			}
+		}
+		// Copy into the pre-distance attenuation buffer ptr as channel count should match
+		else
+		{
+			check(InNumOutputChannels == NumChannels);
+			FMemory::Memcpy(OutBuffer.GetData(), CurrentBuffer, sizeof(float) * InNumFrames * InNumOutputChannels);
+		}
+	}
 
 	const float* FMixerAudioBus::GetCurrentBusBuffer() const
 	{
diff --git a/Engine/Source/Runtime/AudioMixer/Private/AudioMixerSourceManager.cpp b/Engine/Source/Runtime/AudioMixer/Private/AudioMixerSourceManager.cpp
--- a/Engine/Source/Runtime/AudioMixer/Private/AudioMixerSourceManager.cpp
+++ b/Engine/Source/Runtime/AudioMixer/Private/AudioMixerSourceManager.cpp
@@ -1582,11 +1582,9 @@ namespace Audio
 			if (bIsSourceBus)
 			{
 				// Get the source's rendered and mixed audio bus data
-				const TSharedPtr<FMixerAudioBus>* AudioBusPtr = AudioBuses.Find(SourceInfo.AudioBusId);
-				if (ensure(AudioBusPtr != nullptr))
+				const TSharedPtr<FMixerAudioBus> AudioBusPtr = AudioBuses.FindRef(SourceInfo.AudioBusId);
+				if (ensure(AudioBusPtr.IsValid()))
 				{
-					const float* RESTRICT BusBufferPtr = (*AudioBusPtr)->GetCurrentBusBuffer();
-
 					int32 NumFramesPlayed = NumOutputFrames;
 					if (SourceInfo.SourceBusDurationFrames != INDEX_NONE)
 					{
@@ -1599,9 +1597,7 @@ namespace Audio
 					}
 
 					SourceInfo.NumFramesPlayed += NumFramesPlayed;
-
-					// Simply copy into the pre distance attenuation buffer ptr
-					FMemory::Memcpy(PreDistanceAttenBufferPtr, BusBufferPtr, sizeof(float) * NumFramesPlayed * SourceInfo.NumInputChannels);
+					AudioBusPtr->CopyCurrentBuffer(SourceInfo.PreDistanceAttenuationBuffer, NumFramesPlayed, SourceInfo.NumInputChannels);
 				}
 			}
 			else
diff --git a/Engine/Source/Runtime/AudioMixer/Public/AudioMixerBus.h b/Engine/Source/Runtime/AudioMixer/Public/AudioMixerBus.h
--- a/Engine/Source/Runtime/AudioMixer/Public/AudioMixerBus.h
+++ b/Engine/Source/Runtime/AudioMixer/Public/AudioMixerBus.h
@@ -5,6 +5,7 @@
 #include "CoreMinimal.h"
 #include "Audio.h"
 #include "DSP/MultithreadedPatching.h"
+#include "Sound/AudioSettings.h"
 
 namespace Audio
 {
@@ -67,6 +68,9 @@ namespace Audio
 		// Compute the mixed buffer
 		void MixBuffer();
 
+		// Copies the current internal buffer to a provided output buffer. Only supports mono or stereo input/output formats.
+		void CopyCurrentBuffer(AlignedFloatBuffer& OutBuffer, int32 InNumFrames, int32 InNumOutputChannels, EMonoChannelUpmixMethod InMixMethod = EMonoChannelUpmixMethod::EqualPower) const;
+
 		// If this bus was constructed before
 		void SetNumOutputChannels(int32 InNumOutputChannels);
 
diff --git a/Engine/Source/Runtime/AudioMixer/Public/AudioMixerDevice.h b/Engine/Source/Runtime/AudioMixer/Public/AudioMixerDevice.h
--- a/Engine/Source/Runtime/AudioMixer/Public/AudioMixerDevice.h
+++ b/Engine/Source/Runtime/AudioMixer/Public/AudioMixerDevice.h
@@ -226,6 +226,8 @@ namespace Audio
 		double GetAudioRenderThreadTime() const { return AudioThreadTimingData.AudioRenderThreadTime; }
 		double GetAudioClockDelta() const { return AudioClockDelta; }
 
+		EMonoChannelUpmixMethod GetMonoChannelUpmixMethod() const { return MonoChannelUpmixMethod; }
+
 		TArray<Audio::FChannelPositionInfo>* GetDefaultPositionMap(int32 NumChannels);
 
 		static bool IsEndpointSubmix(const USoundSubmixBase* InSubmix);
diff --git a/Engine/Source/Runtime/Engine/Classes/Sound/SoundSourceBus.h b/Engine/Source/Runtime/Engine/Classes/Sound/SoundSourceBus.h
--- a/Engine/Source/Runtime/Engine/Classes/Sound/SoundSourceBus.h
+++ b/Engine/Source/Runtime/Engine/Classes/Sound/SoundSourceBus.h
@@ -42,7 +42,7 @@ public:
 	UAudioBus* AudioBus;
 
 	/** Stop the source bus when the volume goes to zero. */
-	UPROPERTY(EditAnywhere, Category = BusProperties)
+	UPROPERTY(meta = (DeprecatedProperty))
 	uint32 bAutoDeactivateWhenSilent:1;
 
 	//~ Begin UObject Interface.
@@ -59,7 +59,4 @@ public:
 
 protected:
 	void Init();
-
-	uint32 DurationSamples;
-	bool bInitialized;
 };
diff --git a/Engine/Source/Runtime/Engine/Private/SoundSourceBus.cpp b/Engine/Source/Runtime/Engine/Private/SoundSourceBus.cpp
--- a/Engine/Source/Runtime/Engine/Private/SoundSourceBus.cpp
+++ b/Engine/Source/Runtime/Engine/Private/SoundSourceBus.cpp
@@ -32,10 +32,6 @@ void USoundSourceBus::Init()
 	// This sound wave is looping if the source bus duration is 0.0f
 	bLooping = (SourceBusDuration == 0.0f);
 
-	// Keep playing this bus when the volume is 0
-	// Note source buses can't ever be truly virtual as they are procedurally generated.
-	VirtualizationMode = bAutoDeactivateWhenSilent ? EVirtualizationMode::Disabled : EVirtualizationMode::PlayWhenSilent;
-
 	// Set the channels equal to the users channel count choice
 	switch (SourceBusChannels)
 	{
