Fixing reverb dry level

diff --git a/Engine/Source/Runtime/AudioMixer/Classes/SubmixEffects/AudioMixerSubmixEffectReverbFast.h b/Engine/Source/Runtime/AudioMixer/Classes/SubmixEffects/AudioMixerSubmixEffectReverbFast.h
--- a/Engine/Source/Runtime/AudioMixer/Classes/SubmixEffects/AudioMixerSubmixEffectReverbFast.h
+++ b/Engine/Source/Runtime/AudioMixer/Classes/SubmixEffects/AudioMixerSubmixEffectReverbFast.h
@@ -119,6 +119,9 @@ public:
 		return true;
 	}
 
+	// Returns the drylevel of the effect
+	virtual float GetDryLevel() const override { return CurrentWetDry.DryLevel; }
+
 private:
 
 	static const float MinWetness;
diff --git a/Engine/Source/Runtime/AudioMixer/Private/Effects/AudioMixerSubmixEffectReverbFast.cpp b/Engine/Source/Runtime/AudioMixer/Private/Effects/AudioMixerSubmixEffectReverbFast.cpp
--- a/Engine/Source/Runtime/AudioMixer/Private/Effects/AudioMixerSubmixEffectReverbFast.cpp
+++ b/Engine/Source/Runtime/AudioMixer/Private/Effects/AudioMixerSubmixEffectReverbFast.cpp
@@ -143,10 +143,6 @@ void FSubmixEffectReverbFast::OnProcessAudio(const FSoundEffectSubmixInputData&
 	{
 		LastWet = CurrentWetDry.WetLevel;
 	}
-	if (LastDry < 0.0f)
-	{
-		LastDry = CurrentWetDry.DryLevel;
-	}
 
 
 	WetInputBuffer.Reset();
@@ -159,7 +155,6 @@ void FSubmixEffectReverbFast::OnProcessAudio(const FSoundEffectSubmixInputData&
 
 	PlateReverb->ProcessAudio(WetInputBuffer, InData.NumChannels, *OutData.AudioBuffer, OutData.NumChannels);
 
-	Audio::MixInBufferFast(*InData.AudioBuffer, *OutData.AudioBuffer, LastDry, CurrentWetDry.DryLevel);
 }
 
 bool FSubmixEffectReverbFast::SetParameters(const FAudioEffectParameters& InParams)
