Fix for dynamics processor

diff --git a/Engine/Source/Runtime/AudioMixer/Classes/SubmixEffects/AudioMixerSubmixEffectDynamicsProcessor.h b/Engine/Source/Runtime/AudioMixer/Classes/SubmixEffects/AudioMixerSubmixEffectDynamicsProcessor.h
--- a/Engine/Source/Runtime/AudioMixer/Classes/SubmixEffects/AudioMixerSubmixEffectDynamicsProcessor.h
+++ b/Engine/Source/Runtime/AudioMixer/Classes/SubmixEffects/AudioMixerSubmixEffectDynamicsProcessor.h
@@ -89,7 +89,7 @@ struct AUDIOMIXER_API FSubmixEffectDynamicsProcessorSettings
 	ESubmixEffectDynamicsChannelLinkMode LinkMode;
 
 	// The input gain of the dynamics processor
-	UPROPERTY(EditAnywhere, BlueprintReadWrite, Category = General, meta = (DisplayName = "Input Gain (dB)", ClampMin = "-12.0", ClampMax = "20.0", UIMin = "-12.0", UIMax = "20.0"))
+	UPROPERTY(EditAnywhere, BlueprintReadWrite, Category = General, meta = (DisplayName = "Input Gain (dB)", UIMin = "-12.0", UIMax = "20.0"))
 	float InputGainDb;
 
 	// The threshold at which to perform a dynamics processing operation
@@ -134,11 +134,11 @@ struct AUDIOMIXER_API FSubmixEffectDynamicsProcessorSettings
 	uint8 bKeyAudition : 1;
 
 	// Gain to apply to key signal (external signal if supplied or input signal if disabled)
-	UPROPERTY(EditAnywhere, BlueprintReadWrite, Category = Sidechain, meta = (DisplayName = "Key Gain (dB)", ClampMin = "-60.0", ClampMax = "20.0", UIMin = "-60.0", UIMax = "20.0"))
+	UPROPERTY(EditAnywhere, BlueprintReadWrite, Category = Sidechain, meta = (DisplayName = "Key Gain (dB)", UIMin = "-60.0", UIMax = "30.0"))
 	float KeyGainDb;
 
 	// The output gain of the dynamics processor
-	UPROPERTY(EditAnywhere, BlueprintReadWrite, Category = Output, meta = (DisplayName = "Output Gain (dB)", ClampMin = "-60.0", ClampMax = "0.0", UIMin = "-60.0", UIMax = "20.0"))
+	UPROPERTY(EditAnywhere, BlueprintReadWrite, Category = Output, meta = (DisplayName = "Output Gain (dB)", UIMin = "-60.0", UIMax = "30.0"))
 	float OutputGainDb;
 
 	// High Shelf filter settings for key signal (external signal if supplied or input signal if not)
diff --git a/Engine/Source/Runtime/AudioMixer/Private/Effects/AudioMixerSubmixEffectDynamicsProcessor.cpp b/Engine/Source/Runtime/AudioMixer/Private/Effects/AudioMixerSubmixEffectDynamicsProcessor.cpp
--- a/Engine/Source/Runtime/AudioMixer/Private/Effects/AudioMixerSubmixEffectDynamicsProcessor.cpp
+++ b/Engine/Source/Runtime/AudioMixer/Private/Effects/AudioMixerSubmixEffectDynamicsProcessor.cpp
@@ -159,6 +159,7 @@ void FSubmixEffectDynamicsProcessor::OnProcessAudio(const FSoundEffectSubmixInpu
 			}
 			else
 			{
+				DynamicsProcessor.SetKeyGain(0.0f);
 				AudioKeyFrame[Channel] = InBuffer[SampleIndex];
 			}
 		}
@@ -209,7 +210,7 @@ void FSubmixEffectDynamicsProcessor::SetExternalSubmix(USoundSubmix* InSoundSubm
 	else
 	{
 		DeviceCreatedHandle = FAudioDeviceManagerDelegates::OnAudioDeviceCreated.AddRaw(this, &FSubmixEffectDynamicsProcessor::OnNewDeviceCreated);
-		bUseExternalSubmix = true;
+		bUseExternalSubmix = false;
 	}
 }
 
