Remove old ProcessControl code path no longer used by effect system (now done through direct handles to modulators)
#rnx

diff --git a/Engine/Source/Runtime/AudioMixer/Private/AudioMixerDevice.cpp b/Engine/Source/Runtime/AudioMixer/Private/AudioMixerDevice.cpp
--- a/Engine/Source/Runtime/AudioMixer/Private/AudioMixerDevice.cpp
+++ b/Engine/Source/Runtime/AudioMixer/Private/AudioMixerDevice.cpp
@@ -998,11 +998,6 @@ namespace Audio
 		});
 	}
 
-	void FMixerDevice::UpdateModulationControls(const uint32 InSourceId, const FSoundModulationControls& InControls)
-	{
-		SourceManager->UpdateModulationControls(InSourceId, InControls);
-	}
-
 	void FMixerDevice::UpdateSourceEffectChain(const uint32 SourceEffectChainId, const TArray<FSourceEffectChainEntry>& SourceEffectChain, const bool bPlayEffectChainTails)
 	{
 		TArray<FSourceEffectChainEntry>* ExistingOverride = SourceEffectChainOverrides.Find(SourceEffectChainId);
diff --git a/Engine/Source/Runtime/AudioMixer/Private/AudioMixerSource.cpp b/Engine/Source/Runtime/AudioMixer/Private/AudioMixerSource.cpp
--- a/Engine/Source/Runtime/AudioMixer/Private/AudioMixerSource.cpp
+++ b/Engine/Source/Runtime/AudioMixer/Private/AudioMixerSource.cpp
@@ -860,12 +860,7 @@ namespace Audio
 		if (AudioDevice->IsModulationPluginEnabled())
 		{
 			const int32 SourceId = MixerSourceVoice->GetSourceId();
-			const bool bUpdatePending = AudioDevice->ModulationInterface->ProcessControls(SourceId, WaveInstance->SoundModulationControls);
-
-			if (bUpdatePending)
-			{
-				AudioDevice->UpdateModulationControls(SourceId, WaveInstance->SoundModulationControls);
-			}
+			AudioDevice->ModulationInterface->ProcessControls(SourceId, WaveInstance->SoundModulationControls);
 		}
 	}
 
diff --git a/Engine/Source/Runtime/AudioMixer/Private/AudioMixerSourceManager.cpp b/Engine/Source/Runtime/AudioMixer/Private/AudioMixerSourceManager.cpp
--- a/Engine/Source/Runtime/AudioMixer/Private/AudioMixerSourceManager.cpp
+++ b/Engine/Source/Runtime/AudioMixer/Private/AudioMixerSourceManager.cpp
@@ -232,7 +232,6 @@ namespace Audio
 			SourceInfo.bIsVorbis = false;
 			SourceInfo.bIsBypassingLPF = false;
 			SourceInfo.bIsBypassingHPF = false;
-			SourceInfo.bIsModulationUpdated = false;
 
 #if AUDIO_MIXER_ENABLE_DEBUG_MODE
 			SourceInfo.bIsDebugMode = false;
@@ -1176,22 +1175,6 @@ namespace Audio
 		});
 	}
 
-	void FMixerSourceManager::UpdateModulationControls(const int32 SourceId, const FSoundModulationControls& InControls)
-	{
-		AUDIO_MIXER_CHECK(SourceId < NumTotalSources);
-		AUDIO_MIXER_CHECK(GameThreadInfo.bIsBusy[SourceId]);
-		AUDIO_MIXER_CHECK_GAME_THREAD(MixerDevice);
-
-		const FSoundModulationControls UpdatedControls = InControls;
-		AudioMixerThreadCommand([this, SourceId, UpdatedControls]()
-		{
-			AUDIO_MIXER_CHECK_AUDIO_PLAT_THREAD(MixerDevice);
-
-			SourceInfos[SourceId].ModulationControls = UpdatedControls;
-			SourceInfos[SourceId].bIsModulationUpdated = true;
-		});
-	}
-
 	void FMixerSourceManager::SetSpatializationParams(const int32 SourceId, const FSpatializationParams& InParams)
 	{
 		AUDIO_MIXER_CHECK(SourceId < NumTotalSources);
@@ -2246,13 +2229,6 @@ namespace Audio
 							bPresetUpdated = SoundEffectSource->Update();
 						}
 
-						// Modulation must be updated regardless of whether or not the source is
-						// active to allow for initial conditions to be set if source is reactivated.
-						if (bPresetUpdated || SourceInfo.bIsModulationUpdated)
-						{
-							SoundEffectSource->ProcessControls(SourceInfo.ModulationControls);
-						}
-
 						if (SoundEffectSource->IsActive())
 						{
 							SoundEffectSource->ProcessAudio(SourceInfo.SourceEffectInputData, OutputSourceEffectBufferPtr);
@@ -2262,8 +2238,6 @@ namespace Audio
 						}
 					}
 				}
-
-				SourceInfo.bIsModulationUpdated = false;
 			}
 
 			const bool bWasEffectTailsDone = SourceInfo.bEffectTailsDone;
@@ -2339,8 +2313,6 @@ namespace Audio
 				SourceInfo.NextFrameValues.Reset();
 				SourceInfo.CurrentPCMBuffer = nullptr;
 			}
-
-			SourceInfo.bIsModulationUpdated = false;
 		}
 	}
 
diff --git a/Engine/Source/Runtime/AudioMixer/Private/AudioMixerSourceManager.h b/Engine/Source/Runtime/AudioMixer/Private/AudioMixerSourceManager.h
--- a/Engine/Source/Runtime/AudioMixer/Private/AudioMixerSourceManager.h
+++ b/Engine/Source/Runtime/AudioMixer/Private/AudioMixerSourceManager.h
@@ -236,7 +236,6 @@ namespace Audio
 		void UpdateDeviceChannelCount(const int32 InNumOutputChannels);
 
 		void UpdateSourceEffectChain(const uint32 SourceEffectChainId, const TArray<FSourceEffectChainEntry>& SourceEffectChain, const bool bPlayEffectChainTails);
-		void UpdateModulationControls(const int32 SourceId, const FSoundModulationControls& InControls);
 
 		const float* GetPreDistanceAttenuationBuffer(const int32 SourceId) const;
 		const float* GetPreEffectBuffer(const int32 SourceId) const;
@@ -511,7 +510,6 @@ namespace Audio
 			uint8 bIsAmbisonics:1;
 			uint8 bIsBypassingLPF:1;
 			uint8 bIsBypassingHPF:1;
-			uint8 bIsModulationUpdated:1;
 
 			// Source format info
 			int32 NumInputChannels;
diff --git a/Engine/Source/Runtime/AudioMixer/Public/AudioMixerDevice.h b/Engine/Source/Runtime/AudioMixer/Public/AudioMixerDevice.h
--- a/Engine/Source/Runtime/AudioMixer/Public/AudioMixerDevice.h
+++ b/Engine/Source/Runtime/AudioMixer/Public/AudioMixerDevice.h
@@ -97,9 +97,6 @@ namespace Audio
 		virtual void InitSoundEffectPresets() override;
 		virtual int32 GetNumActiveSources() const override;
 
-		// Updates the provided source's modulation control values
-		virtual void UpdateModulationControls(const uint32 InSourceId, const FSoundModulationControls& InControls) override;
-
 		// Updates the source effect chain (using unique object id). 
 		virtual void UpdateSourceEffectChain(const uint32 SourceEffectChainId, const TArray<FSourceEffectChainEntry>& SourceEffectChain, const bool bPlayEffectChainTails) override;
 		virtual bool GetCurrentSourceEffectChain(const uint32 SourceEffectChainId, TArray<FSourceEffectChainEntry>& OutCurrentSourceEffectChainEntries) override;
diff --git a/Engine/Source/Runtime/Engine/Classes/Sound/SoundEffectSource.h b/Engine/Source/Runtime/Engine/Classes/Sound/SoundEffectSource.h
--- a/Engine/Source/Runtime/Engine/Classes/Sound/SoundEffectSource.h
+++ b/Engine/Source/Runtime/Engine/Classes/Sound/SoundEffectSource.h
@@ -115,8 +115,5 @@ public:
 	/** Process the input block of audio. Called on audio thread. */
 	virtual void ProcessAudio(const FSoundEffectSourceInputData& InData, float* OutAudioBufferData) = 0;
 
-	/** Process modulation controls if enabled */
-	virtual void ProcessControls(const FSoundModulationControls& InControls) { }
-
 	friend class USoundEffectPreset;
 };
