Refactoring the downmix data to its own class which encapsulates some complex code.

Adding ability to do a pre-distance submix send.

diff --git a/Engine/Source/Runtime/AudioMixer/Private/AudioMixerSource.cpp b/Engine/Source/Runtime/AudioMixer/Private/AudioMixerSource.cpp
--- a/Engine/Source/Runtime/AudioMixer/Private/AudioMixerSource.cpp
+++ b/Engine/Source/Runtime/AudioMixer/Private/AudioMixerSource.cpp
@@ -118,7 +118,7 @@ namespace Audio
 			InitParams.SourceVoice = MixerSourceVoice;
 			InitParams.bUseHRTFSpatialization = UseObjectBasedSpatialization();
 			InitParams.bIsExternalSend = MixerDevice->bSpatializationIsExternalSend;
-			InitParams.bIsAmbisonics = WaveInstance->bIsAmbisonics && (WaveInstance->WaveData->NumChannels == 4);
+			InitParams.bIsSoundfield = WaveInstance->bIsAmbisonics && (WaveInstance->WaveData->NumChannels == 4);
 
 			if (WaveInstance->bIsAmbisonics && (WaveInstance->WaveData->NumChannels != 4))
 			{
@@ -199,6 +199,7 @@ namespace Audio
 
 					FMixerSourceSubmixSend SubmixSend;
 					SubmixSend.Submix = SubmixPtr;
+					SubmixSend.SubmixSendStage = EMixerSourceSubmixSendStage::PostDistanceAttenuation;
 					SubmixSend.SendLevel = 1.0f;
 					SubmixSend.bIsMainSend = true;
 					SubmixSend.SoundfieldFactory = MixerDevice->GetFactoryForSubmixInstance(SubmixSend.Submix);
@@ -212,6 +213,12 @@ namespace Audio
 					{
 						FMixerSourceSubmixSend SubmixSend;
 						SubmixSend.Submix = MixerDevice->GetSubmixInstance(SendInfo.SoundSubmix);
+
+						SubmixSend.SubmixSendStage = EMixerSourceSubmixSendStage::PostDistanceAttenuation;
+						if (SendInfo.SendStage == ESubmixSendStage::PreDistanceAttenuation)
+						{
+							SubmixSend.SubmixSendStage = EMixerSourceSubmixSendStage::PreDistanceAttenuation;
+						}
 						SubmixSend.SendLevel = SendInfo.SendLevel;
 						SubmixSend.bIsMainSend = false;
 						SubmixSend.SoundfieldFactory = MixerDevice->GetFactoryForSubmixInstance(SubmixSend.Submix);
@@ -247,6 +254,9 @@ namespace Audio
 			// Whether or not we're 3D
 			bIs3D = !UseObjectBasedSpatialization() && WaveInstance->GetUseSpatialization() && SoundBuffer->NumChannels < 3;
 
+			// Pass on the fact that we're 3D to the init params
+			InitParams.bIs3D = bIs3D;
+
 			// Grab the source's reverb plugin settings
 			InitParams.SpatializationPluginSettings = UseSpatializationPlugin() ? WaveInstance->SpatializationPluginSettings : nullptr;
 
diff --git a/Engine/Source/Runtime/AudioMixer/Private/AudioMixerSourceManager.cpp b/Engine/Source/Runtime/AudioMixer/Private/AudioMixerSourceManager.cpp
--- a/Engine/Source/Runtime/AudioMixer/Private/AudioMixerSourceManager.cpp
+++ b/Engine/Source/Runtime/AudioMixer/Private/AudioMixerSourceManager.cpp
@@ -169,13 +169,12 @@ namespace Audio
 
 		MixerSources.Init(nullptr, NumTotalSources);
 
-		// Populate downmix array:
-		DownmixDataArray.Reset();
+		// Populate output sources array with default data
+		SourceSubmixOutputBuffers.Reset();
 		for (int32 Index = 0; Index < NumTotalSources; Index++)
 		{
-			DownmixDataArray.Emplace(2, MixerDevice->GetNumDeviceChannels(), NumOutputFrames);
+			SourceSubmixOutputBuffers.Emplace(MixerDevice, 2, MixerDevice->GetNumDeviceChannels(), NumOutputFrames);
 		}
-		
 
 		SourceInfos.AddDefaulted(NumTotalSources);
 
@@ -232,6 +231,7 @@ namespace Audio
 			SourceInfo.bIsVorbis = false;
 			SourceInfo.bIsBypassingLPF = false;
 			SourceInfo.bIsBypassingHPF = false;
+			SourceInfo.bHasPreDistanceAttenuationSend = false;
 
 #if AUDIO_MIXER_ENABLE_DEBUG_MODE
 			SourceInfo.bIsDebugMode = false;
@@ -539,6 +539,7 @@ namespace Audio
 		SourceInfo.bOutputToBusOnly = false;
 		SourceInfo.bIsBypassingLPF = false;
 		SourceInfo.bIsBypassingHPF = false;
+		SourceInfo.bHasPreDistanceAttenuationSend = false;
 
 #if AUDIO_MIXER_ENABLE_DEBUG_MODE
 		SourceInfo.bIsDebugMode = false;
@@ -674,6 +675,7 @@ namespace Audio
 			SourceInfo.MixerSourceBuffer->Init();
 			SourceInfo.MixerSourceBuffer->OnBeginGenerate();
 
+			SourceInfo.bIs3D = InitParams.bIs3D;
 			SourceInfo.bIsPlaying = false;
 			SourceInfo.bIsPaused = false;
 			SourceInfo.bIsStopping = false;
@@ -685,7 +687,7 @@ namespace Audio
 			SourceInfo.bIsExternalSend = InitParams.bIsExternalSend;
 			SourceInfo.bIsVorbis = InitParams.bIsVorbis;
 			SourceInfo.AudioComponentID = InitParams.AudioComponentID;
-			SourceInfo.bIsAmbisonics = InitParams.bIsAmbisonics;
+			SourceInfo.bIsSoundfield = InitParams.bIsSoundfield;
 
 			// Call initialization from the render thread so anything wanting to do any initialization here can do so (e.g. procedural sound waves)
 			SourceInfo.SourceListener = InitParams.SourceListener;
@@ -815,12 +817,9 @@ namespace Audio
 			// Initialize a new downmix data:
 			check(SourceId < SourceInfos.Num());
 			const int32 SourceInputChannels = (SourceInfo.bUseHRTFSpatializer && !SourceInfo.bIsExternalSend) ? 2 : SourceInfo.NumInputChannels;
-			FSourceDownmixData& DownmixData = InitializeDownmixForSource(SourceId, SourceInputChannels, MixerDevice->GetDeviceOutputChannels(), NumOutputFrames);
 
-			if (SourceInfo.bIsAmbisonics)
-			{
-				DownmixData.AmbisonicsDecoder = CreateDefaultSourceAmbisonicsDecoder(MixerDevice);
-			}
+			// Collect the soundfield encoding keys we need to initialize with our output buffers
+			TArray<FMixerSubmixPtr> SoundfieldSubmixSends;
 
 			for (int32 i = 0; i < InitParams.SubmixSends.Num(); ++i)
 			{
@@ -830,51 +829,33 @@ namespace Audio
 				if (SubmixPtr.IsValid())
 				{
 					SourceInfo.SubmixSends.Add(MixerSubmixSend);
-					SubmixPtr->AddOrSetSourceVoice(InitParams.SourceVoice, MixerSubmixSend.SendLevel);
 
-					
-					if (SubmixPtr->IsSoundfieldSubmix())
+					if (MixerSubmixSend.SubmixSendStage == EMixerSourceSubmixSendStage::PreDistanceAttenuation)
 					{
-						FSoundfieldEncodingKey Key = SubmixPtr->GetKeyForSubmixEncoding();
-						FSubmixSoundfieldData& SubmixSoundfieldInfo = GetChannelInfoForFormat(Key, DownmixData);
-
-						// Set up encoder (or transcoder if this is an ambisonics source).
-						ISoundfieldFactory* Factory = SubmixPtr->GetSoundfieldFactory();
-						check(Factory);
-
-						FAudioPluginInitializationParams PluginInitParams = SubmixPtr->GetInitializationParamsForSoundfieldStream();
-						PluginInitParams.NumOutputChannels = SourceInputChannels;
-
-						SubmixSoundfieldInfo.EncoderSettings = SubmixPtr->GetSoundfieldSettings().Duplicate();
-
-						if (SourceInfo.bIsAmbisonics)
-						{
-							// If this soundfield submix can transcode our ambisonics format, set up a transcoder stream.
-							if (Factory->GetSoundfieldFormatName() == GetUnrealAmbisonicsFormatName())
-							{
-								SubmixSoundfieldInfo.bIsUnrealAmbisonicsSubmix = true;
-							}
-							else if (Factory->CanTranscodeFromSoundfieldFormat(GetUnrealAmbisonicsFormatName(), GetAmbisonicsSourceDefaultSettings()))
-							{
-								SubmixSoundfieldInfo.AmbiTranscoder = Factory->CreateTranscoderStream(GetUnrealAmbisonicsFormatName(), GetAmbisonicsSourceDefaultSettings(), Factory->GetSoundfieldFormatName(), *SubmixSoundfieldInfo.EncoderSettings, PluginInitParams);
-							}
-						}
-						else
-						{
-							check(SubmixSoundfieldInfo.EncoderSettings.IsValid());
-
-							SubmixSoundfieldInfo.Encoder = Factory->CreateEncoderStream(PluginInitParams, *SubmixSoundfieldInfo.EncoderSettings);
-						}
-
-						SubmixSoundfieldInfo.EncodedPacket = Factory->CreateEmptyPacket();
+						SourceInfo.bHasPreDistanceAttenuationSend = true;
 					}
-					else
+
+					SubmixPtr->AddOrSetSourceVoice(InitParams.SourceVoice, MixerSubmixSend.SendLevel, MixerSubmixSend.SubmixSendStage);
+					
+					if (SubmixPtr->IsSoundfieldSubmix())
 					{
-						FSubmixChannelData& SubmixChannelInfo = GetChannelInfoForDevice(DownmixData);
+						SoundfieldSubmixSends.Add(SubmixPtr);
 					}
 				}
 			}
 
+			// Initialize the submix output source for this source id
+			FMixerSourceSubmixOutputBuffer& SourceSubmixOutputBuffer = SourceSubmixOutputBuffers[SourceId];
+
+			FMixerSourceSubmixOutputBufferSettings SourceSubmixOutputResetSettings;
+			SourceSubmixOutputResetSettings.NumOutputChannels = MixerDevice->GetDeviceOutputChannels();
+			SourceSubmixOutputResetSettings.NumSourceChannels = SourceInputChannels;
+			SourceSubmixOutputResetSettings.SoundfieldSubmixSends = SoundfieldSubmixSends;
+			SourceSubmixOutputResetSettings.bIs3D = SourceInfo.bIs3D;
+			SourceSubmixOutputResetSettings.bIsSoundfield = SourceInfo.bIsSoundfield;
+
+			SourceSubmixOutputBuffer.Reset(SourceSubmixOutputResetSettings);
+
 #if AUDIO_MIXER_ENABLE_DEBUG_MODE
 			AUDIO_MIXER_CHECK(!SourceInfo.bIsDebugMode);
 			SourceInfo.bIsDebugMode = InitParams.bIsDebugMode;
@@ -1202,9 +1183,9 @@ namespace Audio
 			check(NumOutputFrames > 0);
 
 			FSourceInfo& SourceInfo = SourceInfos[SourceId];
-			FSourceDownmixData& DownmixData = DownmixDataArray[SourceId];
+			FMixerSourceSubmixOutputBuffer& SourceSubmixOutput = SourceSubmixOutputBuffers[SourceId];
 
-			if (DownmixData.NumInputChannels != NumInputChannels && !SourceInfo.bUseHRTFSpatializer)
+			if (SourceSubmixOutput.GetNumSourceChannels() != NumInputChannels && !SourceInfo.bUseHRTFSpatializer)
 			{
 				// This means that this source has been reinitialized as a different source while this command was in flight,
 				// In which case it is of no use to us. Exit.
@@ -1215,38 +1196,8 @@ namespace Audio
 			SourceInfo.bIs3D = bInIs3D;
 			SourceInfo.bIsCenterChannelOnly = bInIsCenterChannelOnly;
 
-			FSubmixChannelData& ChannelTypeInfo = GetChannelInfoForDevice(DownmixData);
-
-			// Fix up the channel map in case device output count changed
-			const uint32 ChannelMapSize = ChannelTypeInfo.ChannelMap.CopySize / sizeof(float);
-
-			// If this is true, then the device changed while the command was in-flight
-			if (ChannelMap.Num() != ChannelMapSize)
-			{
-				// todo: investigate turning this into a stack array
-				Audio::AlignedFloatBuffer NewChannelMap;
-
-				// If 3d then just zero it out, we'll get another channel map shortly
-				if (bInIs3D)
-				{
-					NewChannelMap.AddZeroed(ChannelMapSize);
-					GameThreadInfo.bNeedsSpeakerMap[SourceId] = true;
-				}
-				// Otherwise, get an appropriate channel map for the new device configuration
-				else
-				{
-					const uint32 NumOutputChannels = ChannelMapSize / NumInputChannels;
-					FMixerDevice::Get2DChannelMap(SourceInfo.bIsVorbis, NumInputChannels, NumOutputChannels, bInIsCenterChannelOnly, NewChannelMap);
-				}
-
-				// Make sure we've been flagged to be using this submix channel type entry
-				ChannelTypeInfo.ChannelMap.SetChannelMap(NewChannelMap.GetData());
-			}
-			else
-			{
-				GameThreadInfo.bNeedsSpeakerMap[SourceId] = false;
-				ChannelTypeInfo.ChannelMap.SetChannelMap(ChannelMap.GetData());
-			}
+			bool bNeedsSpeakerMap = SourceSubmixOutput.SetChannelMap(ChannelMap, bInIsCenterChannelOnly);
+			GameThreadInfo.bNeedsSpeakerMap[SourceId] = bNeedsSpeakerMap;
 		});
 	}
 
@@ -1291,16 +1242,27 @@ namespace Audio
 			if (InSubmixPtr.IsValid())
 			{
 				bool bIsNew = true;
+				
+				SourceInfo.bHasPreDistanceAttenuationSend = false;
 				for (FMixerSourceSubmixSend& SubmixSend : SourceInfo.SubmixSends)
 				{
 					FMixerSubmixPtr SubmixPtr = SubmixSend.Submix.Pin();
 					if (SubmixPtr.IsValid())
 					{
+						if (SubmixSend.SubmixSendStage)
+						{
+							SourceInfo.bHasPreDistanceAttenuationSend = true;
+						}
+					
 						if (SubmixPtr->GetId() == InSubmixPtr->GetId())
 						{
 							SubmixSend.SendLevel = InSubmixSend.SendLevel;
+							SubmixSend.SubmixSendStage = InSubmixSend.SubmixSendStage;
 							bIsNew = false;
-							break;
+							if (SourceInfo.bHasPreDistanceAttenuationSend)
+							{
+								break;
+							}
 						}
 					}
 				}
@@ -1309,8 +1271,14 @@ namespace Audio
 				{
 					SourceInfo.SubmixSends.Add(InSubmixSend);
 				}
+				
+				// If we don't have a pre-distance attenuation send, lets zero out the buffer so the output buffer stops doing math with it.
+				if (!SourceInfo.bHasPreDistanceAttenuationSend)
+				{
+					SourceSubmixOutputBuffers[SourceId].SetPreAttenuationSourceBuffer(nullptr);
+				}
 
-				InSubmixPtr->AddOrSetSourceVoice(MixerSources[SourceId], InSubmixSend.SendLevel);
+				InSubmixPtr->AddOrSetSourceVoice(MixerSources[SourceId], InSubmixSend.SendLevel, InSubmixSend.SubmixSendStage);
 			}
 		});
 	}
@@ -1704,7 +1672,7 @@ namespace Audio
 		SourceInfo.DistanceAttenuationSourceStart = SourceInfo.DistanceAttenuationSourceDestination;
 	}
 
-	void FMixerSourceManager::ComputePluginAudio(FSourceInfo& SourceInfo, FSourceDownmixData& DownmixData, int32 SourceId, int32 NumSamples)
+	void FMixerSourceManager::ComputePluginAudio(FSourceInfo& SourceInfo, FMixerSourceSubmixOutputBuffer& InSourceSubmixOutputBuffer, int32 SourceId, int32 NumSamples)
 	{
 		if (BypassAudioPluginsCvar)
 		{
@@ -1712,8 +1680,12 @@ namespace Audio
 			SourceInfo.NumPostEffectChannels = SourceInfo.NumInputChannels;
 
 			// Set the ptr to use for post-effect buffers:
-			DownmixData.PostEffectBuffers = &SourceInfo.SourceBuffer;
+			InSourceSubmixOutputBuffer.SetPostAttenuationSourceBuffer(&SourceInfo.SourceBuffer);
 
+			if (SourceInfo.bHasPreDistanceAttenuationSend)
+			{
+				InSourceSubmixOutputBuffer.SetPreAttenuationSourceBuffer(&SourceInfo.PreDistanceAttenuationBuffer);
+			}
 			return;
 		}
 
@@ -1743,8 +1715,7 @@ namespace Audio
 			if (!MixerDevice->bReverbIsExternalSend)
 			{
 				// Copy the reverb-processed data back to the source buffer
-				DownmixData.ReverbPluginOutputBuffer.Reset();
-				DownmixData.ReverbPluginOutputBuffer.Append(SourceInfo.AudioPluginOutputData.AudioBuffer);
+				InSourceSubmixOutputBuffer.CopyReverbPluginOutputData(SourceInfo.AudioPluginOutputData.AudioBuffer);
 				bShouldMixInReverb = true;
 			}
 		}
@@ -1772,7 +1743,7 @@ namespace Audio
 			// Copy the occlusion-processed data back to the source buffer and mix with the reverb plugin output buffer
 			if (bShouldMixInReverb)
 			{
-				const float* ReverbPluginOutputBufferPtr = DownmixData.ReverbPluginOutputBuffer.GetData();
+				const float* ReverbPluginOutputBufferPtr = InSourceSubmixOutputBuffer.GetReverbPluginOutputData();
 				const float* AudioPluginOutputDataPtr = SourceInfo.AudioPluginOutputData.AudioBuffer.GetData();
 
 				Audio::SumBuffers(ReverbPluginOutputBufferPtr, AudioPluginOutputDataPtr, PostDistanceAttenBufferPtr, NumSamples);
@@ -1784,7 +1755,7 @@ namespace Audio
 		}
 		else if (bShouldMixInReverb)
 		{
-			const float* ReverbPluginOutputBufferPtr = DownmixData.ReverbPluginOutputBuffer.GetData();
+			const float* ReverbPluginOutputBufferPtr = InSourceSubmixOutputBuffer.GetReverbPluginOutputData();
 			Audio::MixInBufferFast(ReverbPluginOutputBufferPtr, PostDistanceAttenBufferPtr, NumSamples);
 		}
 
@@ -1824,313 +1795,41 @@ namespace Audio
 				SourceInfo.NumPostEffectChannels = SourceInfo.NumInputChannels;
 
 				// Set the ptr to use for post-effect buffers rather than the plugin output data (since the plugin won't have output audio data)
-				DownmixData.PostEffectBuffers = &SourceInfo.SourceBuffer;
+				InSourceSubmixOutputBuffer.SetPostAttenuationSourceBuffer(&SourceInfo.SourceBuffer);
+
+				if (SourceInfo.bHasPreDistanceAttenuationSend)
+				{
+					InSourceSubmixOutputBuffer.SetPreAttenuationSourceBuffer(&SourceInfo.PreDistanceAttenuationBuffer);
+				}
 			}
 			else
 			{
 				// Otherwise, we are now a 2-channel file and should not be spatialized using normal 3d spatialization
 				SourceInfo.NumPostEffectChannels = 2;
-				DownmixData.PostEffectBuffers = &SourceInfo.AudioPluginOutputData.AudioBuffer;
-			}
-		}
-		else
-		{
-			// Otherwise our pre- and post-effect channels are the same as the input channels
-			SourceInfo.NumPostEffectChannels = SourceInfo.NumInputChannels;
-
-			// Set the ptr to use for post-effect buffers
-			DownmixData.PostEffectBuffers = &SourceInfo.SourceBuffer;
-		}
-	}
-
-	void FMixerSourceManager::ComputeDownmix3D(FSourceDownmixData& DownmixData, FMixerDevice* MixerDevice)
-	{
-		// This enormous switch statement handles using the correct function for a given number of input and output channels.
-		// For 3D sources, we interpolate from ChannelStartGains to ChannelDestinationGains.
-		if (DownmixData.AmbisonicsDecoder.IsValid())
-		{
-			FAmbisonicsSoundfieldBuffer AmbiBuffer;
-			AmbiBuffer.AudioBuffer = MoveTemp(*DownmixData.PostEffectBuffers);
-			AmbiBuffer.NumChannels = DownmixData.NumInputChannels;
-			AmbiBuffer.PreviousRotation = AmbiBuffer.Rotation;
-			AmbiBuffer.Rotation = DownmixData.SourceRotation;
-
-			DownmixData.PositionalData.NumChannels = DownmixData.NumDeviceChannels;
-			DownmixData.PositionalData.ChannelPositions = MixerDevice->GetDefaultPositionMap(DownmixData.NumDeviceChannels);
-
-			FSoundfieldDecoderInputData InputData =
-			{
-				AmbiBuffer,
-				DownmixData.PositionalData,
-				static_cast<int32>(DownmixData.PostEffectBuffers->Num() / DownmixData.NumInputChannels),
-				MixerDevice->GetSampleRate()
-			};
-			
-			FSoundfieldDecoderOutputData OutputData = { DownmixData.DeviceSubmixInfo.OutputBuffer };
-
-			DownmixData.AmbisonicsDecoder->Decode(InputData, OutputData);
-
-			// Move the encoded ambisonics source buffer back to PostEffectBuffers to prevent reallocation
-			*DownmixData.PostEffectBuffers = MoveTemp(AmbiBuffer.AudioBuffer);
-		}
-		else if (DownmixData.NumInputChannels == 1)
-		{
-			switch (DownmixData.NumDeviceChannels)
-			{
-			case 8:
-				Audio::MixMonoTo8ChannelsFast(*DownmixData.PostEffectBuffers, DownmixData.DeviceSubmixInfo.OutputBuffer, DownmixData.DeviceSubmixInfo.ChannelMap.ChannelStartGains, DownmixData.DeviceSubmixInfo.ChannelMap.ChannelDestinationGains);
-				break;
-			case 6:
-				Audio::MixMonoTo6ChannelsFast(*DownmixData.PostEffectBuffers, DownmixData.DeviceSubmixInfo.OutputBuffer, DownmixData.DeviceSubmixInfo.ChannelMap.ChannelStartGains, DownmixData.DeviceSubmixInfo.ChannelMap.ChannelDestinationGains);
-				break;
-			case 4:
-				Audio::MixMonoTo4ChannelsFast(*DownmixData.PostEffectBuffers, DownmixData.DeviceSubmixInfo.OutputBuffer, DownmixData.DeviceSubmixInfo.ChannelMap.ChannelStartGains, DownmixData.DeviceSubmixInfo.ChannelMap.ChannelDestinationGains);
-				break;
-			case 2:
-				Audio::MixMonoTo2ChannelsFast(*DownmixData.PostEffectBuffers, DownmixData.DeviceSubmixInfo.OutputBuffer, DownmixData.DeviceSubmixInfo.ChannelMap.ChannelStartGains, DownmixData.DeviceSubmixInfo.ChannelMap.ChannelDestinationGains);
-				break;
-			}
-		}
-		else if (DownmixData.NumInputChannels == 2)
-		{
-			switch (DownmixData.NumDeviceChannels)
-			{
-			case 8:
-				Audio::Mix2ChannelsTo8ChannelsFast(*DownmixData.PostEffectBuffers, DownmixData.DeviceSubmixInfo.OutputBuffer, DownmixData.DeviceSubmixInfo.ChannelMap.ChannelStartGains, DownmixData.DeviceSubmixInfo.ChannelMap.ChannelDestinationGains);
-				break;
-			case 6:
-				Audio::Mix2ChannelsTo6ChannelsFast(*DownmixData.PostEffectBuffers, DownmixData.DeviceSubmixInfo.OutputBuffer, DownmixData.DeviceSubmixInfo.ChannelMap.ChannelStartGains, DownmixData.DeviceSubmixInfo.ChannelMap.ChannelDestinationGains);
-				break;
-			case 4:
-				Audio::Mix2ChannelsTo4ChannelsFast(*DownmixData.PostEffectBuffers, DownmixData.DeviceSubmixInfo.OutputBuffer, DownmixData.DeviceSubmixInfo.ChannelMap.ChannelStartGains, DownmixData.DeviceSubmixInfo.ChannelMap.ChannelDestinationGains);
-				break;
-			case 2:
-				Audio::Mix2ChannelsTo2ChannelsFast(*DownmixData.PostEffectBuffers, DownmixData.DeviceSubmixInfo.OutputBuffer, DownmixData.DeviceSubmixInfo.ChannelMap.ChannelStartGains, DownmixData.DeviceSubmixInfo.ChannelMap.ChannelDestinationGains);
-				break;
-			}
-		}
-		else
-		{
-			// Use generic calls:
-			Audio::DownmixBuffer(DownmixData.NumInputChannels, DownmixData.NumDeviceChannels, *DownmixData.PostEffectBuffers, DownmixData.DeviceSubmixInfo.OutputBuffer, DownmixData.DeviceSubmixInfo.ChannelMap.ChannelStartGains, DownmixData.DeviceSubmixInfo.ChannelMap.ChannelDestinationGains);
-		}
-
-		DownmixData.DeviceSubmixInfo.ChannelMap.CopyDestinationToStart();
-	}
-
-	void FMixerSourceManager::ComputeDownmix2D(FSourceDownmixData& DownmixData, FMixerDevice* MixerDevice)
-	{
-		// This enormous switch statement handles using the correct function for a given number of input and output channels.
-		// For 2D sources, we just apply the gain matrix in ChannelDestionationGains with no interpolation.
-		if (DownmixData.AmbisonicsDecoder.IsValid())
-		{
-			FAmbisonicsSoundfieldBuffer AmbiBuffer;
-			AmbiBuffer.AudioBuffer = MoveTemp(*DownmixData.PostEffectBuffers);
-			AmbiBuffer.NumChannels = DownmixData.NumInputChannels;
-			AmbiBuffer.PreviousRotation = AmbiBuffer.Rotation;
-			AmbiBuffer.Rotation = DownmixData.SourceRotation;
-
-			DownmixData.PositionalData.NumChannels = DownmixData.NumDeviceChannels;
-			DownmixData.PositionalData.ChannelPositions = MixerDevice->GetDefaultPositionMap(DownmixData.NumDeviceChannels);
-
-			FSoundfieldDecoderInputData InputData =
-			{
-				AmbiBuffer,
-				DownmixData.PositionalData,
-				static_cast<int32>(DownmixData.PostEffectBuffers->Num() / DownmixData.NumInputChannels),
-				MixerDevice->GetSampleRate()
-			};
-
-			FSoundfieldDecoderOutputData OutputData = { DownmixData.DeviceSubmixInfo.OutputBuffer };
 
-			DownmixData.AmbisonicsDecoder->Decode(InputData, OutputData);
-
-			// Move the encoded ambisonics source buffer back to PostEffectBuffers to prevent reallocation
-			*DownmixData.PostEffectBuffers = MoveTemp(AmbiBuffer.AudioBuffer);
-		}
-		else if (DownmixData.NumInputChannels == 1)
-		{
-			switch (DownmixData.NumDeviceChannels)
-			{
-			case 8:
-				Audio::MixMonoTo8ChannelsFast(*DownmixData.PostEffectBuffers, DownmixData.DeviceSubmixInfo.OutputBuffer, DownmixData.DeviceSubmixInfo.ChannelMap.ChannelDestinationGains);
-				break;
-			case 6:
-				Audio::MixMonoTo6ChannelsFast(*DownmixData.PostEffectBuffers, DownmixData.DeviceSubmixInfo.OutputBuffer, DownmixData.DeviceSubmixInfo.ChannelMap.ChannelDestinationGains);
-				break;
-			case 4:
-				Audio::MixMonoTo4ChannelsFast(*DownmixData.PostEffectBuffers, DownmixData.DeviceSubmixInfo.OutputBuffer, DownmixData.DeviceSubmixInfo.ChannelMap.ChannelDestinationGains);
-				break;
-			case 2:
-				Audio::MixMonoTo2ChannelsFast(*DownmixData.PostEffectBuffers, DownmixData.DeviceSubmixInfo.OutputBuffer, DownmixData.DeviceSubmixInfo.ChannelMap.ChannelDestinationGains);
-				break;
-			}
-		}
-		else if (DownmixData.NumInputChannels == 2)
-		{
-			switch (DownmixData.NumDeviceChannels)
-			{
-			case 8:
-				Audio::Mix2ChannelsTo8ChannelsFast(*DownmixData.PostEffectBuffers, DownmixData.DeviceSubmixInfo.OutputBuffer, DownmixData.DeviceSubmixInfo.ChannelMap.ChannelDestinationGains);
-				break;
-			case 6:
-				Audio::Mix2ChannelsTo6ChannelsFast(*DownmixData.PostEffectBuffers, DownmixData.DeviceSubmixInfo.OutputBuffer, DownmixData.DeviceSubmixInfo.ChannelMap.ChannelDestinationGains);
-				break;
-			case 4:
-				Audio::Mix2ChannelsTo4ChannelsFast(*DownmixData.PostEffectBuffers, DownmixData.DeviceSubmixInfo.OutputBuffer, DownmixData.DeviceSubmixInfo.ChannelMap.ChannelDestinationGains);
-				break;
-			case 2:
-				Audio::Mix2ChannelsTo2ChannelsFast(*DownmixData.PostEffectBuffers, DownmixData.DeviceSubmixInfo.OutputBuffer, DownmixData.DeviceSubmixInfo.ChannelMap.ChannelDestinationGains);
-				break;
-			}
-
-		}
-		else
-		{
-			// Use generic calls:
-			Audio::DownmixBuffer(DownmixData.NumInputChannels, DownmixData.NumDeviceChannels, *DownmixData.PostEffectBuffers, DownmixData.DeviceSubmixInfo.OutputBuffer, DownmixData.DeviceSubmixInfo.ChannelMap.ChannelDestinationGains);
-		}
-	}
-
-	void FMixerSourceManager::EncodeToSoundfieldFormats(FSourceDownmixData& DownmixData, const FSourceInfo& InSource, FMixerDevice* InDevice)
-	{
-		const FSpatializationParams& SpatializationParams = InSource.SpatParams;
-
-		// First, build our SoundfieldSpeakerPositionalData out of our SpatializationParams.
-		DownmixData.PositionalData.NumChannels = DownmixData.NumInputChannels;
-		DownmixData.PositionalData.Rotation = InSource.SpatParams.EmitterWorldRotation;
-		
-		// Build our input channel positions from the Spatialization params if this is a 3D source, otherwise use the default channel positions.
-		DownmixData.InputChannelPositions.Reset();
-		if (InSource.bIs3D)
-		{
-			if (DownmixData.NumInputChannels == 1)
-			{
-				Audio::FChannelPositionInfo ChannelPosition;
-				ChannelPosition.Channel = EAudioMixerChannel::FrontCenter;
-				ConvertCartesianToSpherical(SpatializationParams.EmitterPosition, ChannelPosition.Azimuth, ChannelPosition.Elevation, ChannelPosition.Radius);
-
-				DownmixData.InputChannelPositions.Add(ChannelPosition);
-
-				DownmixData.PositionalData.ChannelPositions = &DownmixData.InputChannelPositions;
-			}
-			else if (DownmixData.NumInputChannels == 2)
-			{
-				Audio::FChannelPositionInfo LeftChannelPosition;
-				LeftChannelPosition.Channel = EAudioMixerChannel::FrontLeft;
-				ConvertCartesianToSpherical(SpatializationParams.LeftChannelPosition, LeftChannelPosition.Azimuth, LeftChannelPosition.Elevation, LeftChannelPosition.Radius);
-
-				DownmixData.InputChannelPositions.Add(LeftChannelPosition);
-
-				Audio::FChannelPositionInfo RightChannelPosition;
-				LeftChannelPosition.Channel = EAudioMixerChannel::FrontRight;
-				ConvertCartesianToSpherical(SpatializationParams.RightChannelPosition, RightChannelPosition.Azimuth, RightChannelPosition.Elevation, RightChannelPosition.Radius);
-
-				DownmixData.InputChannelPositions.Add(RightChannelPosition);
+				// Set the ptr to use for post-effect buffers rather than the plugin output data (since the plugin won't have output audio data)
+				InSourceSubmixOutputBuffer.SetPostAttenuationSourceBuffer(&SourceInfo.AudioPluginOutputData.AudioBuffer);
 
-				DownmixData.PositionalData.ChannelPositions = &DownmixData.InputChannelPositions;
-			}
-			else
-			{
-				// Spatialization of multichannel audio beyond stereo is not currently supported in the engine.
-				DownmixData.PositionalData.ChannelPositions = InDevice->GetDefaultPositionMap(DownmixData.NumInputChannels);
+				if (SourceInfo.bHasPreDistanceAttenuationSend)
+				{
+					InSourceSubmixOutputBuffer.SetPreAttenuationSourceBuffer(&SourceInfo.PreDistanceAttenuationBuffer);
+				}
 			}
 		}
 		else
 		{
-			DownmixData.PositionalData.ChannelPositions = InDevice->GetDefaultPositionMap(DownmixData.NumInputChannels);
-		}
-
-		// Finally, run our encoders.
-		for (auto& Soundfield : DownmixData.EncodedSoundfieldDownmixes)
-		{
-			FSubmixSoundfieldData& SoundfieldData = Soundfield.Value;
-
-			check(SoundfieldData.EncoderSettings.IsValid());
-			check(SoundfieldData.EncodedPacket.IsValid());
-
-			SoundfieldData.EncodedPacket->Reset();
-
-			// If this is an ambisonics source, transcode it. Otherwise, encode the source.
-			if (SoundfieldData.AmbiTranscoder)
-			{
-				FAmbisonicsSoundfieldBuffer AmbiBuffer;
-				AmbiBuffer.AudioBuffer = MoveTemp(*DownmixData.PostEffectBuffers);
-				AmbiBuffer.NumChannels = DownmixData.NumInputChannels;
-				AmbiBuffer.PreviousRotation = AmbiBuffer.Rotation;
-				AmbiBuffer.Rotation = DownmixData.SourceRotation;
+			// Otherwise our pre- and post-effect channels are the same as the input channels
+			SourceInfo.NumPostEffectChannels = SourceInfo.NumInputChannels;
 
-				SoundfieldData.AmbiTranscoder->Transcode(AmbiBuffer, GetAmbisonicsSourceDefaultSettings(), *SoundfieldData.EncodedPacket, *SoundfieldData.EncoderSettings);
-				*DownmixData.PostEffectBuffers = MoveTemp(AmbiBuffer.AudioBuffer);
-			}
-			else if (SoundfieldData.Encoder)
-			{
-				FSoundfieldEncoderInputData InputData =
-				{
-					*DownmixData.PostEffectBuffers,
-					static_cast<int32>(DownmixData.NumInputChannels),
-					*SoundfieldData.EncoderSettings,
-					DownmixData.PositionalData
-				};
+			InSourceSubmixOutputBuffer.SetPostAttenuationSourceBuffer(&SourceInfo.SourceBuffer);
 
-				SoundfieldData.Encoder->Encode(InputData, *SoundfieldData.EncodedPacket);
-			}
-			else if (SoundfieldData.bIsUnrealAmbisonicsSubmix)
+			if (SourceInfo.bHasPreDistanceAttenuationSend)
 			{
-				ensure(InSource.bIsAmbisonics);
-
-				FAmbisonicsSoundfieldBuffer& OutputPacket = DowncastSoundfieldRef<FAmbisonicsSoundfieldBuffer>(*SoundfieldData.EncodedPacket);
-				// Fixme: This is an array copy. Can we serve DownmixData.PostEffectBuffers directly to this soundfield?
-				OutputPacket.AudioBuffer = *DownmixData.PostEffectBuffers;
-				OutputPacket.NumChannels = DownmixData.NumInputChannels;
-				OutputPacket.PreviousRotation = OutputPacket.Rotation;
-				OutputPacket.Rotation = DownmixData.SourceRotation;
+				InSourceSubmixOutputBuffer.SetPreAttenuationSourceBuffer(&SourceInfo.PreDistanceAttenuationBuffer);
 			}
 		}
 	}
 
-	void FMixerSourceManager::ConvertCartesianToSpherical(const FVector& InVector, float& OutAzimuth, float& OutElevation, float& OutRadius)
-	{
-		// Convert coordinates from unreal cartesian system to left handed spherical coordinates (zenith is positive elevation, right is positive azimuth)
-		const float InX = -InVector.Z; //InVector.Y;
-		const float InY = InVector.X;// -InVector.Z;
-		const float InZ = -InVector.Y;
-
-
-		OutElevation = FMath::Atan2(InY, InX);
-
-		// Note, rather than using arccos(z / radius) here, we use Atan2 to avoid wrapping issues with negative elevation values.
-		OutAzimuth = FMath::Atan2(FMath::Sqrt(InX * InX + InY * InY), InZ);
-		OutRadius = InVector.Size();
-	}
-
-	FMixerSourceManager::FSourceDownmixData& FMixerSourceManager::InitializeDownmixForSource(const int32 SourceId, const int32 NumInputChannels, const int32 NumOutputChannels, const int32 InNumOutputFrames)
-	{
-		DownmixDataArray[SourceId].ResetData(NumInputChannels, NumOutputChannels);
-		return DownmixDataArray[SourceId];
-	}
-
-	const Audio::FMixerSourceManager::FSubmixChannelData& FMixerSourceManager::GetChannelInfoForDevice(const FSourceDownmixData& InDownmixData) const
-	{
-		return InDownmixData.DeviceSubmixInfo;
-	}
-
-	Audio::FMixerSourceManager::FSubmixChannelData& FMixerSourceManager::GetChannelInfoForDevice(FSourceDownmixData& InDownmixData)
-	{
-		return InDownmixData.DeviceSubmixInfo;
-	}
-
-	const FMixerSourceManager::FSubmixSoundfieldData& FMixerSourceManager::GetChannelInfoForFormat(const FSoundfieldEncodingKey& InFormat, const FSourceDownmixData& InDownmixData) const
-	{
-		check(InDownmixData.EncodedSoundfieldDownmixes.Contains(InFormat));
-		return InDownmixData.EncodedSoundfieldDownmixes[InFormat];
-	}
-
-	FMixerSourceManager::FSubmixSoundfieldData& FMixerSourceManager::GetChannelInfoForFormat(const FSoundfieldEncodingKey& InFormat, FSourceDownmixData& InDownmixData)
-	{
-		return InDownmixData.EncodedSoundfieldDownmixes.FindOrAdd(InFormat);
-	}
-
 	void FMixerSourceManager::ComputePostSourceEffectBufferForIdRange(bool bGenerateBuses, const int32 SourceIdStart, const int32 SourceIdEnd)
 	{
 		CSV_SCOPED_TIMING_STAT(Audio, SourceEffectsBuffers);
@@ -2152,8 +1851,6 @@ namespace Audio
 				continue;
 			}
 
-			FSourceDownmixData& DownmixData = DownmixDataArray[SourceId];
-
 			// Copy and store the current state of the pre-distance attenuation buffer before we feed it through our source effects
 			// This is used by pre-effect sends
 			if (SourceInfo.AudioBusSends[(int32)EBusSendType::PreEffect].Num() > 0)
@@ -2301,8 +1998,10 @@ namespace Audio
 				// Apply distance attenuation
 				ApplyDistanceAttenuation(SourceInfo, NumSamples);
 
+				FMixerSourceSubmixOutputBuffer& SourceSubmixOutputBuffer = SourceSubmixOutputBuffers[SourceId];
+
 				// Send source audio to plugins
-				ComputePluginAudio(SourceInfo, DownmixData, SourceId, NumSamples);
+				ComputePluginAudio(SourceInfo, SourceSubmixOutputBuffer, SourceId, NumSamples);
 			}
 
 			// Check the source effect tails condition
@@ -2340,31 +2039,8 @@ namespace Audio
 				continue;
 			}
 
-			FSourceDownmixData& DownmixData = DownmixDataArray[SourceId];
-
-			if (DownmixData.PostEffectBuffers == nullptr)
-			{
-				continue;
-			}
-
-			DownmixData.PositionalData.Rotation = SourceInfo.SpatParams.ListenerOrientation;
-			DownmixData.SourceRotation = SourceInfo.SpatParams.EmitterWorldRotation;
-
-			if (SourceInfo.bIs3D && !DownmixData.bIsInitialDownmix)
-			{
-				ComputeDownmix3D(DownmixData, MixerDevice);
-			}
-			else
-			{
-				ComputeDownmix2D(DownmixData, MixerDevice);
-				DownmixData.bIsInitialDownmix = false;
-			}
-
-			if (DownmixData.EncodedSoundfieldDownmixes.Num())
-			{
-				// Perform All Encoding for all sends.
-				EncodeToSoundfieldFormats(DownmixData, SourceInfo, MixerDevice);
-			}
+			FMixerSourceSubmixOutputBuffer& SourceSubmixOutputBuffer = SourceSubmixOutputBuffers[SourceId];
+			SourceSubmixOutputBuffer.ComputeOutput(SourceInfo.SpatParams);
 		}
 	}
 
@@ -2411,26 +2087,17 @@ namespace Audio
 		}
 	}
 
-	void FMixerSourceManager::MixOutputBuffers(const int32 SourceId, int32 InNumOutputChannels, const float SendLevel, AlignedFloatBuffer& OutWetBuffer) const
+	void FMixerSourceManager::MixOutputBuffers(const int32 SourceId, int32 InNumOutputChannels, const float InSendLevel, EMixerSourceSubmixSendStage InSubmixSendStage, AlignedFloatBuffer& OutWetBuffer) const
 	{
-		if (SendLevel > 0.0f)
+		if (InSendLevel > 0.0f)
 		{
 			const FSourceInfo& SourceInfo = SourceInfos[SourceId];
 
 			// Don't need to mix into submixes if the source is paused
 			if (!SourceInfo.bIsPaused && !SourceInfo.bIsDone && SourceInfo.bIsPlaying)
 			{
-				const FSourceDownmixData& DownmixData = DownmixDataArray[SourceId];
-				const FSubmixChannelData& ChannelInfo = GetChannelInfoForDevice(DownmixData);
-
-				const float* RESTRICT SourceOutputBufferPtr = ChannelInfo.OutputBuffer.GetData();
-
-				// TODO: Figure out a fix for this race condition on device swap:
-				// const int32 OutWetBufferSize = OutWetBuffer.Num();
-				const int32 OutWetBufferSize = FMath::Min<float>(OutWetBuffer.Num(), ChannelInfo.OutputBuffer.Num());
-				float* RESTRICT OutWetBufferPtr = OutWetBuffer.GetData();
-
-				Audio::MixInBufferFast(SourceOutputBufferPtr, OutWetBufferPtr, OutWetBufferSize, SendLevel);
+				const FMixerSourceSubmixOutputBuffer& SourceSubmixOutputBuffer = SourceSubmixOutputBuffers[SourceId];
+				SourceSubmixOutputBuffer.MixOutput(InSendLevel, InSubmixSendStage, OutWetBuffer);
 			}
 		}
 	}
@@ -2444,10 +2111,8 @@ namespace Audio
 		// Don't need to mix into submixes if the source is paused
 		if (!SourceInfo.bIsPaused && !SourceInfo.bIsDone && SourceInfo.bIsPlaying)
 		{
-			const FSourceDownmixData& DownmixData = DownmixDataArray[SourceId];
-			const FSubmixSoundfieldData& ChannelInfo = GetChannelInfoForFormat(InKey, DownmixData);
-
-			return ChannelInfo.EncodedPacket.Get();
+			const FMixerSourceSubmixOutputBuffer& SourceSubmixOutputBuffer = SourceSubmixOutputBuffers[SourceId];
+			return SourceSubmixOutputBuffer.GetSoundfieldPacket(InKey);
 		}
 
 		return nullptr;
@@ -2455,8 +2120,8 @@ namespace Audio
 
 	const FQuat FMixerSourceManager::GetListenerRotation(const int32 SourceId) const
 	{
-		const FSourceDownmixData& DownmixData = DownmixDataArray[SourceId];
-		return DownmixData.PositionalData.Rotation;
+		const FMixerSourceSubmixOutputBuffer& SubmixOutputBuffer = SourceSubmixOutputBuffers[SourceId];
+		return SubmixOutputBuffer.GetListenerRotation();
 	}
 
 	void FMixerSourceManager::UpdateDeviceChannelCount(const int32 InNumOutputChannels)
@@ -2476,10 +2141,8 @@ namespace Audio
 					continue;
 				}
 
-				FSourceDownmixData& DownmixData = DownmixDataArray[SourceId];
-				DownmixData.ResetNumberOfDeviceChannels(InNumOutputChannels);
-
-				FSubmixChannelData& ChannelTypeInfo = GetChannelInfoForDevice(DownmixData);
+				FMixerSourceSubmixOutputBuffer& SourceSubmixOutputBuffer = SourceSubmixOutputBuffers[SourceId];
+				SourceSubmixOutputBuffer.SetNumOutputChannels(InNumOutputChannels);
 
 				SourceInfo.ScratchChannelMap.Reset();
 				const int32 NumSourceChannels = SourceInfo.bUseHRTFSpatializer ? 2 : SourceInfo.NumInputChannels;
@@ -2498,7 +2161,7 @@ namespace Audio
 					MixerDevice->Get2DChannelMap(SourceInfo.bIsVorbis, NumSourceChannels, InNumOutputChannels, SourceInfo.bIsCenterChannelOnly, SourceInfo.ScratchChannelMap);
 				}
 
-				ChannelTypeInfo.ChannelMap.SetChannelMap(SourceInfo.ScratchChannelMap.GetData());
+				SourceSubmixOutputBuffer.SetChannelMap(SourceInfo.ScratchChannelMap, SourceInfo.bIsCenterChannelOnly);
 			}
 		});
 	}
diff --git a/Engine/Source/Runtime/AudioMixer/Private/AudioMixerSourceManager.h b/Engine/Source/Runtime/AudioMixer/Private/AudioMixerSourceManager.h
--- a/Engine/Source/Runtime/AudioMixer/Private/AudioMixerSourceManager.h
+++ b/Engine/Source/Runtime/AudioMixer/Private/AudioMixerSourceManager.h
@@ -10,6 +10,7 @@
 #include "AudioMixerSubmix.h"
 #include "AudioMixerBus.h"
 #include "AudioMixerDevice.h"
+#include "AudioMixerSourceOutputBuffer.h"
 #include "DSP/InterpolatedOnePole.h"
 #include "DSP/Filter.h"
 #include "DSP/InterpolatedOnePole.h"
@@ -20,7 +21,6 @@
 #include "ISoundfieldFormat.h"
 #include "Containers/Queue.h"
 
-
 namespace Audio
 {
 	class FMixerSubmix;
@@ -28,6 +28,7 @@ namespace Audio
 	class FMixerSourceVoice;
 	class FMixerSourceBuffer;
 	class ISourceListener;
+	class FMixerSourceSubmixOutputBuffer;
 
 	/** Struct defining a source voice buffer. */
 	struct FMixerSourceVoiceBuffer
@@ -71,6 +72,9 @@ namespace Audio
 		// Whather or not this is the primary send (i.e. first in the send chain)
 		bool bIsMainSend = false;
 
+		// Whether or not this is a pre-distance attenuation send
+		EMixerSourceSubmixSendStage SubmixSendStage = EMixerSourceSubmixSendStage::PostDistanceAttenuation;
+
 		// If this is a soundfield submix, this is a pointer to the submix's Soundfield Factory.
 		// If this is nullptr, the submix is not a soundfield submix.
 		ISoundfieldFactory* SoundfieldFactory = nullptr;
@@ -105,68 +109,17 @@ namespace Audio
 		USoundModulationPluginSourceSettingsBase* ModulationPluginSettings = nullptr;
 		FName AudioComponentUserID;
 		uint64 AudioComponentID = 0;
+		bool bIs3D = false;
 		bool bPlayEffectChainTails = false;
 		bool bUseHRTFSpatialization = false;
 		bool bIsExternalSend = false;
 		bool bIsDebugMode  = false;
 		bool bOutputToBusOnly = false;
 		bool bIsVorbis = false;
-		bool bIsAmbisonics = false;
+		bool bIsSoundfield = false;
 		bool bIsSeeking = false;
 	};
 
-	struct FSourceChannelMap
-	{
-		alignas(16) float ChannelStartGains[AUDIO_MIXER_MAX_OUTPUT_CHANNELS * AUDIO_MIXER_MAX_OUTPUT_CHANNELS];
-		alignas(16) float ChannelDestinationGains[AUDIO_MIXER_MAX_OUTPUT_CHANNELS * AUDIO_MIXER_MAX_OUTPUT_CHANNELS];
-
-		// This is the number of bytes the gain array is using:
-		// (Number of input channels * number of output channels) * sizeof float.
-		int32 CopySize = 0;
-		bool bIsInit = false;
-
-		FSourceChannelMap(int32 InNumInChannels, int32 InNumOutChannels) 
-			: CopySize(InNumInChannels * InNumOutChannels * sizeof(float))
-		{
-			checkSlow(InNumInChannels <= AUDIO_MIXER_MAX_OUTPUT_CHANNELS);
-			checkSlow(InNumOutChannels <= AUDIO_MIXER_MAX_OUTPUT_CHANNELS);
-			FMemory::Memzero(ChannelStartGains, CopySize);
-		}
-
-		FORCEINLINE void Reset(int32 InNumInChannels, int32 InNumOutChannels)
-		{
-			checkSlow(InNumInChannels <= AUDIO_MIXER_MAX_OUTPUT_CHANNELS);
-			checkSlow(InNumOutChannels <= AUDIO_MIXER_MAX_OUTPUT_CHANNELS);
-
-			CopySize = InNumInChannels * InNumOutChannels * sizeof(float);
-			FMemory::Memzero(ChannelStartGains, CopySize);
-			FMemory::Memzero(ChannelDestinationGains, CopySize);
-			bIsInit = false;
-		}
-
-		FORCEINLINE void CopyDestinationToStart()
-		{
-			FMemory::Memcpy(ChannelStartGains, ChannelDestinationGains, CopySize);
-		}
-
-		FORCEINLINE void SetChannelMap(const float* RESTRICT InChannelGains)
-		{
-			FMemory::Memcpy(ChannelDestinationGains, InChannelGains, CopySize);
-			if (!bIsInit)
-			{
-				FMemory::Memcpy(ChannelStartGains, InChannelGains, CopySize);
-				bIsInit = true;
-			}
-		}
-
-	private:
-		FSourceChannelMap()
-			: CopySize(0)
-			, bIsInit(false)
-		{
-		}
-	};
-
 	struct FSourceManagerInitParams
 	{
 		// Total number of sources to use in the source manager
@@ -219,7 +172,7 @@ namespace Audio
 		bool NeedsSpeakerMap(const int32 SourceId) const;
 		void ComputeNextBlockOfSamples();
 		void ClearStoppingSounds();
-		void MixOutputBuffers(const int32 SourceId, int32 InNumOutputChannels, const float SendLevel, AlignedFloatBuffer& OutWetBuffer) const;
+		void MixOutputBuffers(const int32 SourceId, int32 InNumOutputChannels, const float InSendLevel, EMixerSourceSubmixSendStage InSubmixSendStage, AlignedFloatBuffer& OutWetBuffer) const;
 
 		// Called by a soundfield submix to get encoded audio.
 		// If this source wasn't encoded (possibly because it is paused or finished playing),
@@ -248,7 +201,6 @@ namespace Audio
 		void UpdatePendingReleaseData(bool bForceWait = false);
 		void FlushCommandQueue(bool bPumpCommandQueue = false);
 	private:
-
 		void ReleaseSource(const int32 SourceId);
 		void BuildSourceEffectChain(const int32 SourceId, FSoundEffectSourceInitData& InitData, const TArray<FSourceEffectChainEntry>& SourceEffectChain);
 		void ResetSourceEffectChain(const int32 SourceId);
@@ -326,93 +278,6 @@ namespace Audio
 
 		TArray<int32> DebugSoloSources;
 
-		// All state having to do with a non-soundfield downmix.
-		struct FSubmixChannelData
-		{
-			FSourceChannelMap ChannelMap;
-			Audio::AlignedFloatBuffer OutputBuffer;
-
-			FSubmixChannelData(uint32 InNumInChannels, uint32 InNumOutputChannels, uint32 NumFrames)
-				: ChannelMap(InNumInChannels, InNumOutputChannels)
-			{
-				OutputBuffer.Reset();
-				OutputBuffer.AddUninitialized(NumFrames * InNumOutputChannels);
-			}
-
-			void Reset(uint32 InNumInChannels, uint32 InNumOutputChannels, uint32 NumFrames)
-			{
-				ChannelMap.Reset(InNumInChannels, InNumOutputChannels);
-
-				OutputBuffer.Reset();
-				OutputBuffer.AddUninitialized(NumFrames * InNumOutputChannels);
-			}
-		};
-
-		struct FSubmixSoundfieldData
-		{
-			TUniquePtr<ISoundfieldEncoderStream> Encoder;
-			// If this is an ambisonics source, we use a transcoder stream.
-			TUniquePtr<ISoundfieldTranscodeStream> AmbiTranscoder;
-			TUniquePtr<ISoundfieldEncodingSettingsProxy> EncoderSettings;
-			TUniquePtr<ISoundfieldAudioPacket> EncodedPacket;
-			// If this is a unreal ambisonics soundfield buffer, we hand it the submixed buffer directly.
-			bool bIsUnrealAmbisonicsSubmix;
-		};
-
-		struct FSourceDownmixData
-		{
-			// cached parameters for encoding to a soundfield format.
-			FSoundfieldSpeakerPositionalData PositionalData;
-			FQuat SourceRotation;
-
-			// Output data, after computing a block of sample data, this is read back from mixers
-			Audio::AlignedFloatBuffer ReverbPluginOutputBuffer;
-			Audio::AlignedFloatBuffer* PostEffectBuffers;
-
-			// Data needed for outputting to submixes for each channel configuration.
-			FSubmixChannelData DeviceSubmixInfo;
-			
-			TMap<FSoundfieldEncodingKey, FSubmixSoundfieldData> EncodedSoundfieldDownmixes;
-			TArray<Audio::FChannelPositionInfo> InputChannelPositions;
-
-			uint32 NumInputChannels;
-			const uint32 NumFrames;
-			uint32 NumDeviceChannels;
-			uint8 bIsInitialDownmix : 1;
-
-			// If this source is an ambisonics source, we use this to downmix the source to our output.
-			TUniquePtr<ISoundfieldDecoderStream> AmbisonicsDecoder;
-
-			FSourceDownmixData(uint32 SourceNumChannels, uint32 NumDeviceOutputChannels, uint32 InNumFrames)
-				: SourceRotation(FQuat::Identity)
-				, PostEffectBuffers(nullptr)
-				, DeviceSubmixInfo(FSubmixChannelData(SourceNumChannels, NumDeviceOutputChannels, InNumFrames))
-				, NumInputChannels(SourceNumChannels)
-				, NumFrames(InNumFrames)
-				, NumDeviceChannels(NumDeviceOutputChannels)
-				, bIsInitialDownmix(true) 
-			{
-			}
-
-			void ResetNumberOfDeviceChannels(const uint32 NumDeviceOutputChannels)
-			{
-				NumDeviceChannels = NumDeviceOutputChannels;
-				DeviceSubmixInfo.Reset(NumInputChannels, NumDeviceOutputChannels, NumFrames);
-			}
-
-			void ResetData(const uint32 InNumInputChannels, int32 InNumDeviceChannels)
-			{
-				NumDeviceChannels = InNumDeviceChannels;
-				NumInputChannels = InNumInputChannels;
-				PostEffectBuffers = nullptr;
-
-				DeviceSubmixInfo.Reset(NumInputChannels, NumDeviceChannels, NumFrames);
-				EncodedSoundfieldDownmixes.Reset();
-				AmbisonicsDecoder.Reset();
-				bIsInitialDownmix = true;
-			}
-		};
-
 		struct FSourceInfo
 		{
 			FSourceInfo() {}
@@ -504,12 +369,10 @@ namespace Audio
 			uint8 bIsLastBuffer:1;
 			uint8 bOutputToBusOnly:1;
 			uint8 bIsVorbis:1;
-
-			// TODO: Make this generic to soundfield formats as a whole.
-			// For now, if this is true we just wrap the audio buffer in an FAmbisonicsSoundfieldBuffer.
-			uint8 bIsAmbisonics:1;
+			uint8 bIsSoundfield:1;
 			uint8 bIsBypassingLPF:1;
 			uint8 bIsBypassingHPF:1;
+			uint8 bHasPreDistanceAttenuationSend:1;
 
 			// Source format info
 			int32 NumInputChannels;
@@ -526,21 +389,7 @@ namespace Audio
 		};
 
 		static void ApplyDistanceAttenuation(FSourceInfo& InSourceInfo, int32 NumSamples);
-		void ComputePluginAudio(FSourceInfo& InSourceInfo, FSourceDownmixData& DownmixData, int32 SourceId, int32 NumSamples);
-
-		static void ComputeDownmix3D(FSourceDownmixData& DownmixData, FMixerDevice* MixerDevice);
-		static void ComputeDownmix2D(FSourceDownmixData& DownmixData, FMixerDevice* MixerDevice);
-		static void EncodeToSoundfieldFormats(FSourceDownmixData& DownmixData, const FSourceInfo& InSource, FMixerDevice* InDevice);
-		static void ConvertCartesianToSpherical(const FVector& InVector, float& OutAzimuth, float& OutElevation, float& OutRadius);
-
-		// This function is effectively equivalent to calling DownmixDataArray.EmplaceAt_GetRef(args...), but bypasses that function's intrinsic call to AddUninitialized.
-		FSourceDownmixData& InitializeDownmixForSource(const int32 SourceId, const int32 NumInputChannels, const int32 NumOutputChannels, const int32 InNumOutputFrames);
-
-		const FSubmixChannelData& GetChannelInfoForDevice(const FSourceDownmixData& InDownmixData) const;
-		FSubmixChannelData& GetChannelInfoForDevice(FSourceDownmixData& InDownmixData);
-
-		const FSubmixSoundfieldData& GetChannelInfoForFormat(const FSoundfieldEncodingKey& InFormat, const FSourceDownmixData& InDownmixData) const;
-		FSubmixSoundfieldData& GetChannelInfoForFormat(const FSoundfieldEncodingKey& InFormat, FSourceDownmixData& InDownmixData);
+		void ComputePluginAudio(FSourceInfo& InSourceInfo, FMixerSourceSubmixOutputBuffer& InSourceSubmixOutputBuffer, int32 SourceId, int32 NumSamples);
 
 		// Array of listener transforms
 		TArray<FTransform> ListenerTransforms;
@@ -548,9 +397,8 @@ namespace Audio
 		// Array of source infos.
 		TArray<FSourceInfo> SourceInfos;
 
-		// These structs are used for guaranteed vectorization when downmixing
-		// sources.
-		TArray<FSourceDownmixData> DownmixDataArray;
+		// This array is independent of SourceInfos array to optimize for cache coherency
+		TArray<FMixerSourceSubmixOutputBuffer> SourceSubmixOutputBuffers;
 
 		// Map of bus object Id's to audio bus data. 
 		TMap<uint32, TSharedPtr<FMixerAudioBus>> AudioBuses;
diff --git a/Engine/Source/Runtime/AudioMixer/Private/AudioMixerSourceOutputBuffer.cpp b/Engine/Source/Runtime/AudioMixer/Private/AudioMixerSourceOutputBuffer.cpp
new file mode 100644
--- /dev/null
+++ b/Engine/Source/Runtime/AudioMixer/Private/AudioMixerSourceOutputBuffer.cpp
@@ -0,0 +1,532 @@
+// Copyright Epic Games, Inc. All Rights Reserved.
+
+#include "AudioMixerSourceOutputBuffer.h"
+#include "SoundFieldRendering.h"
+#include "AudioMixerSubmix.h"
+
+
+namespace Audio
+{
+	// Utility function to convert cartesian coordinates to spherical coordinates
+	static void ConvertCartesianToSpherical(const FVector& InVector, float& OutAzimuth, float& OutElevation, float& OutRadius)
+	{
+		// Convert coordinates from unreal cartesian system to left handed spherical coordinates (zenith is positive elevation, right is positive azimuth)
+		const float InX = -InVector.Z; //InVector.Y;
+		const float InY = InVector.X;// -InVector.Z;
+		const float InZ = -InVector.Y;
+
+
+		OutElevation = FMath::Atan2(InY, InX);
+
+		// Note, rather than using arccos(z / radius) here, we use Atan2 to avoid wrapping issues with negative elevation values.
+		OutAzimuth = FMath::Atan2(FMath::Sqrt(InX * InX + InY * InY), InZ);
+		OutRadius = InVector.Size();
+	}
+
+	FMixerSourceSubmixOutputBuffer::FMixerSourceSubmixOutputBuffer(FMixerDevice* InMixerDevice, uint32 InNumSourceChannels, uint32 InNumOutputChannels, uint32 InNumFrames)
+		: PreAttenuationSourceBuffer(nullptr)
+		, PostAttenuationSourceBuffer(nullptr)
+		, SourceChannelMap(InNumSourceChannels, InNumOutputChannels)
+		, NumSourceChannels(InNumSourceChannels)
+		, NumFrames(InNumFrames)
+		, NumOutputChannels(InNumOutputChannels)
+		, MixerDevice(InMixerDevice)
+		, bIsInitialDownmix(true)
+		, bIs3D(false)
+		, bIsVorbis(false)
+	{
+		PreAttenuationOutputBuffer.Reset();
+		PreAttenuationOutputBuffer.AddUninitialized(InNumFrames * InNumOutputChannels);
+
+		PostAttenuationOutputBuffer.Reset();
+		PostAttenuationOutputBuffer.AddUninitialized(InNumFrames * InNumOutputChannels);
+	}
+
+	FMixerSourceSubmixOutputBuffer::~FMixerSourceSubmixOutputBuffer()
+	{
+
+	}
+
+	void FMixerSourceSubmixOutputBuffer::Reset(const FMixerSourceSubmixOutputBufferSettings& InInitSettings)
+	{
+		// Whether or not this is a 3D submix output
+		bIs3D = InInitSettings.bIs3D;
+		bIsVorbis = InInitSettings.bIsVorbis;
+
+		// Reset our record-keeping data
+		NumOutputChannels = InInitSettings.NumOutputChannels;
+		NumSourceChannels = InInitSettings.NumSourceChannels;
+
+		// Reset the source channel map
+		SourceChannelMap.Reset(NumSourceChannels, NumOutputChannels);
+
+		// Reset the output buffers
+		PreAttenuationOutputBuffer.Reset();
+		PreAttenuationOutputBuffer.AddUninitialized(NumFrames * NumOutputChannels);
+
+		PostAttenuationOutputBuffer.Reset();
+		PostAttenuationOutputBuffer.AddUninitialized(NumFrames * NumOutputChannels);
+
+		// Reset the post and pre source buffers
+		PostAttenuationSourceBuffer = nullptr;
+		PreAttenuationSourceBuffer = nullptr;
+
+		// Reset the soundfield data
+		EncodedSoundfieldDownmixes.Reset();
+		SoundfieldDecoder.Reset();
+		bIsInitialDownmix = true;
+
+		// Make a new soundfield decoder if this sound is a soundfield source
+		if (InInitSettings.bIsSoundfield)
+		{
+			SoundfieldDecoder = CreateDefaultSourceAmbisonicsDecoder(MixerDevice);
+		}
+
+		// Create data for each of the soundfield submix sends
+		for (const FMixerSubmixPtr& SoundfieldSubmixPtr : InInitSettings.SoundfieldSubmixSends)
+		{
+			// Get the encoding key for the given soundfield submix this source is sending to
+			FSoundfieldEncodingKey EncodingKey = SoundfieldSubmixPtr->GetKeyForSubmixEncoding();
+
+			// Create a data bucket based on this encoding key
+			FSoundfieldData& SoundfieldData = EncodedSoundfieldDownmixes.FindOrAdd(EncodingKey);
+
+			// Get the soundfield submix's factory object, which creates encoders, decoders, transcoders, etc.
+			ISoundfieldFactory* Factory = SoundfieldSubmixPtr->GetSoundfieldFactory();
+			check(Factory);
+
+			FAudioPluginInitializationParams PluginInitParams = SoundfieldSubmixPtr->GetInitializationParamsForSoundfieldStream();
+			PluginInitParams.NumOutputChannels = NumSourceChannels;
+
+			SoundfieldData.EncoderSettings = SoundfieldSubmixPtr->GetSoundfieldSettings().Duplicate();
+
+			// If this source is soundfield, we need to use a transcoder
+			if (InInitSettings.bIsSoundfield)
+			{
+				if (Factory->GetSoundfieldFormatName() == GetUnrealAmbisonicsFormatName())
+				{
+					SoundfieldData.bIsUnrealAmbisonicsSubmix = true;
+				}
+				else if (Factory->CanTranscodeFromSoundfieldFormat(GetUnrealAmbisonicsFormatName(), GetAmbisonicsSourceDefaultSettings()))
+				{
+					SoundfieldData.SoundfieldTranscoder = Factory->CreateTranscoderStream(GetUnrealAmbisonicsFormatName(), GetAmbisonicsSourceDefaultSettings(), Factory->GetSoundfieldFormatName(), *SoundfieldData.EncoderSettings, PluginInitParams);
+				}
+			}
+			else
+			{
+				check(SoundfieldData.EncoderSettings.IsValid());
+
+				SoundfieldData.SoundfieldEncoder = Factory->CreateEncoderStream(PluginInitParams, *SoundfieldData.EncoderSettings);
+			}
+
+			// Create a blank packet for memory of the encoded packet
+			SoundfieldData.EncodedPacket = Factory->CreateEmptyPacket();
+		}
+	}
+
+	void FMixerSourceSubmixOutputBuffer::SetNumOutputChannels(uint32 InNumOutputChannels)
+	{
+		NumOutputChannels = InNumOutputChannels;
+
+		SourceChannelMap.Reset(NumSourceChannels, NumOutputChannels);
+
+		PreAttenuationOutputBuffer.Reset();
+		PreAttenuationOutputBuffer.AddUninitialized(NumFrames * NumOutputChannels);
+
+		PostAttenuationOutputBuffer.Reset();
+		PostAttenuationOutputBuffer.AddUninitialized(NumFrames * NumOutputChannels);
+	}
+
+	bool FMixerSourceSubmixOutputBuffer::SetChannelMap(const AlignedFloatBuffer& InChannelMap, bool bInIsCenterChannelOnly)
+	{
+		bool bNeedsNewChannelMap = false;
+
+		// Fix up the channel map in case the device output count changed
+		const uint32 ChannelMapSize = SourceChannelMap.CopySize / sizeof(float);
+
+		if (InChannelMap.Num() != ChannelMapSize)
+		{
+			AlignedFloatBuffer NewChannelMap;
+
+			if (bIs3D)
+			{
+				NewChannelMap.AddZeroed(ChannelMapSize);
+				bNeedsNewChannelMap = true;
+			}
+			else
+			{
+				const uint32 ChannelMapOutputChannels = ChannelMapSize / NumSourceChannels;
+				FMixerDevice::Get2DChannelMap(bIsVorbis, NumSourceChannels, ChannelMapOutputChannels, bInIsCenterChannelOnly, NewChannelMap);
+			}
+
+			SourceChannelMap.SetChannelMap(NewChannelMap.GetData());
+		}
+		else
+		{
+			SourceChannelMap.SetChannelMap(InChannelMap.GetData());
+		}
+
+		return bNeedsNewChannelMap;
+	}
+
+	void FMixerSourceSubmixOutputBuffer::SetPreAttenuationSourceBuffer(AlignedFloatBuffer* InPreAttenuationSourceBuffer)
+	{
+		PreAttenuationSourceBuffer = InPreAttenuationSourceBuffer;
+	}
+
+	void FMixerSourceSubmixOutputBuffer::SetPostAttenuationSourceBuffer(AlignedFloatBuffer* InPostAttenuationSourceBuffer)
+	{
+		PostAttenuationSourceBuffer = InPostAttenuationSourceBuffer;
+	}
+
+	void FMixerSourceSubmixOutputBuffer::ComputeOutput(const FSpatializationParams& InSpatParams)
+	{
+		// No need to compute an output if there is no source buffer available
+		if (!PostAttenuationSourceBuffer && !PreAttenuationSourceBuffer)
+		{
+			return;
+		}
+
+		// Update our rotational data based off the spat params
+		SoundfieldPositionalData.Rotation = InSpatParams.ListenerOrientation;
+		SoundSourceRotation = InSpatParams.EmitterWorldRotation;
+		
+		if (bIs3D && !bIsInitialDownmix)
+		{
+			ComputeOutput3D();
+		}
+		else
+		{
+			ComputeOutput2D();
+			bIsInitialDownmix = false;
+		}
+
+		// Now check if we need to do any sound field encoding
+		if (EncodedSoundfieldDownmixes.Num())
+		{
+			EncodeToSoundfieldFormats(InSpatParams);
+		}
+	}
+
+	void FMixerSourceSubmixOutputBuffer::ComputeOutput3D(AlignedFloatBuffer& InSourceBuffer, AlignedFloatBuffer& OutSourceBuffer)
+	{
+		if (SoundfieldDecoder.IsValid())
+		{
+			FAmbisonicsSoundfieldBuffer SoundfieldBuffer;
+			SoundfieldBuffer.AudioBuffer = MoveTemp(InSourceBuffer);
+			SoundfieldBuffer.NumChannels = NumSourceChannels;
+			SoundfieldBuffer.PreviousRotation = SoundfieldBuffer.Rotation;
+			SoundfieldBuffer.Rotation = SoundSourceRotation;
+
+			SoundfieldPositionalData.NumChannels = NumOutputChannels;
+			SoundfieldPositionalData.ChannelPositions = MixerDevice->GetDefaultPositionMap(NumOutputChannels);
+
+			FSoundfieldDecoderInputData SoundfieldDecoderInputData =
+			{
+				SoundfieldBuffer,
+				SoundfieldPositionalData,
+				static_cast<int32>(InSourceBuffer.Num() / NumSourceChannels),
+				MixerDevice->GetSampleRate()
+			};
+
+			FSoundfieldDecoderOutputData SoundFieldDecoderOutputData = { OutSourceBuffer };
+
+			SoundfieldDecoder->Decode(SoundfieldDecoderInputData, SoundFieldDecoderOutputData);
+
+			InSourceBuffer = MoveTemp(SoundfieldBuffer.AudioBuffer);
+		}
+		else if (NumSourceChannels == 1)
+		{
+			switch (NumOutputChannels)
+			{
+			case 8:
+				MixMonoTo8ChannelsFast(InSourceBuffer, OutSourceBuffer, SourceChannelMap.ChannelStartGains, SourceChannelMap.ChannelDestinationGains);
+				break;
+			case 6:
+				MixMonoTo6ChannelsFast(InSourceBuffer, OutSourceBuffer, SourceChannelMap.ChannelStartGains, SourceChannelMap.ChannelDestinationGains);
+				break;
+			case 4:
+				MixMonoTo4ChannelsFast(InSourceBuffer, OutSourceBuffer, SourceChannelMap.ChannelStartGains, SourceChannelMap.ChannelDestinationGains);
+				break;
+			case 2:
+				MixMonoTo2ChannelsFast(InSourceBuffer, OutSourceBuffer, SourceChannelMap.ChannelStartGains, SourceChannelMap.ChannelDestinationGains);
+				break;
+			}
+
+		}
+		else if (NumSourceChannels == 2)
+		{
+			switch (NumOutputChannels)
+			{
+			case 8:
+				Mix2ChannelsTo8ChannelsFast(InSourceBuffer, OutSourceBuffer, SourceChannelMap.ChannelStartGains, SourceChannelMap.ChannelDestinationGains);
+				break;
+			case 6:
+				Mix2ChannelsTo6ChannelsFast(InSourceBuffer, OutSourceBuffer, SourceChannelMap.ChannelStartGains, SourceChannelMap.ChannelDestinationGains);
+				break;
+			case 4:
+				Mix2ChannelsTo4ChannelsFast(InSourceBuffer, OutSourceBuffer, SourceChannelMap.ChannelStartGains, SourceChannelMap.ChannelDestinationGains);
+				break;
+			case 2:
+				Mix2ChannelsTo2ChannelsFast(InSourceBuffer, OutSourceBuffer, SourceChannelMap.ChannelStartGains, SourceChannelMap.ChannelDestinationGains);
+				break;
+			}
+		}
+		else
+		{
+			DownmixBuffer(NumSourceChannels, NumOutputChannels, InSourceBuffer, OutSourceBuffer, SourceChannelMap.ChannelStartGains, SourceChannelMap.ChannelDestinationGains);
+		}
+	}
+
+	void FMixerSourceSubmixOutputBuffer::ComputeOutput3D()
+	{
+		// Compute the output buffers if there is a corresponding pre or post source buffer set
+
+		if (PreAttenuationSourceBuffer)
+		{
+			ComputeOutput3D(*PreAttenuationSourceBuffer, PreAttenuationOutputBuffer);
+		}
+
+		if (PostAttenuationSourceBuffer)
+		{
+			ComputeOutput3D(*PostAttenuationSourceBuffer, PostAttenuationOutputBuffer);
+		}
+
+		// Do the channel map copy from dest to start (which prevents zippering of dynamic channel mapping for 3d audio) after the data has been computed for both pre- and post- attenuation channel downmixing
+		SourceChannelMap.CopyDestinationToStart();
+	}
+
+	void FMixerSourceSubmixOutputBuffer::ComputeOutput2D(AlignedFloatBuffer& InSourceBuffer, AlignedFloatBuffer& OutSourceBuffer)
+	{
+		// For 2D sources, we just apply the gain matrix in ChannelDestionationGains with no interpolation.
+		if (SoundfieldDecoder.IsValid())
+		{
+			FAmbisonicsSoundfieldBuffer SoundfieldBuffer;
+			SoundfieldBuffer.AudioBuffer = MoveTemp(InSourceBuffer);
+			SoundfieldBuffer.NumChannels = NumSourceChannels;
+			SoundfieldBuffer.PreviousRotation = SoundfieldBuffer.Rotation;
+			SoundfieldBuffer.Rotation = SoundSourceRotation;
+
+			SoundfieldPositionalData.NumChannels = NumOutputChannels;
+			SoundfieldPositionalData.ChannelPositions = MixerDevice->GetDefaultPositionMap(NumOutputChannels);
+
+			FSoundfieldDecoderInputData SoundfieldDecoderInputData =
+			{
+				SoundfieldBuffer,
+				SoundfieldPositionalData,
+				static_cast<int32>(InSourceBuffer.Num() / NumSourceChannels),
+				MixerDevice->GetSampleRate()
+			};
+
+			FSoundfieldDecoderOutputData SoundfieldDecoderOutputData = { OutSourceBuffer };
+
+			SoundfieldDecoder->Decode(SoundfieldDecoderInputData, SoundfieldDecoderOutputData);
+
+			// Move the encoded ambisonics source buffer back to PostEffectBuffers to prevent reallocation
+			InSourceBuffer = MoveTemp(SoundfieldBuffer.AudioBuffer);
+		}
+		else if (NumSourceChannels == 1)
+		{
+			switch (NumOutputChannels)
+			{
+			case 8:
+				Audio::MixMonoTo8ChannelsFast(InSourceBuffer, OutSourceBuffer, SourceChannelMap.ChannelDestinationGains);
+				break;
+			case 6:
+				Audio::MixMonoTo6ChannelsFast(InSourceBuffer, OutSourceBuffer, SourceChannelMap.ChannelDestinationGains);
+				break;
+			case 4:
+				Audio::MixMonoTo4ChannelsFast(InSourceBuffer, OutSourceBuffer, SourceChannelMap.ChannelDestinationGains);
+				break;
+			case 2:
+				Audio::MixMonoTo2ChannelsFast(InSourceBuffer, OutSourceBuffer, SourceChannelMap.ChannelDestinationGains);
+				break;
+			}
+		}
+		else if (NumSourceChannels == 2)
+		{
+			switch (NumOutputChannels)
+			{
+			case 8:
+				Audio::Mix2ChannelsTo8ChannelsFast(InSourceBuffer, OutSourceBuffer, SourceChannelMap.ChannelDestinationGains);
+				break;
+			case 6:
+				Audio::Mix2ChannelsTo6ChannelsFast(InSourceBuffer, OutSourceBuffer, SourceChannelMap.ChannelDestinationGains);
+				break;
+			case 4:
+				Audio::Mix2ChannelsTo4ChannelsFast(InSourceBuffer, OutSourceBuffer, SourceChannelMap.ChannelDestinationGains);
+				break;
+			case 2:
+				Audio::Mix2ChannelsTo2ChannelsFast(InSourceBuffer, OutSourceBuffer, SourceChannelMap.ChannelDestinationGains);
+				break;
+			}
+		}
+		else
+		{
+			Audio::DownmixBuffer(NumSourceChannels, NumOutputChannels, InSourceBuffer, OutSourceBuffer, SourceChannelMap.ChannelDestinationGains);
+		}
+	}
+
+	void FMixerSourceSubmixOutputBuffer::ComputeOutput2D()
+	{
+		if (PreAttenuationSourceBuffer)
+		{
+			ComputeOutput2D(*PreAttenuationSourceBuffer, PreAttenuationOutputBuffer);
+		}
+
+		if (PostAttenuationSourceBuffer)
+		{
+			ComputeOutput2D(*PostAttenuationSourceBuffer, PostAttenuationOutputBuffer);
+		}
+	}
+
+	void FMixerSourceSubmixOutputBuffer::EncodeToSoundfieldFormats(const FSpatializationParams& InSpatParams)
+	{
+		check(MixerDevice);
+
+		SoundfieldPositionalData.NumChannels = NumSourceChannels;
+
+		// Spoof rotation of the source as if it's rotation of the listener when encoding non-sound-field to sound-field
+		SoundfieldPositionalData.Rotation = InSpatParams.EmitterWorldRotation;
+
+		InputChannelPositions.Reset();
+
+		if (bIs3D)
+		{
+			if (NumSourceChannels == 1)
+			{
+				FChannelPositionInfo ChannelPosition;
+				ChannelPosition.Channel = EAudioMixerChannel::FrontCenter;
+				ConvertCartesianToSpherical(InSpatParams.EmitterPosition, ChannelPosition.Azimuth, ChannelPosition.Elevation, ChannelPosition.Radius);
+
+				InputChannelPositions.Add(ChannelPosition);
+
+				SoundfieldPositionalData.ChannelPositions = &InputChannelPositions;
+			}
+			else if (NumSourceChannels == 2)
+			{
+				FChannelPositionInfo LeftChannelPosition;
+				LeftChannelPosition.Channel = EAudioMixerChannel::FrontLeft;
+				ConvertCartesianToSpherical(InSpatParams.LeftChannelPosition, LeftChannelPosition.Azimuth, LeftChannelPosition.Elevation, LeftChannelPosition.Radius);
+
+				InputChannelPositions.Add(LeftChannelPosition);
+
+				FChannelPositionInfo RightChannelPosition;
+				LeftChannelPosition.Channel = EAudioMixerChannel::FrontRight;
+				ConvertCartesianToSpherical(InSpatParams.RightChannelPosition, RightChannelPosition.Azimuth, RightChannelPosition.Elevation, RightChannelPosition.Radius);
+
+				InputChannelPositions.Add(RightChannelPosition);
+
+				SoundfieldPositionalData.ChannelPositions = &InputChannelPositions;
+			}
+		}
+
+		// if 2D or not a supported channel configuration, use default position map
+		if (!InputChannelPositions.Num())
+		{
+			SoundfieldPositionalData.ChannelPositions = MixerDevice->GetDefaultPositionMap(NumSourceChannels);
+		}
+
+		// Run the encoders.
+		for (auto& Soundfield : EncodedSoundfieldDownmixes)
+		{
+			FSoundfieldData& SoundfieldData = Soundfield.Value;
+
+			if (PreAttenuationSourceBuffer)
+			{
+				EncodeSoundfield(SoundfieldData , *PreAttenuationSourceBuffer);
+			}
+
+			if (PostAttenuationSourceBuffer)
+			{
+				EncodeSoundfield(SoundfieldData, *PostAttenuationSourceBuffer);
+			}
+		}
+	}
+
+	void FMixerSourceSubmixOutputBuffer::EncodeSoundfield(FSoundfieldData& InSoundfieldData, Audio::AlignedFloatBuffer& InSourceBuffer)
+	{
+		check(InSoundfieldData.EncoderSettings.IsValid());
+		check(InSoundfieldData.EncodedPacket.IsValid());
+
+		InSoundfieldData.EncodedPacket->Reset();
+
+		// We will have an soundfield transcoder if this sound source is a soundfield format
+		if (InSoundfieldData.SoundfieldTranscoder)
+		{
+			FAmbisonicsSoundfieldBuffer SoundfieldBuffer;
+			SoundfieldBuffer.AudioBuffer = MoveTemp(InSourceBuffer);
+			SoundfieldBuffer.NumChannels = NumSourceChannels;
+			SoundfieldBuffer.PreviousRotation = SoundfieldBuffer.Rotation;
+			SoundfieldBuffer.Rotation = SoundSourceRotation;
+
+			InSoundfieldData.SoundfieldTranscoder->Transcode(SoundfieldBuffer, GetAmbisonicsSourceDefaultSettings(), *InSoundfieldData.EncodedPacket, *InSoundfieldData.EncoderSettings);
+			InSourceBuffer = MoveTemp(SoundfieldBuffer.AudioBuffer);
+		}
+		else if (InSoundfieldData.SoundfieldEncoder)
+		{
+			FSoundfieldEncoderInputData SoundfieldEncoderInputData =
+			{
+				InSourceBuffer,
+				static_cast<int32>(NumSourceChannels),
+				*InSoundfieldData.EncoderSettings,
+				SoundfieldPositionalData
+			};
+
+			InSoundfieldData.SoundfieldEncoder->Encode(SoundfieldEncoderInputData, *InSoundfieldData.EncodedPacket);
+		}
+		else if (InSoundfieldData.bIsUnrealAmbisonicsSubmix)
+		{
+			FAmbisonicsSoundfieldBuffer& OutputPacket = DowncastSoundfieldRef<FAmbisonicsSoundfieldBuffer>(*InSoundfieldData.EncodedPacket);
+
+			// Fixme: This is an array copy. Can we serve InPositionalData directly to this soundfield?
+			OutputPacket.AudioBuffer = InSourceBuffer;
+			OutputPacket.NumChannels = NumSourceChannels;
+			OutputPacket.PreviousRotation = OutputPacket.Rotation;
+			OutputPacket.Rotation = SoundSourceRotation;
+		}
+	}
+
+	void FMixerSourceSubmixOutputBuffer::MixOutput(float InSendLevel, EMixerSourceSubmixSendStage InSubmixSendStage, AlignedFloatBuffer& OutMixedBuffer) const
+	{
+		if (InSubmixSendStage == EMixerSourceSubmixSendStage::PostDistanceAttenuation)
+		{
+			Audio::MixInBufferFast(PostAttenuationOutputBuffer, OutMixedBuffer, InSendLevel);
+		}
+		else
+		{
+			Audio::MixInBufferFast(PreAttenuationOutputBuffer, OutMixedBuffer, InSendLevel);
+		}
+	}
+
+	FQuat FMixerSourceSubmixOutputBuffer::GetListenerRotation() const
+	{
+		return SoundfieldPositionalData.Rotation;
+	}
+
+	void FMixerSourceSubmixOutputBuffer::CopyReverbPluginOutputData(AlignedFloatBuffer& InAudioBuffer)
+	{
+		ReverbPluginOutputBuffer.Reset();
+		ReverbPluginOutputBuffer.Append(InAudioBuffer);
+	}
+
+	const float* FMixerSourceSubmixOutputBuffer::GetReverbPluginOutputData() const
+	{
+		return ReverbPluginOutputBuffer.GetData();
+	}
+
+	const ISoundfieldAudioPacket* FMixerSourceSubmixOutputBuffer::GetSoundfieldPacket(const FSoundfieldEncodingKey& InKey) const
+	{
+		check(EncodedSoundfieldDownmixes.Contains(InKey));
+		const FSoundfieldData& SoundfieldData = EncodedSoundfieldDownmixes[InKey];
+		return SoundfieldData.EncodedPacket.Get();
+	}
+
+	ISoundfieldAudioPacket* FMixerSourceSubmixOutputBuffer::GetSoundFieldPacket(const FSoundfieldEncodingKey& InKey)
+	{
+		FSoundfieldData& SoundfieldData = EncodedSoundfieldDownmixes.FindOrAdd(InKey);
+		return SoundfieldData.EncodedPacket.Get();
+	}
+
+
+}
\ No newline at end of file
diff --git a/Engine/Source/Runtime/AudioMixer/Private/AudioMixerSourceOutputBuffer.h b/Engine/Source/Runtime/AudioMixer/Private/AudioMixerSourceOutputBuffer.h
new file mode 100644
--- /dev/null
+++ b/Engine/Source/Runtime/AudioMixer/Private/AudioMixerSourceOutputBuffer.h
@@ -0,0 +1,210 @@
+// Copyright Epic Games, Inc. All Rights Reserved.
+
+#pragma once
+
+
+#include "CoreMinimal.h"
+#include "ISoundfieldFormat.h"
+#include "IAudioExtensionPlugin.h"
+#include "AudioMixerDevice.h"
+#include "AudioMixerSubmix.h"
+
+namespace Audio
+{
+	class FMixerSubmix;
+
+	struct FMixerSourceSubmixOutputBufferSettings
+	{
+		uint32 NumSourceChannels;
+		uint32 NumOutputChannels;
+		TArray<FMixerSubmixPtr> SoundfieldSubmixSends;
+		bool bIs3D;
+		bool bIsVorbis;
+		bool bIsSoundfield;
+	};
+
+	/** 
+	* Used in audio mixer source manager to convert source audio to channel output and mix into submixes
+	*/
+	class FMixerSourceSubmixOutputBuffer
+	{
+	public:
+		FMixerSourceSubmixOutputBuffer(FMixerDevice* InMixerDevice, uint32 InNumSourceChannels, uint32 InNumOutputChannels, uint32 InNumFrames);
+
+		~FMixerSourceSubmixOutputBuffer();
+
+		// Resets the source submix output to prepare for new source data
+		void Reset(const FMixerSourceSubmixOutputBufferSettings& InResetSettings);
+
+		// Sets the number of output channels
+		void SetNumOutputChannels(uint32 InNumOutputChannels);
+
+		// Returns the number of source channels of this source
+		uint32 GetNumSourceChannels() const { return NumSourceChannels; }
+
+		// Sets the channel map
+		bool SetChannelMap(const AlignedFloatBuffer& InChannelMap, bool bInIsCenterChannelOnly);
+
+		// Sets the pre and post attenuation source buffers. This is the source buffer data derived from source manager source processing.
+		void SetPreAttenuationSourceBuffer(AlignedFloatBuffer* InPreAttenuationBuffer);
+		void SetPostAttenuationSourceBuffer(AlignedFloatBuffer* InPostAttenuationBuffer);
+		void CopyReverbPluginOutputData(AlignedFloatBuffer& InAudioBuffer);
+		const float* GetReverbPluginOutputData() const;
+
+		// Retrieves the current sound field packet for the given key
+		const ISoundfieldAudioPacket* GetSoundfieldPacket(const FSoundfieldEncodingKey& InKey) const;
+		ISoundfieldAudioPacket* GetSoundFieldPacket(const FSoundfieldEncodingKey& InKey);
+
+		// Computes the output buffer given the spat params. This buffer can then be mixed to submixes
+		void ComputeOutput(const FSpatializationParams& InSpatParams);
+
+		// Called by submixes to mix this output buffer to their buffer
+		void MixOutput(float SendLevel, EMixerSourceSubmixSendStage InSubmixSendStage, AlignedFloatBuffer& OutMixedBuffer) const;
+
+		// Return the listener rotation
+		// TODO: consolidate the code that is using this to be private to this class.
+		FQuat GetListenerRotation() const;
+
+	private: // private classes
+
+		// Private struct used for channel mapping
+		struct FSourceChannelMap
+		{
+			alignas(16) float ChannelStartGains[AUDIO_MIXER_MAX_OUTPUT_CHANNELS * AUDIO_MIXER_MAX_OUTPUT_CHANNELS];
+			alignas(16) float ChannelDestinationGains[AUDIO_MIXER_MAX_OUTPUT_CHANNELS * AUDIO_MIXER_MAX_OUTPUT_CHANNELS];
+
+			// This is the number of bytes the gain array is using:
+			// (Number of input channels * number of output channels) * sizeof float.
+			int32 CopySize = 0;
+			bool bIsInit = false;
+
+			FSourceChannelMap(int32 InNumInChannels, int32 InNumOutChannels)
+				: CopySize(InNumInChannels* InNumOutChannels * sizeof(float))
+			{
+				checkSlow(InNumInChannels <= AUDIO_MIXER_MAX_OUTPUT_CHANNELS);
+				checkSlow(InNumOutChannels <= AUDIO_MIXER_MAX_OUTPUT_CHANNELS);
+				FMemory::Memzero(ChannelStartGains, CopySize);
+			}
+
+			FORCEINLINE void Reset(int32 InNumInChannels, int32 InNumOutChannels)
+			{
+				checkSlow(InNumInChannels <= AUDIO_MIXER_MAX_OUTPUT_CHANNELS);
+				checkSlow(InNumOutChannels <= AUDIO_MIXER_MAX_OUTPUT_CHANNELS);
+
+				CopySize = InNumInChannels * InNumOutChannels * sizeof(float);
+				FMemory::Memzero(ChannelStartGains, CopySize);
+				FMemory::Memzero(ChannelDestinationGains, CopySize);
+				bIsInit = false;
+			}
+
+			FORCEINLINE void CopyDestinationToStart()
+			{
+				FMemory::Memcpy(ChannelStartGains, ChannelDestinationGains, CopySize);
+			}
+
+			FORCEINLINE void SetChannelMap(const float* RESTRICT InChannelGains)
+			{
+				FMemory::Memcpy(ChannelDestinationGains, InChannelGains, CopySize);
+				if (!bIsInit)
+				{
+					FMemory::Memcpy(ChannelStartGains, InChannelGains, CopySize);
+					bIsInit = true;
+				}
+			}
+
+		private:
+			FSourceChannelMap()
+				: CopySize(0)
+				, bIsInit(false)
+			{
+			}
+		};
+
+		struct FSoundfieldData
+		{
+			// If the sound source is not a soundfield source and we're sending to a soundfield submix, it gets encoded
+			TUniquePtr<ISoundfieldEncoderStream> SoundfieldEncoder;
+
+			// If the sound source is soundfield, and we're sending to a soundfield submix, it gets transcoded
+			TUniquePtr<ISoundfieldTranscodeStream> SoundfieldTranscoder;
+
+			// The optional soundfield encoder settings
+			TUniquePtr<ISoundfieldEncodingSettingsProxy> EncoderSettings;
+
+			// The encoded soundfield packet
+			TUniquePtr<ISoundfieldAudioPacket> EncodedPacket;
+
+			// If this is a unreal ambisonics soundfield buffer, we hand it the submixed buffer directly.
+			bool bIsUnrealAmbisonicsSubmix;
+		};
+
+	private: // private methods
+
+		void ComputeOutput3D(AlignedFloatBuffer& InSource, AlignedFloatBuffer& InOutput);
+		void ComputeOutput3D();
+		void ComputeOutput2D(AlignedFloatBuffer& InSource, AlignedFloatBuffer& InOutput);
+		void ComputeOutput2D();
+
+		void EncodeSoundfield(FSoundfieldData& InSoundfieldData, Audio::AlignedFloatBuffer& InSourceBuffer);
+		void EncodeToSoundfieldFormats(const FSpatializationParams& InSpatParams);
+
+	private: // private data
+
+		// Buffer for reverb plugins to write into
+		AlignedFloatBuffer ReverbPluginOutputBuffer;
+
+		// The source buffer created before distance attenuation is applied
+		AlignedFloatBuffer* PreAttenuationSourceBuffer;
+
+		// The source buffer created after distance attenuation is applied
+		AlignedFloatBuffer* PostAttenuationSourceBuffer;
+
+		// Data used to map the source channel data to the output channel configuration
+		FSourceChannelMap SourceChannelMap;
+
+		// The result of the source output with source data derived pre distance attenuation
+		AlignedFloatBuffer PreAttenuationOutputBuffer;
+
+		// The result of the source output with source data derived post distance attenuation
+		AlignedFloatBuffer PostAttenuationOutputBuffer;
+
+		// The number of source channels (sound input)
+		uint32 NumSourceChannels;
+
+		// The number of interleaved frames in the source
+		const uint32 NumFrames;
+
+		// The number of device output channels
+		uint32 NumOutputChannels;
+
+		// The owning mixer device
+		FMixerDevice* MixerDevice;
+
+		// Whether this is the initial downmix
+		bool bIsInitialDownmix;
+
+		// If this is a 3D sound source
+		bool bIs3D;
+
+		// Whether this source is vorbis encoded (used for 2D speakermapping)
+		bool bIsVorbis;
+
+		// Cached parameters for encoding to a soundfield format
+		FSoundfieldSpeakerPositionalData SoundfieldPositionalData;
+
+		// The current sound source rotation
+		FQuat SoundSourceRotation;
+
+		// Map of sound field encoding to encoding data
+		TMap<FSoundfieldEncodingKey, FSoundfieldData> EncodedSoundfieldDownmixes;
+
+		// Channel positions of the sound field
+		TArray<Audio::FChannelPositionInfo> InputChannelPositions;
+
+		// If this source is an ambisonics source, we use this to down mix the source to output channel mix.
+		TUniquePtr<ISoundfieldDecoderStream> SoundfieldDecoder;
+
+	};
+
+
+}
diff --git a/Engine/Source/Runtime/AudioMixer/Private/AudioMixerSourceVoice.cpp b/Engine/Source/Runtime/AudioMixer/Private/AudioMixerSourceVoice.cpp
--- a/Engine/Source/Runtime/AudioMixer/Private/AudioMixerSourceVoice.cpp
+++ b/Engine/Source/Runtime/AudioMixer/Private/AudioMixerSourceVoice.cpp
@@ -249,13 +249,13 @@ namespace Audio
 		return SourceManager->GetEnvelopeValue(SourceId);
 	}
 
-	void FMixerSourceVoice::MixOutputBuffers(int32 InNumOutputChannels, const float SendLevel, AlignedFloatBuffer& OutWetBuffer) const
+	void FMixerSourceVoice::MixOutputBuffers(int32 InNumOutputChannels, const float SendLevel, EMixerSourceSubmixSendStage InSubmixSendStage, AlignedFloatBuffer& OutWetBuffer) const
 	{
 		AUDIO_MIXER_CHECK_AUDIO_PLAT_THREAD(MixerDevice);
 
 		check(!bOutputToBusOnly);
 
-		return SourceManager->MixOutputBuffers(SourceId, InNumOutputChannels, SendLevel, OutWetBuffer);
+		return SourceManager->MixOutputBuffers(SourceId, InNumOutputChannels, SendLevel, InSubmixSendStage, OutWetBuffer);
 	}
 
 	const ISoundfieldAudioPacket* FMixerSourceVoice::GetEncodedOutput(const FSoundfieldEncodingKey& InKey) const
diff --git a/Engine/Source/Runtime/AudioMixer/Private/AudioMixerSourceVoice.h b/Engine/Source/Runtime/AudioMixer/Private/AudioMixerSourceVoice.h
--- a/Engine/Source/Runtime/AudioMixer/Private/AudioMixerSourceVoice.h
+++ b/Engine/Source/Runtime/AudioMixer/Private/AudioMixerSourceVoice.h
@@ -95,7 +95,7 @@ namespace Audio
 		float GetEnvelopeValue() const;
 
 		// Mixes the dry and wet buffer audio into the given buffers.
-		void MixOutputBuffers(int32 InNumChannels, const float SendLevel, AlignedFloatBuffer& OutWetBuffer) const;
+		void MixOutputBuffers(int32 InNumChannels, const float SendLevel, EMixerSourceSubmixSendStage InSubmixSendStage, AlignedFloatBuffer& OutWetBuffer) const;
 
 		// For soundfield conversions, get the encoded audio.
 		const ISoundfieldAudioPacket* GetEncodedOutput(const FSoundfieldEncodingKey& InKey) const;
diff --git a/Engine/Source/Runtime/AudioMixer/Private/AudioMixerSubmix.cpp b/Engine/Source/Runtime/AudioMixer/Private/AudioMixerSubmix.cpp
--- a/Engine/Source/Runtime/AudioMixer/Private/AudioMixerSubmix.cpp
+++ b/Engine/Source/Runtime/AudioMixer/Private/AudioMixerSubmix.cpp
@@ -432,12 +432,14 @@ namespace Audio
 		return EffectSubmixChain.Num();
 	}
 
-	void FMixerSubmix::AddOrSetSourceVoice(FMixerSourceVoice* InSourceVoice, const float InSendLevel)
+
+	void FMixerSubmix::AddOrSetSourceVoice(FMixerSourceVoice* InSourceVoice, const float InSendLevel, EMixerSourceSubmixSendStage InSubmixSendStage)
 	{
 		AUDIO_MIXER_CHECK_AUDIO_PLAT_THREAD(MixerDevice);
 
 		FSubmixVoiceData NewVoiceData;
 		NewVoiceData.SendLevel = InSendLevel;
+		NewVoiceData.SubmixSendStage = InSubmixSendStage;
 
 		MixerSourceVoices.Add(InSourceVoice, NewVoiceData);
 	}
@@ -686,9 +688,7 @@ namespace Audio
 				// If this is true, the Soundfield Factory explicitly requested that a seperate encoder stream was set up for every
 				// non-soundfield child submix.
 				if (Child.Encoder.IsValid())
-				{
-					
-
+				{					
 					ChildSubmixSharedPtr->ProcessAudio(ScratchBuffer);
 
 					// Encode the resulting audio and mix it in.
@@ -954,9 +954,9 @@ namespace Audio
 			{
 				const FMixerSourceVoice* MixerSourceVoice = MixerSourceVoiceIter.Key;
 				const float SendLevel = MixerSourceVoiceIter.Value.SendLevel;
+				const EMixerSourceSubmixSendStage SubmixSendStage = MixerSourceVoiceIter.Value.SubmixSendStage;
 
-
-				MixerSourceVoice->MixOutputBuffers(NumChannels, SendLevel, InputBuffer);
+				MixerSourceVoice->MixOutputBuffers(NumChannels, SendLevel, SubmixSendStage, InputBuffer);
 			}
 		}
 
@@ -1056,6 +1056,7 @@ namespace Audio
 			FMemory::Memzero((void*)BufferPtr, sizeof(float) * NumSamples);
 		}
 	
+	
 		// If we are recording, Add out buffer to the RecordingData buffer:
 		{
 			FScopeLock ScopedLock(&RecordingCriticalSection);
@@ -1867,4 +1868,4 @@ namespace Audio
 		return SoundfieldStreams.Factory;
 	}
 
-}
+}
\ No newline at end of file
diff --git a/Engine/Source/Runtime/AudioMixer/Public/AudioMixerSubmix.h b/Engine/Source/Runtime/AudioMixer/Public/AudioMixerSubmix.h
--- a/Engine/Source/Runtime/AudioMixer/Public/AudioMixerSubmix.h
+++ b/Engine/Source/Runtime/AudioMixer/Public/AudioMixerSubmix.h
@@ -24,12 +24,23 @@ namespace Audio
 	class FMixerSourceVoice;
 	class FMixerDevice;
 
+	enum EMixerSourceSubmixSendStage
+	{
+		// Whether to do the send pre distance attenuation
+		PostDistanceAttenuation,
+
+		// Whether to do the send post distance attenuation
+		PreDistanceAttenuation,
+	};
+
 	struct FSubmixVoiceData
 	{
 		float SendLevel;
+		EMixerSourceSubmixSendStage SubmixSendStage;
 
 		FSubmixVoiceData()
 			: SendLevel(1.0f)
+			, SubmixSendStage(EMixerSourceSubmixSendStage::PostDistanceAttenuation)
 		{
 		}
 	};
@@ -111,7 +122,7 @@ namespace Audio
 		int32 GetSizeOfSubmixChain() const;
 
 		// Add (if not already added) or sets the amount of the source voice's send amount
-		void AddOrSetSourceVoice(FMixerSourceVoice* InSourceVoice, const float SendLevel);
+		void AddOrSetSourceVoice(FMixerSourceVoice* InSourceVoice, const float SendLevel, EMixerSourceSubmixSendStage InSubmixSendStage);
 
 		/** Removes the given source voice from the submix. */
 		void RemoveSourceVoice(FMixerSourceVoice* InSourceVoice);
diff --git a/Engine/Source/Runtime/Engine/Classes/Sound/SoundSubmixSend.h b/Engine/Source/Runtime/Engine/Classes/Sound/SoundSubmixSend.h
--- a/Engine/Source/Runtime/Engine/Classes/Sound/SoundSubmixSend.h
+++ b/Engine/Source/Runtime/Engine/Classes/Sound/SoundSubmixSend.h
@@ -78,6 +78,16 @@ enum class ESendLevelControlMethod : uint8
 	Manual,
 };
 
+UENUM(BlueprintType)
+enum class ESubmixSendStage : uint8
+{
+	// Whether to do the send pre distance attenuation
+	PostDistanceAttenuation,
+
+	// Whether to do the send post distance attenuation
+	PreDistanceAttenuation,
+};
+
 // Class used to send audio to submixes from USoundBase
 USTRUCT(BlueprintType)
 struct ENGINE_API FSoundSubmixSendInfo
@@ -94,6 +104,10 @@ struct ENGINE_API FSoundSubmixSendInfo
 	UPROPERTY(EditAnywhere, BlueprintReadWrite, Category = SubmixSend)
 	ESendLevelControlMethod SendLevelControlMethod;
 
+	/** Defines at what mix stage the send should happen.*/
+	UPROPERTY(EditAnywhere, BlueprintReadWrite, Category = SubmixSend)
+	ESubmixSendStage SendStage;
+
 	// The submix to send the audio to
 	UPROPERTY(EditAnywhere, BlueprintReadWrite, Category = SubmixSend)
 	USoundSubmixBase* SoundSubmix;
diff --git a/Engine/Source/Runtime/Engine/Private/SoundSubmixSend.cpp b/Engine/Source/Runtime/Engine/Private/SoundSubmixSend.cpp
--- a/Engine/Source/Runtime/Engine/Private/SoundSubmixSend.cpp
+++ b/Engine/Source/Runtime/Engine/Private/SoundSubmixSend.cpp
@@ -5,6 +5,7 @@
 
 FSoundSubmixSendInfo::FSoundSubmixSendInfo()
 	: SendLevelControlMethod(ESendLevelControlMethod::Manual)
+	, SendStage(ESubmixSendStage::PostDistanceAttenuation)
 	, SoundSubmix(nullptr)
 	, SendLevel(0.0f)
 	, MinSendLevel(0.0f)
