Adding ability for procedural audio to use a u-object generator class.

All synths will eventually be implemented using this API.

Change also has a test synth that exercises code that is a hello-world style "tone generator".

#

diff --git a/Engine/Plugins/Runtime/Synthesis/Source/Synthesis/Classes/SynthComponents/SynthComponentToneGenerator.h b/Engine/Plugins/Runtime/Synthesis/Source/Synthesis/Classes/SynthComponents/SynthComponentToneGenerator.h
new file mode 100644
--- /dev/null
+++ b/Engine/Plugins/Runtime/Synthesis/Source/Synthesis/Classes/SynthComponents/SynthComponentToneGenerator.h
@@ -0,0 +1,63 @@
+// Copyright Epic Games, Inc. All Rights Reserved.
+
+#pragma once
+
+#include "CoreMinimal.h"
+#include "Components/SynthComponent.h"
+#include "Sound/SoundBase.h"
+#include "DSP/SinOsc.h"
+#include "Sound/SoundGenerator.h"
+#include "SynthComponentToneGenerator.generated.h"
+
+class FToneGenerator : public ISoundGenerator
+{
+public:
+	FToneGenerator(int32 InSampleRate, int32 InNumChannels, int32 InFrequency, float InVolume);
+	virtual ~FToneGenerator();
+
+	//~ Begin FSoundGenerator 
+	virtual int32 GetNumChannels() { return NumChannels; };
+	virtual int32 OnGenerateAudio(float* OutAudio, int32 NumSamples) override;
+	//~ End FSoundGenerator
+
+	void SetFrequency(float InFrequency);
+	void SetVolume(float InVolume);
+
+private:
+	int32 NumChannels = 2;
+	Audio::FSineOsc SineOsc;
+	Audio::AlignedFloatBuffer Buffer;
+};
+
+UCLASS(ClassGroup = Synth, meta = (BlueprintSpawnableComponent))
+class SYNTHESIS_API USynthComponentToneGenerator : public USynthComponent
+{
+	GENERATED_BODY()
+
+	USynthComponentToneGenerator(const FObjectInitializer& ObjInitializer);
+	virtual ~USynthComponentToneGenerator();
+
+public:
+	// The frequency (in hz) of the tone generator.
+	UPROPERTY(EditAnywhere, BlueprintReadOnly, Category = "Tone Generator", meta = (ClampMin = "10.0", ClampMax = "20000.0"))
+	float Frequency;
+
+	// The linear volume of the tone generator.
+	UPROPERTY(EditAnywhere, BlueprintReadOnly, Category = "Tone Generator", meta = (ClampMin = "0.0", ClampMax = "1.0"))
+	float Volume;
+
+	// Sets the frequency of the tone generator
+	UFUNCTION(BlueprintCallable, Category = "Tone Generator")
+	void SetFrequency(float InFrequency);
+
+	// Sets the volume of the tone generator
+	UFUNCTION(BlueprintCallable, Category = "Tone Generator")
+	void SetVolume(float InVolume);
+
+	virtual ISoundGeneratorPtr CreateSoundGenerator(int32 InSampleRate, int32 InNumChannels) override;
+
+public:
+
+private:
+	ISoundGeneratorPtr ToneGenerator;
+};
diff --git a/Engine/Plugins/Runtime/Synthesis/Source/Synthesis/Private/SynthComponentToneGenerator.cpp b/Engine/Plugins/Runtime/Synthesis/Source/Synthesis/Private/SynthComponentToneGenerator.cpp
new file mode 100644
--- /dev/null
+++ b/Engine/Plugins/Runtime/Synthesis/Source/Synthesis/Private/SynthComponentToneGenerator.cpp
@@ -0,0 +1,85 @@
+// Copyright Epic Games, Inc. All Rights Reserved.
+
+
+#include "SynthComponents/SynthComponentToneGenerator.h"
+
+FToneGenerator::FToneGenerator(int32 InSampleRate, int32 InNumChannels, int32 InFrequency, float InVolume)
+	: NumChannels(InNumChannels)
+{
+	SineOsc.Init(InSampleRate, InFrequency, InVolume);
+}
+
+FToneGenerator::~FToneGenerator()
+{
+}
+
+void FToneGenerator::SetFrequency(float InFrequency)
+{
+	SynthCommand([this, InFrequency]()
+	{
+		SineOsc.SetFrequency(InFrequency);
+	});
+}
+
+void FToneGenerator::SetVolume(float InVolume)
+{
+	SynthCommand([this, InVolume]()
+	{
+		SineOsc.SetScale(InVolume);
+	});
+}
+
+int32 FToneGenerator::OnGenerateAudio(float* OutAudio, int32 NumSamples)
+{
+	check(NumChannels != 0);
+
+	int32 NumFrames = NumSamples / NumChannels;
+	int32 SampleIndex = 0;
+
+	for (int32 FrameIndex = 0; FrameIndex < NumFrames; ++FrameIndex)
+	{
+		float Sample = SineOsc.ProcessAudio() * 0.5f;
+		for (int32 Channel = 0; Channel < NumChannels; ++Channel)
+		{
+			OutAudio[SampleIndex++] = Sample;
+		}
+	}
+	return NumSamples;
+}
+
+USynthComponentToneGenerator::USynthComponentToneGenerator(const FObjectInitializer& ObjInitializer)
+	: Super(ObjInitializer)
+{
+	Frequency = 440.0f;
+	Volume = 0.5f;
+	NumChannels = 1;
+}
+
+USynthComponentToneGenerator::~USynthComponentToneGenerator()
+{
+}
+
+void USynthComponentToneGenerator::SetFrequency(float InFrequency)
+{
+	Frequency = InFrequency;
+	if (ToneGenerator.IsValid())
+	{
+		FToneGenerator* ToneGen = static_cast<FToneGenerator*>(ToneGenerator.Get());
+		ToneGen->SetFrequency(InFrequency);
+	}
+}
+
+void USynthComponentToneGenerator::SetVolume(float InVolume)
+{
+	Volume = InVolume;
+	if (ToneGenerator.IsValid())
+	{
+		FToneGenerator* ToneGen = static_cast<FToneGenerator*>(ToneGenerator.Get());
+		ToneGen->SetVolume(InVolume);
+	}
+}
+
+ISoundGeneratorPtr USynthComponentToneGenerator::CreateSoundGenerator(int32 InSampleRate, int32 InNumChannels)
+{
+	return ToneGenerator = ISoundGeneratorPtr(new FToneGenerator(InSampleRate, InNumChannels, Frequency, Volume));
+}
\ No newline at end of file
diff --git a/Engine/Source/Runtime/AudioMixer/Private/AudioMixerSource.cpp b/Engine/Source/Runtime/AudioMixer/Private/AudioMixerSource.cpp
--- a/Engine/Source/Runtime/AudioMixer/Private/AudioMixerSource.cpp
+++ b/Engine/Source/Runtime/AudioMixer/Private/AudioMixerSource.cpp
@@ -498,8 +498,8 @@ namespace Audio
 			}
 		}
 
-		check(!MixerSourceBuffer.IsValid());
-		MixerSourceBuffer = FMixerSourceBuffer::Create(*MixerBuffer, SoundWave, InWaveInstance->LoopingMode, bIsSeeking);
+		check(!MixerSourceBuffer.IsValid());		
+		MixerSourceBuffer = FMixerSourceBuffer::Create(AudioDevice->GetSampleRate(), *MixerBuffer, SoundWave, InWaveInstance->LoopingMode, bIsSeeking);
 		
 		if (!MixerSourceBuffer.IsValid())
 		{
diff --git a/Engine/Source/Runtime/AudioMixer/Private/AudioMixerSourceBuffer.cpp b/Engine/Source/Runtime/AudioMixer/Private/AudioMixerSourceBuffer.cpp
--- a/Engine/Source/Runtime/AudioMixer/Private/AudioMixerSourceBuffer.cpp
+++ b/Engine/Source/Runtime/AudioMixer/Private/AudioMixerSourceBuffer.cpp
@@ -56,7 +56,7 @@ namespace Audio
 		return CurrentSample >= NumSamples;
 	}
 
-	TSharedPtr<FMixerSourceBuffer, ESPMode::ThreadSafe> FMixerSourceBuffer::Create(FMixerBuffer& InBuffer, USoundWave& InWave, ELoopingMode InLoopingMode, bool bInIsSeeking, bool bInForceSyncDecode)
+	TSharedPtr<FMixerSourceBuffer, ESPMode::ThreadSafe> FMixerSourceBuffer::Create(int32 InSampleRate, FMixerBuffer& InBuffer, USoundWave& InWave, ELoopingMode InLoopingMode, bool bInIsSeeking, bool bInForceSyncDecode)
 	{
 		LLM_SCOPE(ELLMTag::AudioMixer);
 
@@ -72,11 +72,11 @@ namespace Audio
 			return nullptr;
 		}
 
-		TSharedPtr<FMixerSourceBuffer, ESPMode::ThreadSafe> NewSourceBuffer = MakeShareable(new FMixerSourceBuffer(InBuffer, InWave, InLoopingMode, bInIsSeeking, bInForceSyncDecode));
+		TSharedPtr<FMixerSourceBuffer, ESPMode::ThreadSafe> NewSourceBuffer = MakeShareable(new FMixerSourceBuffer(InSampleRate, InBuffer, InWave, InLoopingMode, bInIsSeeking, bInForceSyncDecode));
 		return NewSourceBuffer;
 	}
 
-	FMixerSourceBuffer::FMixerSourceBuffer(FMixerBuffer& InBuffer, USoundWave& InWave, ELoopingMode InLoopingMode, bool bInIsSeeking, bool bInForceSyncDecode)
+	FMixerSourceBuffer::FMixerSourceBuffer(int32 InSampleRate, FMixerBuffer& InBuffer, USoundWave& InWave, ELoopingMode InLoopingMode, bool bInIsSeeking, bool bInForceSyncDecode)
 		: NumBuffersQeueued(0)
 		, CurrentBuffer(0)
 		, SoundWave(&InWave)
@@ -95,8 +95,15 @@ namespace Audio
 		, bIsBus(InWave.bIsSourceBus)
 		, bForceSyncDecode(bInForceSyncDecode)
 	{
+		// TODO: remove the need to do this here. 1) remove need for decoders to depend on USoundWave and 2) remove need for procedural sounds to use USoundWaveProcedural
 		InWave.AddPlayingSource(this);
 
+		// Retrieve a sound generator if this is a procedural sound wave
+		if (bProcedural)
+		{
+			SoundGenerator = InWave.CreateSoundGenerator(InSampleRate, NumChannels);
+		}
+
 		const uint32 TotalSamples = MONO_PCM_BUFFER_SAMPLES * NumChannels;
 		for (int32 BufferIndex = 0; BufferIndex < Audio::MAX_BUFFERS_QUEUED; ++BufferIndex)
 		{
@@ -326,9 +333,20 @@ namespace Audio
 			
 			if (Lock.IsLocked() && SoundWave)
 			{
-				check(SoundWave && SoundWave->bProcedural);
 				FProceduralAudioTaskData NewTaskData;
-				NewTaskData.ProceduralSoundWave = SoundWave;
+
+				// Pass the generator instance to the async task
+				if (SoundGenerator.IsValid())
+				{
+					NewTaskData.SoundGenerator = SoundGenerator;
+				}
+				else
+				{
+					// Otherwise pass the raw sound wave procedural ptr.
+					check(SoundWave && SoundWave->bProcedural);
+					NewTaskData.ProceduralSoundWave = SoundWave;
+				}
+
 				NewTaskData.AudioData = SourceVoiceBuffers[BufferIndex]->AudioData.GetData();
 				NewTaskData.NumSamples = MaxSamples;
 				NewTaskData.NumChannels = NumChannels;
@@ -550,11 +568,22 @@ namespace Audio
 	void FMixerSourceBuffer::OnBeginGenerate()
 	{
 		FScopeTryLock Lock(&SoundWaveCritSec);
+		if (!Lock.IsLocked())
+		{
+			return;
+		}
 
-		if (Lock.IsLocked() && SoundWave && bProcedural)
+		if (SoundGenerator.IsValid())
+		{
+			SoundGenerator->OnBeginGenerate();
+		}
+		else
 		{
-			check(SoundWave && SoundWave->bProcedural);
-			SoundWave->OnBeginGenerate();
+			if (SoundWave && bProcedural)
+			{
+				check(SoundWave && SoundWave->bProcedural);
+				SoundWave->OnBeginGenerate();
+			}
 		}
 	}
 
@@ -569,12 +598,18 @@ namespace Audio
 			return;
 		}
 
-
-		// Only need to call OnEndGenerate and access SoundWave here if we successfully initialized
-		if (SoundWave && bInitialized && bProcedural)
+		if (SoundGenerator.IsValid())
+		{
+			SoundGenerator->OnEndGenerate();
+		}
+		else
 		{
-			check(SoundWave && SoundWave->bProcedural);
-			SoundWave->OnEndGenerate();
+			// Only need to call OnEndGenerate and access SoundWave here if we successfully initialized
+			if (SoundWave && bInitialized && bProcedural)
+			{
+				check(SoundWave && SoundWave->bProcedural);
+				SoundWave->OnEndGenerate();
+			}
 		}
 	}
 
diff --git a/Engine/Source/Runtime/AudioMixer/Private/AudioMixerSourceBuffer.h b/Engine/Source/Runtime/AudioMixer/Private/AudioMixerSourceBuffer.h
--- a/Engine/Source/Runtime/AudioMixer/Private/AudioMixerSourceBuffer.h
+++ b/Engine/Source/Runtime/AudioMixer/Private/AudioMixerSourceBuffer.h
@@ -6,6 +6,7 @@
 #include "AudioMixerBuffer.h"
 #include "AudioMixerSourceManager.h"
 #include "Sound/SoundWave.h"
+#include "Sound/SoundGenerator.h"
 
 namespace Audio
 {
@@ -49,7 +50,7 @@ namespace Audio
 	class FMixerSourceBuffer : public ISoundWaveClient
 	{
 	public:		
-		static FMixerSourceBufferPtr Create(FMixerBuffer& InBuffer, USoundWave& InWave, ELoopingMode InLoopingMode, bool bInIsSeeking, bool bInForceSyncDecode = false);
+		static FMixerSourceBufferPtr Create(int32 InSampleRate, FMixerBuffer& InBuffer, USoundWave& InWave, ELoopingMode InLoopingMode, bool bInIsSeeking, bool bInForceSyncDecode = false);
 
 		~FMixerSourceBuffer();
 
@@ -96,7 +97,7 @@ namespace Audio
 		void OnEndGenerate();
 		void ClearWave() { SoundWave = nullptr; }
 	private:
-		FMixerSourceBuffer(FMixerBuffer& InBuffer, USoundWave& InWave, ELoopingMode InLoopingMode, bool bInIsSeeking, bool bInForceSyncDecode = false);
+		FMixerSourceBuffer(int32 InSampleRate, FMixerBuffer& InBuffer, USoundWave& InWave, ELoopingMode InLoopingMode, bool bInIsSeeking, bool bInForceSyncDecode = false);
 
 		void SubmitInitialPCMBuffers();
 		void SubmitInitialRealtimeBuffers();
@@ -114,6 +115,7 @@ namespace Audio
 		int32 CurrentBuffer;
 		// SoundWaves are only set for procedural sound waves
 		USoundWave* SoundWave;
+		ISoundGeneratorPtr SoundGenerator;
 		IAudioTask* AsyncRealtimeAudioTask;
 		ICompressedAudioInfo* DecompressionState;
 		ELoopingMode LoopingMode;
diff --git a/Engine/Source/Runtime/AudioMixer/Private/AudioMixerSourceDecode.cpp b/Engine/Source/Runtime/AudioMixer/Private/AudioMixerSourceDecode.cpp
--- a/Engine/Source/Runtime/AudioMixer/Private/AudioMixerSourceDecode.cpp
+++ b/Engine/Source/Runtime/AudioMixer/Private/AudioMixerSourceDecode.cpp
@@ -9,6 +9,7 @@
 #include "AudioMixerBuffer.h"
 #include "Async/Async.h"
 #include "AudioDecompress.h"
+#include "Sound/SoundGenerator.h"
 
 static int32 ForceSyncAudioDecodesCvar = 0;
 FAutoConsoleVariableRef CVarForceSyncAudioDecodes(
@@ -56,45 +57,54 @@ public:
 			case EAudioTaskType::Procedural:
 			{
 				QUICK_SCOPE_CYCLE_COUNTER(STAT_FAsyncDecodeWorker_Procedural);
-				// Make sure we've been flagged as active
-				if (!ProceduralTaskData.ProceduralSoundWave || !ProceduralTaskData.ProceduralSoundWave->IsGeneratingAudio())
+				if (ProceduralTaskData.SoundGenerator.IsValid())
 				{
-					// Act as if we generated audio, but return silence.
+					// Pre-zero the buffer before calling into the generator code as a convenience to implementers
 					FMemory::Memzero(ProceduralTaskData.AudioData, ProceduralTaskData.NumSamples * sizeof(float));
-					ProceduralResult.NumSamplesWritten = ProceduralTaskData.NumSamples;
-					return;
+					ProceduralResult.NumSamplesWritten = ProceduralTaskData.SoundGenerator->GetNextBuffer(ProceduralTaskData.AudioData, ProceduralTaskData.NumSamples);
 				}
-
-				// If we're not a float format, we need to convert the format to float
-				const EAudioMixerStreamDataFormat::Type FormatType = ProceduralTaskData.ProceduralSoundWave->GetGeneratedPCMDataFormat();
-				if (FormatType != EAudioMixerStreamDataFormat::Float)
+				else
 				{
-					check(FormatType == EAudioMixerStreamDataFormat::Int16);
+					// Make sure we've been flagged as active
+					if (!ProceduralTaskData.ProceduralSoundWave || !ProceduralTaskData.ProceduralSoundWave->IsGeneratingAudio())
+					{
+						// Act as if we generated audio, but return silence.
+						FMemory::Memzero(ProceduralTaskData.AudioData, ProceduralTaskData.NumSamples * sizeof(float));
+						ProceduralResult.NumSamplesWritten = ProceduralTaskData.NumSamples;
+						return;
+					}
 
-					int32 NumChannels = ProceduralTaskData.NumChannels;
-					int32 ByteSize = NumChannels * ProceduralTaskData.NumSamples * sizeof(int16);
+					// If we're not a float format, we need to convert the format to float
+					const EAudioMixerStreamDataFormat::Type FormatType = ProceduralTaskData.ProceduralSoundWave->GetGeneratedPCMDataFormat();
+					if (FormatType != EAudioMixerStreamDataFormat::Float)
+					{
+						check(FormatType == EAudioMixerStreamDataFormat::Int16);
+
+						int32 NumChannels = ProceduralTaskData.NumChannels;
+						int32 ByteSize = NumChannels * ProceduralTaskData.NumSamples * sizeof(int16);
 
-					TArray<uint8> DecodeBuffer;
-					DecodeBuffer.AddUninitialized(ByteSize);
+						TArray<uint8> DecodeBuffer;
+						DecodeBuffer.AddUninitialized(ByteSize);
 
-					const int32 NumBytesWritten = ProceduralTaskData.ProceduralSoundWave->GeneratePCMData(DecodeBuffer.GetData(), ProceduralTaskData.NumSamples);
+						const int32 NumBytesWritten = ProceduralTaskData.ProceduralSoundWave->GeneratePCMData(DecodeBuffer.GetData(), ProceduralTaskData.NumSamples);
 
-					check(NumBytesWritten <= ByteSize);
+						check(NumBytesWritten <= ByteSize);
 
-					ProceduralResult.NumSamplesWritten = NumBytesWritten / sizeof(int16);
+						ProceduralResult.NumSamplesWritten = NumBytesWritten / sizeof(int16);
 
-					// Convert the buffer to float
-					int16* DecodedBufferPtr = (int16*)DecodeBuffer.GetData();
-					for (int32 SampleIndex = 0; SampleIndex < ProceduralResult.NumSamplesWritten; ++SampleIndex)
+						// Convert the buffer to float
+						int16* DecodedBufferPtr = (int16*)DecodeBuffer.GetData();
+						for (int32 SampleIndex = 0; SampleIndex < ProceduralResult.NumSamplesWritten; ++SampleIndex)
+						{
+							ProceduralTaskData.AudioData[SampleIndex] = (float)(DecodedBufferPtr[SampleIndex]) / 32768.0f;
+						}
+					}
+					else
 					{
-						ProceduralTaskData.AudioData[SampleIndex] = (float)(DecodedBufferPtr[SampleIndex]) / 32768.0f;
+						const int32 NumBytesWritten = ProceduralTaskData.ProceduralSoundWave->GeneratePCMData((uint8*)ProceduralTaskData.AudioData, ProceduralTaskData.NumSamples);
+						ProceduralResult.NumSamplesWritten = NumBytesWritten / sizeof(float);
 					}
 				}
-				else
-				{
-					const int32 NumBytesWritten = ProceduralTaskData.ProceduralSoundWave->GeneratePCMData((uint8*)ProceduralTaskData.AudioData, ProceduralTaskData.NumSamples);
-					ProceduralResult.NumSamplesWritten = NumBytesWritten / sizeof(float);
-				}
 			}
 			break;
 
diff --git a/Engine/Source/Runtime/AudioMixer/Private/AudioMixerSourceDecode.h b/Engine/Source/Runtime/AudioMixer/Private/AudioMixerSourceDecode.h
--- a/Engine/Source/Runtime/AudioMixer/Private/AudioMixerSourceDecode.h
+++ b/Engine/Source/Runtime/AudioMixer/Private/AudioMixerSourceDecode.h
@@ -5,6 +5,7 @@
 #include "CoreMinimal.h"
 #include "Containers/Queue.h"
 #include "Sound/SoundWaveProcedural.h"
+#include "Sound/SoundGenerator.h"
 #include "AudioDecompress.h"
 #include "AudioMixerBuffer.h"
 
@@ -18,8 +19,12 @@ namespace Audio
 	struct FProceduralAudioTaskData
 	{
 		// The procedural sound wave ptr to use to generate audio with
+		// TODO: remove the need for this
 		USoundWave* ProceduralSoundWave;
 
+		// The sound generator to use to generate audio
+		ISoundGeneratorPtr SoundGenerator;
+
 		// The audio buffer to fill from the results of the generation
 		float* AudioData;
 
diff --git a/Engine/Source/Runtime/AudioMixer/Private/Components/SynthComponent.cpp b/Engine/Source/Runtime/AudioMixer/Private/Components/SynthComponent.cpp
--- a/Engine/Source/Runtime/AudioMixer/Private/Components/SynthComponent.cpp
+++ b/Engine/Source/Runtime/AudioMixer/Private/Components/SynthComponent.cpp
@@ -98,12 +98,21 @@ int32 USynthSound::OnGeneratePCMAudio(TArray<uint8>& OutAudio, int32 NumSamples)
 void USynthSound::OnEndGenerate()
 {
 	// Mark pending kill can null this out on the game thread in rare cases.
-	if(OwningSynthComponent)
+	if (OwningSynthComponent)
 	{
 		OwningSynthComponent->OnEndGenerate();
 	}
 }
 
+ISoundGeneratorPtr USynthSound::CreateSoundGenerator(int32 InSampleRate, int32 InNumChannels)
+{
+	if (OwningSynthComponent)
+	{
+		return OwningSynthComponent->CreateSoundGeneratorInternal(SampleRate, NumChannels);
+	}
+	return nullptr;
+}
+
 Audio::EAudioMixerStreamDataFormat::Type USynthSound::GetGeneratedPCMDataFormat() const
 {
 	// Only audio mixer supports return float buffers
@@ -403,7 +412,7 @@ void USynthComponent::PumpPendingMessages()
 	}
 }
 
-FAudioDevice* USynthComponent::GetAudioDevice()
+FAudioDevice* USynthComponent::GetAudioDevice() const
 {
 	// If the synth component has a world, that means it was already registed with that world
 	if (UWorld* World = GetWorld())
@@ -482,7 +491,10 @@ void USynthComponent::Start()
 
 		if (IsActive())
 		{
-			PendingSynthEvents.Enqueue(ESynthEvent::Start);
+			if (!SoundGenerator.IsValid())
+			{
+				PendingSynthEvents.Enqueue(ESynthEvent::Start);
+			}
 		}
 	}
 }
@@ -491,7 +503,10 @@ void USynthComponent::Stop()
 {
 	if (IsActive())
 	{
-		PendingSynthEvents.Enqueue(ESynthEvent::Stop);
+		if (!SoundGenerator.IsValid())
+		{
+			PendingSynthEvents.Enqueue(ESynthEvent::Stop);
+		}
 
 		if (AudioComponent)
 		{
@@ -525,5 +540,17 @@ void USynthComponent::SetSubmixSend(USoundSubmixBase* Submix, float SendLevel)
 
 void USynthComponent::SynthCommand(TFunction<void()> Command)
 {
-	CommandQueue.Enqueue(MoveTemp(Command));
+	if (SoundGenerator.IsValid())
+	{
+		UE_LOG(LogAudioMixer, Error, TEXT("Synthesis component '%s' has implemented a sound generator interface. Do not call SynthCommand on the USynthComponent)."), *GetName());
+	}
+	else
+	{
+		CommandQueue.Enqueue(MoveTemp(Command));
+	}
+}
+
+ISoundGeneratorPtr USynthComponent::CreateSoundGeneratorInternal(int32 InSampleRate, int32 InNumChannels)
+{	
+	return SoundGenerator = CreateSoundGenerator(InSampleRate, InNumChannels);
 }
diff --git a/Engine/Source/Runtime/AudioMixer/Private/SoundWaveDecoder.cpp b/Engine/Source/Runtime/AudioMixer/Private/SoundWaveDecoder.cpp
--- a/Engine/Source/Runtime/AudioMixer/Private/SoundWaveDecoder.cpp
+++ b/Engine/Source/Runtime/AudioMixer/Private/SoundWaveDecoder.cpp
@@ -53,7 +53,7 @@ namespace Audio
 		const ELoopingMode LoopingMode = SoundWave->bLooping ? ELoopingMode::LOOP_Forever : ELoopingMode::LOOP_Never;
 		const bool bIsSeeking = SeekTime > 0.0f;
 
-		MixerSourceBuffer = FMixerSourceBuffer::Create(*MixerBuffer, *SoundWave, LoopingMode, bIsSeeking, bForceSyncDecode);
+		MixerSourceBuffer = FMixerSourceBuffer::Create(InSampleRate, *MixerBuffer, *SoundWave, LoopingMode, bIsSeeking, bForceSyncDecode);
 		return MixerSourceBuffer.IsValid();
 	}
 
diff --git a/Engine/Source/Runtime/AudioMixer/Public/Components/SynthComponent.h b/Engine/Source/Runtime/AudioMixer/Public/Components/SynthComponent.h
--- a/Engine/Source/Runtime/AudioMixer/Public/Components/SynthComponent.h
+++ b/Engine/Source/Runtime/AudioMixer/Public/Components/SynthComponent.h
@@ -7,6 +7,7 @@
 #include "Engine/EngineTypes.h"
 #include "IAudioExtensionPlugin.h"
 #include "Sound/SoundWaveProcedural.h"
+#include "Sound/SoundGenerator.h"
 #include "UObject/ObjectMacros.h"
 
 #include "SynthComponent.generated.h"
@@ -50,6 +51,7 @@ class AUDIOMIXER_API USynthSound : public USoundWaveProcedural
 	virtual int32 OnGeneratePCMAudio(TArray<uint8>& OutAudio, int32 NumSamples) override;
 	virtual void OnEndGenerate() override;
 	virtual Audio::EAudioMixerStreamDataFormat::Type GetGeneratedPCMDataFormat() const override;
+	virtual ISoundGeneratorPtr CreateSoundGenerator(int32 InSampleRate, int32 InNumChannels) override;
 	/** End USoundWave */
 
 protected:
@@ -240,14 +242,19 @@ protected:
 	virtual void OnEndGenerate() {}
 
 	// Called when more audio is needed to be generated
+	// This method of generating audio is soon to be deprecated. For all new synth components, create an FSoundGenerator instance and implement CreateSoundGenerator method to create an instance.
 	virtual int32 OnGenerateAudio(float* OutAudio, int32 NumSamples) PURE_VIRTUAL(USynthComponent::OnGenerateAudio, return 0; );
 
+	// Implemented by the synth component to create a generator object instead of generating audio directly on the synth component.
+	// This method prevents UObjects from having to exist in the audio render thread.
+	virtual ISoundGeneratorPtr CreateSoundGenerator(int32 InSampleRate, int32 InNumChannels) { return nullptr; }
+
 	// Called by procedural sound wave
 	// Returns the number of samples actually generated
 	int32 OnGeneratePCMAudio(float* GeneratedPCMData, int32 NumSamples);
 
 	// Gets the audio device associated with this synth component
-	FAudioDevice* GetAudioDevice();
+	FAudioDevice* GetAudioDevice() const;
 
 	// Can be set by the derived class, defaults to 2
 	int32 NumChannels;
@@ -256,6 +263,8 @@ protected:
 	int32 PreferredBufferLength;
 
 private:
+	// Creates the synth component's sound generator, calls into overridden client code to create the instance.
+	ISoundGeneratorPtr CreateSoundGeneratorInternal(int32 InSampleRate, int32 InNumChannels);
 
 	UPROPERTY(Transient)
 	USynthSound* Synth;
@@ -285,5 +294,9 @@ private:
 
 	TQueue<ESynthEvent> PendingSynthEvents;
 
+	// Synth component's handle to its sound generator instance.
+	// used to forward BP functions to the instance directly.
+	ISoundGeneratorPtr SoundGenerator;
+
 	friend class USynthSound;
 };
\ No newline at end of file
diff --git a/Engine/Source/Runtime/Engine/Classes/Sound/SoundBase.h b/Engine/Source/Runtime/Engine/Classes/Sound/SoundBase.h
--- a/Engine/Source/Runtime/Engine/Classes/Sound/SoundBase.h
+++ b/Engine/Source/Runtime/Engine/Classes/Sound/SoundBase.h
@@ -11,6 +11,7 @@
 #include "SoundConcurrency.h"
 #include "SoundSourceBusSend.h"
 #include "SoundSubmixSend.h"
+#include "SoundGenerator.h"
 #include "UObject/Object.h"
 #include "UObject/ObjectMacros.h"
 #include "AudioDeviceManager.h"
@@ -257,5 +258,9 @@ public:
 	/** Queries if the sound has cooked FFT or envelope data. */
 	virtual bool HasCookedFFTData() const { return false; }
 	virtual bool HasCookedAmplitudeEnvelopeData() const { return false; }
+
+	/** Creates a sound generator instance from this sound base. Return true if this is being implemented by a subclass. Sound generators procedurally generate audio in the audio render thread. */
+	virtual ISoundGeneratorPtr CreateSoundGenerator(int32 InSampleRate, int32 InNumChannels) { return nullptr; }
+
 };
 
diff --git a/Engine/Source/Runtime/Engine/Classes/Sound/SoundGenerator.h b/Engine/Source/Runtime/Engine/Classes/Sound/SoundGenerator.h
new file mode 100644
--- /dev/null
+++ b/Engine/Source/Runtime/Engine/Classes/Sound/SoundGenerator.h
@@ -0,0 +1,55 @@
+// Copyright Epic Games, Inc. All Rights Reserved.
+
+#pragma once
+
+#include "CoreMinimal.h"
+#include "Containers\Queue.h"
+#include "Templates\Function.h"
+
+enum class EGeneratorState : uint8
+{
+	IsIdle,
+	IsGenerating,
+};
+
+class ENGINE_API ISoundGenerator
+{
+public:
+	ISoundGenerator();
+	virtual ~ISoundGenerator();
+
+	// Called when a new buffer is required. 
+	virtual int32 OnGenerateAudio(float* OutAudio, int32 NumSamples) = 0;
+
+	// Optional. Called on audio generator thread right when the generator begins generating.
+	virtual void OnBeginGenerate() {}
+
+	// Optional. Called on audio generator thread right when the generator ends generating.
+	virtual void OnEndGenerate() {}
+
+	// Retrieves the next buffer of audio from the generator, called from the audio mixer
+	int32 GetNextBuffer(float* OutAudio, int32 NumSamples);
+
+	// Returns the current state of the sound generator
+	EGeneratorState GetGeneratorState() const { return GeneratorState; }
+
+protected:
+
+	// Protected method to execute lambda in audio render thread
+	// Used for conveying parameter changes or events to the generator thread.
+	void SynthCommand(TUniqueFunction<void()> Command);
+
+private:
+
+	void PumpPendingMessages();
+
+	// The current state of the generator
+	EGeneratorState GeneratorState;
+
+	// The command queue used to convey commands from game thread to generator thread 
+	TQueue<TUniqueFunction<void()>> CommandQueue;
+
+	friend class USynthComponent;
+};
+
+typedef TSharedPtr<ISoundGenerator, ESPMode::ThreadSafe> ISoundGeneratorPtr;
\ No newline at end of file
diff --git a/Engine/Source/Runtime/Engine/Private/SoundGenerator.cpp b/Engine/Source/Runtime/Engine/Private/SoundGenerator.cpp
new file mode 100644
--- /dev/null
+++ b/Engine/Source/Runtime/Engine/Private/SoundGenerator.cpp
@@ -0,0 +1,34 @@
+// Copyright Epic Games, Inc. All Rights Reserved.
+
+#include "Sound/SoundGenerator.h"
+
+ISoundGenerator::ISoundGenerator()
+{
+}
+
+ISoundGenerator::~ISoundGenerator()
+{
+}
+
+int32 ISoundGenerator::GetNextBuffer(float* OutAudio, int32 NumSamples)
+{
+	PumpPendingMessages();
+
+	return OnGenerateAudio(OutAudio, NumSamples);
+}
+
+void ISoundGenerator::SynthCommand(TUniqueFunction<void()> Command)
+{
+	CommandQueue.Enqueue(MoveTemp(Command));
+}
+
+void ISoundGenerator::PumpPendingMessages()
+{
+	TUniqueFunction<void()> Command;
+	while (CommandQueue.Dequeue(Command))
+	{
+		Command();
+	}
+}
+
+
diff --git a/Engine/Source/Runtime/SignalProcessing/Private/SinOsc.cpp b/Engine/Source/Runtime/SignalProcessing/Private/SinOsc.cpp
--- a/Engine/Source/Runtime/SignalProcessing/Private/SinOsc.cpp
+++ b/Engine/Source/Runtime/SignalProcessing/Private/SinOsc.cpp
@@ -46,6 +46,16 @@ namespace Audio
 		SetFrequency(InFrequencyHz);
 	}
 
+	void FSineOsc::SetScale(const float InScale)
+	{
+		Scale = InScale;
+	}
+
+	void FSineOsc::SetAdd(const float InAdd)
+	{
+		Add = InAdd;
+	}
+
 	void FSineOsc::SetFrequency(const float InFrequencyHz)
 	{
 		check(SampleRate > 0);
diff --git a/Engine/Source/Runtime/SignalProcessing/Public/DSP/SinOsc.h b/Engine/Source/Runtime/SignalProcessing/Public/DSP/SinOsc.h
--- a/Engine/Source/Runtime/SignalProcessing/Public/DSP/SinOsc.h
+++ b/Engine/Source/Runtime/SignalProcessing/Public/DSP/SinOsc.h
@@ -28,6 +28,12 @@ namespace Audio
 		/** Initialize the oscillator with a sample rate and new frequency. Must be called before playing oscillator. */
 		void Init(const int32 InSampleRate, const float InFrequencyHz, const float Scale = 1.0f, const float Add = 0.0f);
 
+		/** Sets the scale of the oscillator. */
+		void SetScale(const float InScale);
+
+		/** Sets the scale of the oscillator. */
+		void SetAdd(const float InAdd);
+
 		/** Sets the frequency of the oscillator in Hz (based on sample rate). Performs initial condition calculation to avoid pops. */
 		void SetFrequency(const float InFrequencyHz);
 
