- Edigrate Aaron's fix for dynamics processor.
- Set bExternalSubmix to whether or not an external input is supplied, even when device is not yet fully initialized to avoid potential scenario where a few frames being processed with the input as the key as opposed to nothing (as external input is not yet initialized)
- Clear up UX as Key Gain is no longer technically a Key Gain, but rather an External Input Gain
#rnx

diff --git a/Engine/Source/Runtime/AudioMixer/Classes/SubmixEffects/AudioMixerSubmixEffectDynamicsProcessor.h b/Engine/Source/Runtime/AudioMixer/Classes/SubmixEffects/AudioMixerSubmixEffectDynamicsProcessor.h
--- a/Engine/Source/Runtime/AudioMixer/Classes/SubmixEffects/AudioMixerSubmixEffectDynamicsProcessor.h
+++ b/Engine/Source/Runtime/AudioMixer/Classes/SubmixEffects/AudioMixerSubmixEffectDynamicsProcessor.h
@@ -133,8 +133,8 @@ struct AUDIOMIXER_API FSubmixEffectDynamicsProcessorSettings
 	UPROPERTY(EditAnywhere, BlueprintReadWrite, Category = Sidechain, meta = (DisplayName = "Key Audition"))
 	uint8 bKeyAudition : 1;
 
-	// Gain to apply to key signal (external signal if supplied or input signal if disabled)
-	UPROPERTY(EditAnywhere, BlueprintReadWrite, Category = Sidechain, meta = (DisplayName = "Key Gain (dB)", UIMin = "-60.0", UIMax = "30.0"))
+	// Gain to apply to key signal if external input is supplied
+	UPROPERTY(EditAnywhere, BlueprintReadWrite, Category = Sidechain, meta = (DisplayName = "External Input Gain (dB)", EditCondition = "ExternalSubmix != nullptr", UIMin = "-60.0", UIMax = "30.0"))
 	float KeyGainDb;
 
 	// The output gain of the dynamics processor
diff --git a/Engine/Source/Runtime/AudioMixer/Private/Effects/AudioMixerSubmixEffectDynamicsProcessor.cpp b/Engine/Source/Runtime/AudioMixer/Private/Effects/AudioMixerSubmixEffectDynamicsProcessor.cpp
--- a/Engine/Source/Runtime/AudioMixer/Private/Effects/AudioMixerSubmixEffectDynamicsProcessor.cpp
+++ b/Engine/Source/Runtime/AudioMixer/Private/Effects/AudioMixerSubmixEffectDynamicsProcessor.cpp
@@ -210,7 +210,7 @@ void FSubmixEffectDynamicsProcessor::SetExternalSubmix(USoundSubmix* InSoundSubm
 	else
 	{
 		DeviceCreatedHandle = FAudioDeviceManagerDelegates::OnAudioDeviceCreated.AddRaw(this, &FSubmixEffectDynamicsProcessor::OnNewDeviceCreated);
-		bUseExternalSubmix = false;
+		bUseExternalSubmix = ExternalSubmix.IsValid();
 	}
 }
 
