fix for mem stomp in editor on importing audio files.

could likely happen on any audio files, not necessarily multi-channel audio files, though that may exacerbate it.

diff --git a/Engine/Source/Runtime/AudioMixer/Private/SoundFileIO/SoundFileIOEnums.h b/Engine/Source/Runtime/AudioMixer/Private/SoundFileIO/SoundFileIOEnums.h
--- a/Engine/Source/Runtime/AudioMixer/Private/SoundFileIO/SoundFileIOEnums.h
+++ b/Engine/Source/Runtime/AudioMixer/Private/SoundFileIO/SoundFileIOEnums.h
@@ -165,7 +165,8 @@ namespace Audio
 	 */
 	namespace ESoundFileChannelMap
 	{
-		enum class Type : uint8
+		// this is used to populate an array which is passed into a sound file API call so must be uint32
+		enum class Type : uint32
 		{
 			INVALID = 0,
 			MONO,
