Edigrate AudioBus Support in DynamicsProcessor
- Ability to subscribe submix listener via patch splitter

diff --git a/Engine/Source/Runtime/AudioMixer/Classes/SubmixEffects/AudioMixerSubmixEffectDynamicsProcessor.h b/Engine/Source/Runtime/AudioMixer/Classes/SubmixEffects/AudioMixerSubmixEffectDynamicsProcessor.h
--- a/Engine/Source/Runtime/AudioMixer/Classes/SubmixEffects/AudioMixerSubmixEffectDynamicsProcessor.h
+++ b/Engine/Source/Runtime/AudioMixer/Classes/SubmixEffects/AudioMixerSubmixEffectDynamicsProcessor.h
@@ -4,6 +4,8 @@
 #include "AudioDevice.h"
 #include "Delegates/IDelegateInstance.h"
 #include "DSP/DynamicsProcessor.h"
+#include "DSP/MultithreadedPatching.h"
+#include "Misc/ScopeLock.h"
 #include "Sound/SoundEffectSubmix.h"
 #include "Sound/SoundSubmix.h"
 #include "Sound/SoundSubmixSend.h"
@@ -48,6 +50,21 @@ enum class ESubmixEffectDynamicsChannelLinkMode : uint8
 	Count UMETA(Hidden)
 };
 
+UENUM(BlueprintType)
+enum class ESubmixEffectDynamicsKeySource : uint8
+{
+	// Defaults to use local submix (input) as key
+	Default = 0,
+
+	// Uses audio bus as key
+	AudioBus,
+
+	// Uses external submix as key
+	Submix,
+
+	Count UMETA(Hidden)
+};
+
 USTRUCT(BlueprintType)
 struct AUDIOMIXER_API FSubmixEffectDynamicProcessorFilterSettings
 {
@@ -81,103 +98,100 @@ struct AUDIOMIXER_API FSubmixEffectDynamicsProcessorSettings
 
 	// Type of processor to apply
 	UPROPERTY(EditAnywhere, BlueprintReadWrite, Category = General, meta = (DisplayName = "Type"))
-	ESubmixEffectDynamicsProcessorType DynamicsProcessorType;
+	ESubmixEffectDynamicsProcessorType DynamicsProcessorType = ESubmixEffectDynamicsProcessorType::Compressor;
 
 	// Mode of peak detection used on input key signal
-	UPROPERTY(EditAnywhere, BlueprintReadWrite, Category = Dynamics)
-	ESubmixEffectDynamicsPeakMode PeakMode;
+	UPROPERTY(EditAnywhere, BlueprintReadWrite, Category = Dynamics, meta = (EditCondition = "!bBypass"))
+	ESubmixEffectDynamicsPeakMode PeakMode = ESubmixEffectDynamicsPeakMode::RootMeanSquared;
 
 	// Mode of peak detection if key signal is multi-channel
-	UPROPERTY(EditAnywhere, BlueprintReadWrite, Category = Dynamics)
-	ESubmixEffectDynamicsChannelLinkMode LinkMode;
+	UPROPERTY(EditAnywhere, BlueprintReadWrite, Category = Dynamics, meta = (EditCondition = "!bBypass"))
+	ESubmixEffectDynamicsChannelLinkMode LinkMode = ESubmixEffectDynamicsChannelLinkMode::Average;
 
 	// The input gain of the dynamics processor
-	UPROPERTY(EditAnywhere, BlueprintReadWrite, Category = General, meta = (DisplayName = "Input Gain (dB)", UIMin = "-12.0", UIMax = "20.0"))
-	float InputGainDb;
+	UPROPERTY(EditAnywhere, BlueprintReadWrite, Category = General, meta = (DisplayName = "Input Gain (dB)", UIMin = "-12.0", UIMax = "20.0", EditCondition = "!bBypass"))
+	float InputGainDb = 0.0f;
 
 	// The threshold at which to perform a dynamics processing operation
-	UPROPERTY(EditAnywhere, BlueprintReadWrite, Category = Dynamics, meta = (DisplayName = "Threshold (dB)", ClampMin = "-60.0", ClampMax = "0.0", UIMin = "-60.0", UIMax = "0.0"))
-	float ThresholdDb;
+	UPROPERTY(EditAnywhere, BlueprintReadWrite, Category = Dynamics, meta = (DisplayName = "Threshold (dB)", ClampMin = "-60.0", ClampMax = "0.0", UIMin = "-60.0", UIMax = "0.0", EditCondition = "!bBypass"))
+	float ThresholdDb = -6.0f;
 
 	// The dynamics processor ratio used for compression/expansion
 	UPROPERTY(EditAnywhere, BlueprintReadWrite, Category = Dynamics, meta = (
-		EditCondition = "DynamicsProcessorType == ESubmixEffectDynamicsProcessorType::Compressor || DynamicsProcessorType == ESubmixEffectDynamicsProcessorType::Expander",
+		EditCondition = "!bBypass && DynamicsProcessorType == ESubmixEffectDynamicsProcessorType::Compressor || DynamicsProcessorType == ESubmixEffectDynamicsProcessorType::Expander",
 		ClampMin = "1.0", ClampMax = "20.0", UIMin = "1.0", UIMax = "20.0"))
-	float Ratio;
+	float Ratio = 1.5f;
 
 	// The knee bandwidth of the processor to use
-	UPROPERTY(EditAnywhere, BlueprintReadWrite, Category = Dynamics, meta = (DisplayName = "Knee (dB)", ClampMin = "0.0", ClampMax = "20.0", UIMin = "0.0", UIMax = "20.0"))
-	float KneeBandwidthDb;
+	UPROPERTY(EditAnywhere, BlueprintReadWrite, Category = Dynamics, meta = (DisplayName = "Knee (dB)", ClampMin = "0.0", ClampMax = "20.0", UIMin = "0.0", UIMax = "20.0", EditCondition = "!bBypass"))
+	float KneeBandwidthDb = 10.0f;
 
 	// The amount of time to look ahead of the current audio (Allows for transients to be included in dynamics processing)
-	UPROPERTY(EditAnywhere, BlueprintReadWrite, Category = Response,  meta = (DisplayName = "Look Ahead (ms)", ClampMin = "0.0", ClampMax = "50.0", UIMin = "0.0", UIMax = "50.0"))
-	float LookAheadMsec;
+	UPROPERTY(EditAnywhere, BlueprintReadWrite, Category = Response,  meta = (DisplayName = "Look Ahead (ms)", ClampMin = "0.0", ClampMax = "50.0", UIMin = "0.0", UIMax = "50.0", EditCondition = "!bBypass"))
+	float LookAheadMsec = 3.0f;
 
 	// The amount of time to ramp into any dynamics processing effect
-	UPROPERTY(EditAnywhere, BlueprintReadWrite, Category = Response, meta = (DisplayName = "AttackTime (ms)", ClampMin = "1.0", ClampMax = "300.0", UIMin = "1.0", UIMax = "200.0"))
-	float AttackTimeMsec;
+	UPROPERTY(EditAnywhere, BlueprintReadWrite, Category = Response, meta = (DisplayName = "AttackTime (ms)", ClampMin = "1.0", ClampMax = "300.0", UIMin = "1.0", UIMax = "200.0", EditCondition = "!bBypass"))
+	float AttackTimeMsec = 10.0f;
 
 	// The amount of time to release the dynamics processing effect
-	UPROPERTY(EditAnywhere, BlueprintReadWrite, Category = Response, meta = (DisplayName = "Release Time (ms)", ClampMin = "20.0", ClampMax = "5000.0", UIMin = "20.0", UIMax = "5000.0"))
-	float ReleaseTimeMsec;
+	UPROPERTY(EditAnywhere, BlueprintReadWrite, Category = Response, meta = (DisplayName = "Release Time (ms)", ClampMin = "20.0", ClampMax = "5000.0", UIMin = "20.0", UIMax = "5000.0", EditCondition = "!bBypass"))
+	float ReleaseTimeMsec = 100.0f;
+
+	UPROPERTY(EditAnywhere, BlueprintReadWrite, Category = Sidechain, meta = (EditCondition = "!bBypass"))
+	ESubmixEffectDynamicsKeySource KeySource = ESubmixEffectDynamicsKeySource::Default;
 
 	// If set, uses output of provided submix as modulator of input signal for dynamics processor (Uses input signal as default modulator)
-	UPROPERTY(EditAnywhere, BlueprintReadWrite, Category = Sidechain)
-	USoundSubmix* ExternalSubmix;
+	UPROPERTY(EditAnywhere, BlueprintReadWrite, Category = Sidechain, meta = (EditCondition = "!bBypass && KeySource == ESubmixEffectDynamicsKeySource::AudioBus", EditConditionHides))
+	UAudioBus* ExternalAudioBus = nullptr;
+
+	// If set, uses output of provided submix as modulator of input signal for dynamics processor (Uses input signal as default modulator)
+	UPROPERTY(EditAnywhere, BlueprintReadWrite, Category = Sidechain, meta = (EditCondition = "!bBypass && KeySource == ESubmixEffectDynamicsKeySource::Submix", EditConditionHides))
+	USoundSubmix* ExternalSubmix = nullptr;
 
 	UPROPERTY()
 	uint8 bChannelLinked_DEPRECATED : 1;
 
 	// Toggles treating the attack and release envelopes as analog-style vs digital-style (Analog will respond a bit more naturally/slower)
-	UPROPERTY(EditAnywhere, BlueprintReadWrite, Category = Response)
+	UPROPERTY(EditAnywhere, BlueprintReadWrite, Category = Response, meta = (EditCondition = "!bBypass"))
 	uint8 bAnalogMode : 1;
 
+	// Whether or not to bypass effect
+	UPROPERTY(EditAnywhere, BlueprintReadWrite, Category = General, meta = (DisplayName = "Bypass", DisplayAfter = "DynamicsProcessorType"))
+	uint8 bBypass : 1;
+
 	// Audition the key modulation signal, bypassing enveloping and processing the input signal.
-	UPROPERTY(EditAnywhere, BlueprintReadWrite, Category = Sidechain, meta = (DisplayName = "Key Audition"))
+	UPROPERTY(EditAnywhere, BlueprintReadWrite, Category = Sidechain, meta = (DisplayName = "Key Audition", EditCondition = "!bBypass"))
 	uint8 bKeyAudition : 1;
 
 	// Gain to apply to key signal if external input is supplied
-	UPROPERTY(EditAnywhere, BlueprintReadWrite, Category = Sidechain, meta = (DisplayName = "External Input Gain (dB)", EditCondition = "ExternalSubmix != nullptr", UIMin = "-60.0", UIMax = "30.0"))
-	float KeyGainDb;
+	UPROPERTY(EditAnywhere, BlueprintReadWrite, Category = Sidechain, meta = (DisplayName = "External Input Gain (dB)", EditCondition = "!bBypass && ExternalSubmix != nullptr", UIMin = "-60.0", UIMax = "30.0"))
+	float KeyGainDb = 0.0f;
 
 	// The output gain of the dynamics processor
-	UPROPERTY(EditAnywhere, BlueprintReadWrite, Category = Output, meta = (DisplayName = "Output Gain (dB)", UIMin = "-60.0", UIMax = "30.0"))
-	float OutputGainDb;
+	UPROPERTY(EditAnywhere, BlueprintReadWrite, Category = Output, meta = (DisplayName = "Output Gain (dB)", UIMin = "-60.0", UIMax = "30.0", EditCondition = "!bBypass"))
+	float OutputGainDb = 0.0f;
 
 	// High Shelf filter settings for key signal (external signal if supplied or input signal if not)
-	UPROPERTY(EditAnywhere, BlueprintReadWrite, Category = Sidechain, meta = (DisplayName = "Key Highshelf"))
+	UPROPERTY(EditAnywhere, BlueprintReadWrite, Category = Sidechain, meta = (DisplayName = "Key Highshelf", EditCondition = "!bBypass"))
 	FSubmixEffectDynamicProcessorFilterSettings KeyHighshelf;
 
 	// Low Shelf filter settings for key signal (external signal if supplied or input signal if not)
-	UPROPERTY(EditAnywhere, BlueprintReadWrite, Category = Sidechain, meta = (DisplayName = "Key Lowshelf"))
+	UPROPERTY(EditAnywhere, BlueprintReadWrite, Category = Sidechain, meta = (DisplayName = "Key Lowshelf", EditCondition = "!bBypass"))
 	FSubmixEffectDynamicProcessorFilterSettings KeyLowshelf;
 
 	FSubmixEffectDynamicsProcessorSettings()
-		: DynamicsProcessorType(ESubmixEffectDynamicsProcessorType::Compressor)
-		, PeakMode(ESubmixEffectDynamicsPeakMode::RootMeanSquared)
-		, LinkMode(ESubmixEffectDynamicsChannelLinkMode::Average)
-		, InputGainDb(0.0f)
-		, ThresholdDb(-6.0f)
-		, Ratio(1.5f)
-		, KneeBandwidthDb(10.0f)
-		, LookAheadMsec(3.0f)
-		, AttackTimeMsec(10.0f)
-		, ReleaseTimeMsec(100.0f)
-		, ExternalSubmix(nullptr)
-		, bChannelLinked_DEPRECATED(true)
+		: bChannelLinked_DEPRECATED(true)
 		, bAnalogMode(true)
+		, bBypass(false)
 		, bKeyAudition(false)
-		, KeyGainDb(0.0f)
-		, OutputGainDb(0.0f)
 	{
 		KeyLowshelf.Cutoff = 20000.0f;
 	}
 };
 
 
-class AUDIOMIXER_API FSubmixEffectDynamicsProcessor :
-	public FSoundEffectSubmix,
-	public ISubmixBufferListener
+class AUDIOMIXER_API FSubmixEffectDynamicsProcessor : public FSoundEffectSubmix
 {
 public:
 	FSubmixEffectDynamicsProcessor();
@@ -196,33 +210,94 @@ public:
 	// Called when an audio effect preset is changed
 	virtual void OnPresetChanged() override;
 
-	void SetExternalSubmix(USoundSubmix* InSoundSubmix);
 
 protected:
-	// ISubmixBufferListener interface
-	virtual void OnNewSubmixBuffer(
-		const USoundSubmix* InOwningSubmix,
-		float*				InAudioData,
-		int32				InNumSamples,
-		int32				InNumChannels,
-		const int32			InSampleRate,
-		double				InAudioClock) override;
+	Audio::FMixerDevice* GetMixerDevice();
+
+	void ResetKey();
+	void UpdateKeyFromSettings(const FSubmixEffectDynamicsProcessorSettings& InSettings);
+	bool UpdateKeySourcePatch();
 
 	void OnNewDeviceCreated(Audio::FDeviceId InDeviceId);
 	
-	TWeakObjectPtr<USoundSubmix> ExternalSubmix;
 	Audio::AlignedFloatBuffer AudioExternal;
 
-	FThreadSafeBool bUseExternalSubmix;
-
 	TArray<float> AudioKeyFrame;
 	TArray<float> AudioInputFrame;
 	TArray<float> AudioOutputFrame;
 
-	Audio::FDeviceId DeviceId;
-	Audio::FDynamicsProcessor DynamicsProcessor;
+	Audio::FDeviceId DeviceId = INDEX_NONE;
 
+	bool bBypass = false;
+
+private:
+	struct FKeySource
+	{
+		ESubmixEffectDynamicsKeySource Type = ESubmixEffectDynamicsKeySource::Default;
+		int32 NumChannels = 0;
+		uint32 ObjectId = INDEX_NONE;
+
+		mutable FCriticalSection MutateSourceCritSection;
+
+	public:
+		Audio::FPatchOutputStrongPtr Patch;
+
+		bool IsValid() const
+		{
+			const FScopeLock ScopeLock(&MutateSourceCritSection);
+			return Patch.IsValid();
+		}
+
+		void Reset()
+		{
+			Patch.Reset();
+
+			const FScopeLock ScopeLock(&MutateSourceCritSection);
+			{
+				NumChannels = 0;
+				ObjectId = INDEX_NONE;
+				Type = ESubmixEffectDynamicsKeySource::Default;
+			}
+		}
+
+		uint32 GetObjectId() const
+		{
+			const FScopeLock ScopeLock(&MutateSourceCritSection);
+			return ObjectId;
+		}
+
+		int32 GetNumChannels() const
+		{
+			const FScopeLock ScopeLock(&MutateSourceCritSection);
+			return NumChannels;
+		}
+
+		void SetNumChannels(const int32 InNumChannels)
+		{
+			const FScopeLock ScopeLock(&MutateSourceCritSection);
+			NumChannels = InNumChannels;
+		}
+
+		void Update(ESubmixEffectDynamicsKeySource InType, uint32 InObjectId, int32 InNumChannels = 0)
+		{
+			const FScopeLock ScopeLock(&MutateSourceCritSection);
+
+			if (Type != InType || ObjectId != InObjectId || NumChannels != InNumChannels)
+			{
+				Type = InType;
+				ObjectId = InObjectId;
+				NumChannels = InNumChannels;
+
+				Patch.Reset();
+			}
+		}
+	};
+
+	FKeySource KeySource;
+	Audio::FDynamicsProcessor DynamicsProcessor;
 	FDelegateHandle DeviceCreatedHandle;
+
+	friend class USubmixEffectDynamicsProcessorPreset;
 };
 
 UCLASS(ClassGroup = AudioSourceEffect, meta = (BlueprintSpawnableComponent))
@@ -237,6 +312,16 @@ public:
 
 	virtual void Serialize(FStructuredArchive::FRecord Record) override;
 
+	UFUNCTION(BlueprintCallable, Category = "Audio|Effects")
+	void ResetKey();
+
+	// Sets the source key input as the provided AudioBus' output.  If no object is provided, key is set
+	// to effect's input.
+	UFUNCTION(BlueprintCallable, Category = "Audio|Effects")
+	void SetAudioBus(UAudioBus* AudioBus);
+
+	// Sets the source key input as the provided Submix's output.  If no object is provided, key is set
+	// to effect's input.
 	UFUNCTION(BlueprintCallable, Category = "Audio|Effects")
 	void SetExternalSubmix(USoundSubmix* Submix);
 
@@ -245,4 +330,7 @@ public:
 
 	UPROPERTY(EditAnywhere, BlueprintReadWrite, Category = SubmixEffectPreset, meta = (ShowOnlyInnerProperties))
 	FSubmixEffectDynamicsProcessorSettings Settings;
+
+private:
+	void SetKey(ESubmixEffectDynamicsKeySource InKeySource, UObject* InObject, int32 InNumChannels = 0);
 };
diff --git a/Engine/Source/Runtime/AudioMixer/Private/AudioMixerDevice.cpp b/Engine/Source/Runtime/AudioMixer/Private/AudioMixerDevice.cpp
--- a/Engine/Source/Runtime/AudioMixer/Private/AudioMixerDevice.cpp
+++ b/Engine/Source/Runtime/AudioMixer/Private/AudioMixerDevice.cpp
@@ -1249,6 +1249,19 @@ namespace Audio
 		return false;
 	}
 
+	FMixerSubmixPtr FMixerDevice::GetMasterSubmixInstance(uint32 InSubmixId)
+	{
+		check(MasterSubmixes.Num() == EMasterSubmixType::Count);
+		for (int32 i = 0; i < EMasterSubmixType::Count; ++i)
+		{
+			if (InSubmixId == MasterSubmixes[i]->GetUniqueID())
+			{
+				return MasterSubmixInstances[i];
+			}
+		}
+		return nullptr;
+	}
+
 	FMixerSubmixPtr FMixerDevice::GetMasterSubmixInstance(const USoundSubmixBase* InSubmix)
 	{
 		check(MasterSubmixes.Num() == EMasterSubmixType::Count);
@@ -1458,6 +1471,123 @@ namespace Audio
 #endif
 	}
 
+	FMixerSubmixPtr FMixerDevice::FindSubmixInstanceByObjectId(uint32 InObjectId)
+	{
+		for (const USoundSubmix* MasterSubmix : MasterSubmixes)
+		{
+			if (!ensure(MasterSubmix))
+			{
+				continue;
+			}
+
+			if (MasterSubmix->GetUniqueID() == InObjectId)
+			{
+				return GetMasterSubmixInstance(MasterSubmix);
+			}
+		}
+
+		for (TPair<const USoundSubmixBase*, FMixerSubmixPtr>& Pair : Submixes)
+		{
+			const USoundSubmixBase* SubmixBase = Pair.Key;
+			if (SubmixBase)
+			{
+				if (SubmixBase->GetUniqueID() == InObjectId)
+				{
+					return Pair.Value;
+				}
+			}
+		}
+
+		return nullptr;
+	}
+
+	FMixerSubmixPtr FMixerDevice::GetSubmixInstance(uint32 InSubmixId)
+	{
+		LLM_SCOPE(ELLMTag::AudioMixer);
+
+		ensure(!IsInAudioThread() && !IsInGameThread());
+
+		FMixerSubmixPtr MixerSubmix = GetMasterSubmixInstance(InSubmixId);
+		if (MixerSubmix.IsValid())
+		{
+			return MixerSubmix;
+		}
+
+		for (TPair<const USoundSubmixBase*, FMixerSubmixPtr>& Pair : Submixes)
+		{
+			FMixerSubmixPtr SubmixPtr = Pair.Value;
+			if (ensure(SubmixPtr.IsValid()))
+			{
+				if (SubmixPtr->GetId() == InSubmixId)
+				{
+					return SubmixPtr;
+				}
+			}
+		}
+
+		return nullptr;
+	}
+
+	void FMixerDevice::InitDefaultAudioBuses()
+	{
+		if (!ensure(IsInGameThread()))
+		{
+			return;
+		}
+
+		if (const UAudioSettings* AudioSettings = GetDefault<UAudioSettings>())
+		{
+			TArray<TStrongObjectPtr<UAudioBus>> StaleBuses = DefaultAudioBuses;
+			DefaultAudioBuses.Reset();
+
+			for (const FDefaultAudioBusSettings& BusSettings : AudioSettings->DefaultAudioBuses)
+			{
+				if (UObject* BusObject = BusSettings.AudioBus.TryLoad())
+				{
+					if (UAudioBus* AudioBus = Cast<UAudioBus>(BusObject))
+					{
+						const int32 NumChannels = static_cast<int32>(AudioBus->AudioBusChannels) + 1;
+						StartAudioBus(AudioBus->GetUniqueID(), NumChannels, false /* bInIsAutomatic */);
+
+						TStrongObjectPtr<UAudioBus>AddedBus(AudioBus);
+						DefaultAudioBuses.AddUnique(AddedBus);
+						StaleBuses.Remove(AddedBus);
+					}
+				}
+			}
+
+			for (TStrongObjectPtr<UAudioBus>& Bus : StaleBuses)
+			{
+				if (Bus.IsValid())
+				{
+					StopAudioBus(Bus->GetUniqueID());
+				}
+			}
+		}
+		else
+		{
+			UE_LOG(LogAudioMixer, Error, TEXT("Failed to initialized Default Audio Buses. Audio Settings not found."));
+		}
+	}
+
+	void FMixerDevice::ShutdownDefaultAudioBuses()
+	{
+		if (!ensure(IsInGameThread()))
+		{
+			return;
+		}
+
+		for (TStrongObjectPtr<UAudioBus>& Bus : DefaultAudioBuses)
+		{
+			if (Bus.IsValid())
+			{
+				StopAudioBus(Bus->GetUniqueID());
+			}
+		}
+
+		DefaultAudioBuses.Reset();
+	}
+
 	FMixerSubmixWeakPtr FMixerDevice::GetSubmixInstance(const USoundSubmixBase* SoundSubmix)
 	{
 		LLM_SCOPE(ELLMTag::AudioMixer);
@@ -2147,11 +2277,11 @@ namespace Audio
 		return SourceManager->IsAudioBusActive(InAudioBusId);
 	}
 
-	FPatchOutputStrongPtr FMixerDevice::AddPatchForAudioBus(uint32 InAudioBusId, float PatchGain)
+	FPatchOutputStrongPtr FMixerDevice::AddPatchForAudioBus(uint32 InAudioBusId, float InPatchGain)
 	{
 		if (ensure(!IsInGameThread() && !IsInAudioThread()))
 		{
-			return SourceManager->AddPatchForAudioBus(InAudioBusId, PatchGain);
+			return SourceManager->AddPatchForAudioBus(InAudioBusId, InPatchGain);
 		}
 		else
 		{
@@ -2159,6 +2289,22 @@ namespace Audio
 		}
 	}
 
+	Audio::FPatchOutputStrongPtr FMixerDevice::AddPatchForSubmix(uint32 InObjectId, float InPatchGain)
+	{
+		if (!ensure(!IsInGameThread() && !IsInAudioThread()))
+		{
+			return nullptr;
+		}
+
+		FMixerSubmixPtr SubmixPtr = FindSubmixInstanceByObjectId(InObjectId);
+		if (SubmixPtr.IsValid())
+		{
+			return SubmixPtr->AddPatch(InPatchGain);
+		}
+
+		return nullptr;
+	}
+
 	int32 FMixerDevice::GetDeviceSampleRate() const
 	{
 		return SampleRate;
diff --git a/Engine/Source/Runtime/AudioMixer/Private/AudioMixerSourceManager.cpp b/Engine/Source/Runtime/AudioMixer/Private/AudioMixerSourceManager.cpp
--- a/Engine/Source/Runtime/AudioMixer/Private/AudioMixerSourceManager.cpp
+++ b/Engine/Source/Runtime/AudioMixer/Private/AudioMixerSourceManager.cpp
@@ -1068,10 +1068,10 @@ namespace Audio
 	FPatchOutputStrongPtr FMixerSourceManager::AddPatchForAudioBus(uint32 InAudioBusId, float PatchGain)
 	{
 		AUDIO_MIXER_CHECK_AUDIO_PLAT_THREAD(MixerDevice);
-		TSharedPtr<FMixerAudioBus>* AudioBusPtr = AudioBuses.Find(InAudioBusId);
-		if (AudioBusPtr)
+		TSharedPtr<FMixerAudioBus> AudioBusPtr = AudioBuses.FindRef(InAudioBusId);
+		if (AudioBusPtr.IsValid())
 		{
-			return (*AudioBusPtr)->AddNewPatch(NumOutputFrames * 2, PatchGain);
+			return AudioBusPtr->AddNewPatch(NumOutputFrames * AudioBusPtr->GetNumChannels(), PatchGain);
 		}
 		return nullptr;
 	}
@@ -2645,10 +2645,10 @@ namespace Audio
 
 	const float* FMixerSourceManager::GetPreviousAudioBusBuffer(const int32 AudioBusId) const
 	{
-		const TSharedPtr<FMixerAudioBus>* AudioBusPtr = AudioBuses.Find(AudioBusId);
-		if (AudioBusPtr)
+		const TSharedPtr<FMixerAudioBus> AudioBusPtr = AudioBuses.FindRef(AudioBusId);
+		if (AudioBusPtr.IsValid())
 		{
-			return (*AudioBusPtr)->GetPreviousBusBuffer();
+			return AudioBusPtr->GetPreviousBusBuffer();
 		}
 		return nullptr;
 	}
diff --git a/Engine/Source/Runtime/AudioMixer/Private/AudioMixerSubmix.cpp b/Engine/Source/Runtime/AudioMixer/Private/AudioMixerSubmix.cpp
--- a/Engine/Source/Runtime/AudioMixer/Private/AudioMixerSubmix.cpp
+++ b/Engine/Source/Runtime/AudioMixer/Private/AudioMixerSubmix.cpp
@@ -463,7 +463,6 @@ namespace Audio
 		return 0;
 	}
 
-
 	void FMixerSubmix::AddOrSetSourceVoice(FMixerSourceVoice* InSourceVoice, const float InSendLevel, EMixerSourceSubmixSendStage InSubmixSendStage)
 	{
 		AUDIO_MIXER_CHECK_AUDIO_PLAT_THREAD(MixerDevice);
@@ -475,6 +474,17 @@ namespace Audio
 		MixerSourceVoices.Add(InSourceVoice, NewVoiceData);
 	}
 
+	FPatchOutputStrongPtr FMixerSubmix::AddPatch(float InGain)
+	{
+		if (IsSoundfieldSubmix())
+		{
+			UE_LOG(LogAudioMixer, Warning, TEXT("Patch listening to SoundfieldSubmixes is not supported."));
+			return nullptr;
+		}
+
+		return PatchSplitter.AddNewPatch(NumSamples, InGain);
+	}
+
 	void FMixerSubmix::RemoveSourceVoice(FMixerSourceVoice* InSourceVoice)
 	{
 		AUDIO_MIXER_CHECK_AUDIO_PLAT_THREAD(MixerDevice);
@@ -792,7 +802,7 @@ namespace Audio
 				// If this is true, the Soundfield Factory explicitly requested that a seperate encoder stream was set up for every
 				// non-soundfield child submix.
 				if (Child.Encoder.IsValid())
-				{					
+				{
 					ChildSubmixSharedPtr->ProcessAudio(ScratchBuffer);
 
 					// Encode the resulting audio and mix it in.
@@ -1262,6 +1272,8 @@ namespace Audio
 				check(BufferListener);
 				BufferListener->OnNewSubmixBuffer(SoundSubmix, OutAudioBuffer.GetData(), OutAudioBuffer.Num(), NumChannels, SampleRate, AudioClock);
 			}
+
+			PatchSplitter.PushAudio(OutAudioBuffer.GetData(), OutAudioBuffer.Num());
 		}
 	}
 
@@ -1321,7 +1333,6 @@ namespace Audio
 		return bProcessedAnEffect;
 	}
 
-
 	void FMixerSubmix::ProcessAudio(ISoundfieldAudioPacket& OutputAudio)
 	{
 		check(IsSoundfieldSubmix());
diff --git a/Engine/Source/Runtime/AudioMixer/Private/Effects/AudioMixerSubmixEffectDynamicsProcessor.cpp b/Engine/Source/Runtime/AudioMixer/Private/Effects/AudioMixerSubmixEffectDynamicsProcessor.cpp
--- a/Engine/Source/Runtime/AudioMixer/Private/Effects/AudioMixerSubmixEffectDynamicsProcessor.cpp
+++ b/Engine/Source/Runtime/AudioMixer/Private/Effects/AudioMixerSubmixEffectDynamicsProcessor.cpp
@@ -4,6 +4,7 @@
 
 #include "AudioDeviceManager.h"
 #include "AudioMixerDevice.h"
+#include "AudioMixerSubmix.h"
 #include "ProfilingDebugging/CsvProfiler.h"
 
 // Link to "Audio" profiling category
@@ -11,29 +12,21 @@ CSV_DECLARE_CATEGORY_MODULE_EXTERN(AUDIOMIXERCORE_API, Audio);
 
 DEFINE_STAT(STAT_AudioMixerSubmixDynamics);
 
+static int32 bBypassSubmixDynamicsProcessor = 0;
+FAutoConsoleVariableRef CVarBypassDynamicsProcessor(
+	TEXT("au.Submix.Effects.DynamicsProcessor.Bypass"),
+	bBypassSubmixDynamicsProcessor,
+	TEXT("If non-zero, bypasses all submix dynamics processors currently active.\n"),
+	ECVF_Default);
+
 FSubmixEffectDynamicsProcessor::FSubmixEffectDynamicsProcessor()
-	: DeviceId(INDEX_NONE)
 {
 }
 
 FSubmixEffectDynamicsProcessor::~FSubmixEffectDynamicsProcessor()
 {
 	FAudioDeviceManagerDelegates::OnAudioDeviceCreated.Remove(DeviceCreatedHandle);
-
-	if (FAudioDeviceManager* DeviceManager = FAudioDeviceManager::Get())
-	{
-		FAudioDeviceHandle DeviceHandle = DeviceManager->GetAudioDevice(DeviceId);
-		if (DeviceHandle.IsValid())
-		{
-			Audio::FMixerDevice* MixerDevice = static_cast<Audio::FMixerDevice*>(DeviceHandle.GetAudioDevice());
-			check(MixerDevice);
-
-			if (ExternalSubmix.IsValid())
-			{
-				MixerDevice->UnregisterSubmixBufferListener(this, ExternalSubmix.Get());
-			}
-		}
-	}
+	ResetKey();
 }
 
 Audio::FDeviceId FSubmixEffectDynamicsProcessor::GetDeviceId() const
@@ -60,26 +53,46 @@ void FSubmixEffectDynamicsProcessor::Init(const FSoundEffectSubmixInitData& Init
 
 	if (USubmixEffectDynamicsProcessorPreset* ProcPreset = Cast<USubmixEffectDynamicsProcessorPreset>(Preset.Get()))
 	{
-		SetExternalSubmix(ProcPreset->Settings.ExternalSubmix);
+		switch (ProcPreset->Settings.KeySource)
+		{
+			case ESubmixEffectDynamicsKeySource::AudioBus:
+			{
+				if (UAudioBus* AudioBus = ProcPreset->Settings.ExternalAudioBus)
+				{
+					KeySource.Update(ESubmixEffectDynamicsKeySource::AudioBus, AudioBus->GetUniqueID(), static_cast<int32>(AudioBus->AudioBusChannels) + 1);
+				}
+			}
+			break;
+
+			case ESubmixEffectDynamicsKeySource::Submix:
+			{
+				if (USoundSubmix* Submix = ProcPreset->Settings.ExternalSubmix)
+				{
+					KeySource.Update(ESubmixEffectDynamicsKeySource::Submix, Submix->GetUniqueID());
+				}
+			}
+			break;
+
+			default:
+			{
+				// KeySource is this effect's submix/input, so do nothing
+			}
+			break;
+		}
 	}
 }
 
-void FSubmixEffectDynamicsProcessor::OnNewSubmixBuffer(
-	const USoundSubmix* InOwningSubmix,
-	float*				InAudioData,
-	int32				InNumSamples,
-	int32				InNumChannels,
-	const int32			InSampleRate,
-	double				InAudioClock)
+void FSubmixEffectDynamicsProcessor::ResetKey()
 {
-	AudioExternal.Reset();
-	AudioExternal.Append(InAudioData, InNumSamples);
+	KeySource.Reset();
 }
 
 void FSubmixEffectDynamicsProcessor::OnPresetChanged()
 {
 	GET_EFFECT_SETTINGS(SubmixEffectDynamicsProcessor);
 
+	bBypass = Settings.bBypass;
+
 	switch (Settings.DynamicsProcessorType)
 	{
 	default:
@@ -137,6 +150,68 @@ void FSubmixEffectDynamicsProcessor::OnPresetChanged()
 
 	static_assert(static_cast<int32>(ESubmixEffectDynamicsChannelLinkMode::Count) == static_cast<int32>(Audio::EDynamicsProcessorChannelLinkMode::Count), "Enumerations must match");
 	DynamicsProcessor.SetChannelLinkMode(static_cast<Audio::EDynamicsProcessorChannelLinkMode>(Settings.LinkMode));
+
+	UpdateKeyFromSettings(Settings);
+}
+
+Audio::FMixerDevice* FSubmixEffectDynamicsProcessor::GetMixerDevice()
+{
+	if (FAudioDeviceManager* DeviceManager = FAudioDeviceManager::Get())
+	{
+		return static_cast<Audio::FMixerDevice*>(DeviceManager->GetAudioDeviceRaw(DeviceId));
+	}
+
+	return nullptr;
+}
+
+bool FSubmixEffectDynamicsProcessor::UpdateKeySourcePatch()
+{
+	if (KeySource.Patch.IsValid())
+	{
+		return true;
+	}
+
+	if (Audio::FMixerDevice* MixerDevice = GetMixerDevice())
+	{
+		switch (KeySource.Type)
+		{
+			case ESubmixEffectDynamicsKeySource::AudioBus:
+			{
+				KeySource.Patch = MixerDevice->AddPatchForAudioBus(KeySource.ObjectId, 1.0f /* PatchGain */);
+				if (KeySource.Patch.IsValid())
+				{
+					DynamicsProcessor.SetKeyNumChannels(KeySource.GetNumChannels());
+					return true;
+				}
+			}
+			break;
+
+			case ESubmixEffectDynamicsKeySource::Submix:
+			{
+				KeySource.Patch = MixerDevice->AddPatchForSubmix(KeySource.ObjectId, 1.0f /* PatchGain */);
+				if (KeySource.Patch.IsValid())
+				{
+					Audio::FMixerSubmixPtr SubmixPtr = MixerDevice->FindSubmixInstanceByObjectId(KeySource.ObjectId);
+					if (SubmixPtr.IsValid())
+					{
+						const int32 SubmixNumChannels = SubmixPtr->GetNumOutputChannels();
+						KeySource.SetNumChannels(SubmixNumChannels);
+						DynamicsProcessor.SetKeyNumChannels(SubmixNumChannels);
+						return true;
+					}
+				}
+			}
+			break;
+
+			case ESubmixEffectDynamicsKeySource::Default:
+			default:
+			{
+			}
+			break;
+		}
+	}
+
+	return false;
 }
 
 void FSubmixEffectDynamicsProcessor::OnProcessAudio(const FSoundEffectSubmixInputData& InData, FSoundEffectSubmixOutputData& OutData)
@@ -144,93 +219,145 @@ void FSubmixEffectDynamicsProcessor::OnProcessAudio(const FSoundEffectSubmixInpu
 	CSV_SCOPED_TIMING_STAT(Audio, SubmixDynamics);
 	SCOPE_CYCLE_COUNTER(STAT_AudioMixerSubmixDynamics);
 
+	ensure(InData.NumChannels == OutData.NumChannels);
+
 	const Audio::AlignedFloatBuffer& InBuffer = *InData.AudioBuffer;
 	Audio::AlignedFloatBuffer& OutBuffer = *OutData.AudioBuffer;
 
+	const bool bBypassDueToInvalidChannelCount = !ensure(InData.NumChannels <= AudioInputFrame.Num());
+	if (bBypassDueToInvalidChannelCount || bBypassSubmixDynamicsProcessor || bBypass)
+	{
+		FMemory::Memcpy(OutBuffer.GetData(), InBuffer.GetData(), sizeof(float) * InBuffer.Num());
+		return;
+	}
+
+	int32 NumKeyChannels = DynamicsProcessor.GetKeyNumChannels();
+	AudioExternal.Reset();
+	if (KeySource.ObjectId != INDEX_NONE)
+	{
+		if (UpdateKeySourcePatch())
+		{
+			// Refresh num channels in case it changed when updating patch
+			NumKeyChannels = DynamicsProcessor.GetKeyNumChannels();
+
+			const int32 NumSamples = InData.NumFrames * NumKeyChannels;
+
+			AudioExternal.AddZeroed(NumSamples);
+			KeySource.Patch->PopAudio(AudioExternal.GetData(), NumSamples, true /* bUseLatestAudio */);
+		}
+	}
+
+	const int32 NumChannels = DynamicsProcessor.GetNumChannels();
+	if (InData.NumChannels != NumChannels)
+	{
+		DynamicsProcessor.SetNumChannels(InData.NumChannels);
+	}
+
 	for (int32 Frame = 0; Frame < InData.NumFrames; ++Frame)
 	{
 		// Copy the data to the frame input
-		const int32 SampleIndexOfFrame = Frame * InData.NumChannels;
+		const int32 SampleIndexOfInputFrame = Frame * InData.NumChannels;
 		for (int32 Channel = 0; Channel < InData.NumChannels; ++Channel)
 		{
-			const int32 SampleIndex = SampleIndexOfFrame + Channel;
+			const int32 SampleIndex = SampleIndexOfInputFrame + Channel;
 			AudioInputFrame[Channel] = InBuffer[SampleIndex];
-			if (bUseExternalSubmix)
-			{
-				AudioKeyFrame[Channel] = AudioExternal.Num() > SampleIndex ? AudioExternal[SampleIndex] : 0;
-			}
-			else
+		}
+
+		// Copy buffer data to key if key is external source
+		if (AudioExternal.Num() > 0)
+		{
+			// Copy the data to the frame input
+			const int32 SampleIndexOfKeyFrame = Frame * NumKeyChannels;
+			for (int32 Channel = 0; Channel < NumKeyChannels; ++Channel)
 			{
-				DynamicsProcessor.SetKeyGain(0.0f);
-				AudioKeyFrame[Channel] = InBuffer[SampleIndex];
+				const int32 SampleIndex = SampleIndexOfKeyFrame + Channel;
+				AudioKeyFrame[Channel] = AudioExternal[SampleIndex];
 			}
-		}
 
-		// Process
-		DynamicsProcessor.ProcessAudio(AudioInputFrame.GetData(), InData.NumChannels, AudioOutputFrame.GetData(), AudioKeyFrame.GetData());
+			DynamicsProcessor.ProcessAudio(AudioInputFrame.GetData(), InData.NumChannels, AudioOutputFrame.GetData(), AudioKeyFrame.GetData());
+		}
+		else
+		{
+			DynamicsProcessor.ProcessAudio(AudioInputFrame.GetData(), InData.NumChannels, AudioOutputFrame.GetData());
+		}
 
 		// Copy the data to the frame output
 		for (int32 Channel = 0; Channel < InData.NumChannels; ++Channel)
 		{
-			const int32 SampleIndex = SampleIndexOfFrame + Channel;
+			const int32 SampleIndex = SampleIndexOfInputFrame + Channel;
 			OutBuffer[SampleIndex] = AudioOutputFrame[Channel];
 		}
 	}
+
+	AudioExternal.Reset();
 }
 
-void FSubmixEffectDynamicsProcessor::SetExternalSubmix(USoundSubmix* InSoundSubmix)
+void FSubmixEffectDynamicsProcessor::UpdateKeyFromSettings(const FSubmixEffectDynamicsProcessorSettings& InSettings)
 {
-	FAudioDeviceManager* DeviceManager = FAudioDeviceManager::Get();
-	if (!DeviceManager)
+	uint32 ObjectId = INDEX_NONE;
+	int32 SourceNumChannels = 0;
+	switch (InSettings.KeySource)
 	{
-		return;
-	}
-
-	FAudioDeviceHandle DeviceHandle = DeviceManager->GetAudioDevice(DeviceId);
-	if (DeviceHandle.IsValid())
-	{
-		Audio::FMixerDevice* MixerDevice = static_cast<Audio::FMixerDevice*>(DeviceHandle.GetAudioDevice());
-		check(MixerDevice);
-
-		if (ExternalSubmix.Get() != InSoundSubmix)
+		case ESubmixEffectDynamicsKeySource::AudioBus:
 		{
-			if (ExternalSubmix.IsValid())
+			if (InSettings.ExternalAudioBus)
 			{
-				MixerDevice->UnregisterSubmixBufferListener(this, ExternalSubmix.Get());
+				ObjectId = InSettings.ExternalAudioBus->GetUniqueID();
+				SourceNumChannels = static_cast<int32>(InSettings.ExternalAudioBus->AudioBusChannels) + 1;
 			}
+		}
+		break;
 
-			ExternalSubmix = InSoundSubmix;
-
-			bUseExternalSubmix = ExternalSubmix.IsValid();
-			if (bUseExternalSubmix)
+		case ESubmixEffectDynamicsKeySource::Submix:
+		{
+			if (InSettings.ExternalSubmix)
 			{
-				MixerDevice->RegisterSubmixBufferListener(this, ExternalSubmix.Get());
+				ObjectId = InSettings.ExternalSubmix->GetUniqueID();
 			}
 		}
+		break;
 
+		default:
+		{
+		}
+		break;
 	}
-	else
-	{
-		DeviceCreatedHandle = FAudioDeviceManagerDelegates::OnAudioDeviceCreated.AddRaw(this, &FSubmixEffectDynamicsProcessor::OnNewDeviceCreated);
-		bUseExternalSubmix = ExternalSubmix.IsValid();
-	}
+
+	KeySource.Update(InSettings.KeySource, ObjectId, SourceNumChannels);
 }
 
 void FSubmixEffectDynamicsProcessor::OnNewDeviceCreated(Audio::FDeviceId InDeviceId)
 {
 	if (InDeviceId == DeviceId)
 	{
-		FAudioDeviceManagerDelegates::OnAudioDeviceCreated.Remove(DeviceCreatedHandle);
-
 		GET_EFFECT_SETTINGS(SubmixEffectDynamicsProcessor);
-		SetExternalSubmix(Settings.ExternalSubmix);
+		UpdateKeyFromSettings(Settings);
 
+		FAudioDeviceManagerDelegates::OnAudioDeviceCreated.Remove(DeviceCreatedHandle);
 	}
 }
 
 void USubmixEffectDynamicsProcessorPreset::OnInit()
 {
-	SetExternalSubmix(Settings.ExternalSubmix);
+	switch (Settings.KeySource)
+	{
+		case ESubmixEffectDynamicsKeySource::AudioBus:
+		{
+			SetAudioBus(Settings.ExternalAudioBus);
+		}
+		break;
+
+		case ESubmixEffectDynamicsKeySource::Submix:
+		{
+			SetExternalSubmix(Settings.ExternalSubmix);
+		}
+		break;
+
+		default:
+		{
+		}
+		break;
+	}
 }
 
 void USubmixEffectDynamicsProcessorPreset::Serialize(FStructuredArchive::FRecord Record)
@@ -248,16 +375,42 @@ void USubmixEffectDynamicsProcessorPreset::Serialize(FStructuredArchive::FRecord
 	Super::Serialize(Record);
 }
 
-void USubmixEffectDynamicsProcessorPreset::SetExternalSubmix(USoundSubmix* InSubmix)
+void USubmixEffectDynamicsProcessorPreset::ResetKey()
 {
-	IterateEffects<FSubmixEffectDynamicsProcessor>([this, InSubmix](FSubmixEffectDynamicsProcessor& Instance)
+	EffectCommand<FSubmixEffectDynamicsProcessor>([](FSubmixEffectDynamicsProcessor& Instance)
 	{
-		Instance.SetExternalSubmix(InSubmix);
+		Instance.ResetKey();
 	});
 }
 
+void USubmixEffectDynamicsProcessorPreset::SetAudioBus(UAudioBus* InAudioBus)
+{
+	const int32 BusChannels = static_cast<int32>(InAudioBus->AudioBusChannels) + 1;
+	SetKey(ESubmixEffectDynamicsKeySource::AudioBus, InAudioBus, BusChannels);
+}
+
+void USubmixEffectDynamicsProcessorPreset::SetExternalSubmix(USoundSubmix* InSubmix)
+{
+	SetKey(ESubmixEffectDynamicsKeySource::Submix, InSubmix);
+}
+
+void USubmixEffectDynamicsProcessorPreset::SetKey(ESubmixEffectDynamicsKeySource InKeySource, UObject* InObject, int32 InNumChannels)
+{
+	if (InObject)
+	{
+		EffectCommand<FSubmixEffectDynamicsProcessor>([this, ObjectId = InObject->GetUniqueID(), InKeySource, InNumChannels](FSubmixEffectDynamicsProcessor& Instance)
+		{
+			Instance.KeySource.Update(InKeySource, ObjectId, InNumChannels);
+		});
+	}
+}
+
 void USubmixEffectDynamicsProcessorPreset::SetSettings(const FSubmixEffectDynamicsProcessorSettings& InSettings)
 {
 	UpdateSettings(InSettings);
-	SetExternalSubmix(InSettings.ExternalSubmix);
+
+	IterateEffects<FSubmixEffectDynamicsProcessor>([&](FSubmixEffectDynamicsProcessor& Instance)
+	{
+		Instance.UpdateKeyFromSettings(InSettings);
+	});
 }
diff --git a/Engine/Source/Runtime/AudioMixer/Public/AudioMixerDevice.h b/Engine/Source/Runtime/AudioMixer/Public/AudioMixerDevice.h
--- a/Engine/Source/Runtime/AudioMixer/Public/AudioMixerDevice.h
+++ b/Engine/Source/Runtime/AudioMixer/Public/AudioMixerDevice.h
@@ -139,6 +139,8 @@ namespace Audio
 		virtual void FlushExtended(UWorld* WorldToFlush, bool bClearActivatedReverb);
 		virtual void FlushAudioRenderingCommands(bool bPumpSynchronously = false) override;
 
+		virtual FPatchOutputStrongPtr AddPatchForSubmix(uint32 InObjectId, float InPatchGain) override;
+
 		// Audio Device Properties
 		virtual bool IsNonRealtime() const override;
 
@@ -149,6 +151,9 @@ namespace Audio
 		virtual void OnAudioStreamShutdown() override;
 		//~ End IAudioMixer
 
+		FMixerSubmixPtr FindSubmixInstanceByObjectId(uint32 InObjectId);
+		FMixerSubmixPtr GetSubmixInstance(uint32 InSubmixId);
+
 		FMixerSubmixWeakPtr GetSubmixInstance(const USoundSubmixBase* SoundSubmix);
 
 		// If SoundSubmix is a soundfield submix, this will return the factory used to encode 
@@ -240,7 +245,7 @@ namespace Audio
 		void StartAudioBus(uint32 InAudioBusId, int32 InNumChannels, bool bInIsAutomatic);
 		void StopAudioBus(uint32 InAudioBusId);
 		bool IsAudioBusActive(uint32 InAudioBusId);
-		FPatchOutputStrongPtr AddPatchForAudioBus(uint32 InAudioBusId, float PatchGain);
+		FPatchOutputStrongPtr AddPatchForAudioBus(uint32 InAudioBusId, float InPatchGain);
 
 		// Clock Manager for quantized event handling on Audio Render Thread
 		FQuartzClockManager QuantizedEventClockManager;
@@ -285,6 +290,7 @@ namespace Audio
 	private:
 
 		bool IsMasterSubmixType(const USoundSubmixBase* InSubmix) const;
+		FMixerSubmixPtr GetMasterSubmixInstance(uint32 InSubmixId);
 		FMixerSubmixPtr GetMasterSubmixInstance(const USoundSubmixBase* InSubmix);
 		
 		// Pumps the audio render thread command queue
diff --git a/Engine/Source/Runtime/AudioMixer/Public/AudioMixerSubmix.h b/Engine/Source/Runtime/AudioMixer/Public/AudioMixerSubmix.h
--- a/Engine/Source/Runtime/AudioMixer/Public/AudioMixerSubmix.h
+++ b/Engine/Source/Runtime/AudioMixer/Public/AudioMixerSubmix.h
@@ -8,6 +8,7 @@
 #include "ISoundfieldEndpoint.h"
 #include "Sound/SoundSubmix.h"
 #include "DSP/EnvelopeFollower.h"
+#include "DSP/MultithreadedPatching.h"
 #include "DSP/SpectrumAnalyzer.h"
 #include "Templates/SharedPointer.h"
 #include "AudioDynamicParameter.h"
@@ -151,6 +152,8 @@ namespace Audio
 		// Add (if not already added) or sets the amount of the source voice's send amount
 		void AddOrSetSourceVoice(FMixerSourceVoice* InSourceVoice, const float SendLevel, EMixerSourceSubmixSendStage InSubmixSendStage);
 
+		FPatchOutputStrongPtr AddPatch(float InGain);
+
 		/** Removes the given source voice from the submix. */
 		void RemoveSourceVoice(FMixerSourceVoice* InSourceVoice);
 
@@ -585,6 +588,8 @@ namespace Audio
 		// Handle back to the owning USoundSubmix. Used when the device is shutdown to prematurely end a recording.
 		const USoundSubmixBase* OwningSubmixObject;
 
+		Audio::FPatchSplitter PatchSplitter;
+
 		friend class FMixerDevice;
 	};
 }
diff --git a/Engine/Source/Runtime/Engine/Public/AudioDevice.h b/Engine/Source/Runtime/Engine/Public/AudioDevice.h
--- a/Engine/Source/Runtime/Engine/Public/AudioDevice.h
+++ b/Engine/Source/Runtime/Engine/Public/AudioDevice.h
@@ -5,6 +5,7 @@
 #include "AudioDeviceManager.h"
 #include "Components/AudioComponent.h"
 #include "CoreMinimal.h"
+#include "DSP/MultithreadedPatching.h"
 #include "DSP/SpectrumAnalyzer.h"
 #include "Engine/Engine.h"
 #include "EngineGlobals.h"
@@ -861,6 +862,12 @@ public:
 		UE_LOG(LogAudio, Error, TEXT("Submix buffer listener only works with the audio mixer. Please run with audio mixer enabled."));
 	}
 
+	virtual Audio::FPatchOutputStrongPtr AddPatchForSubmix(uint32 InObjectId, float InPatchGain)
+	{
+		UE_LOG(LogAudio, Error, TEXT("Submix patching only works with the audio mixer. Please run with audio mixer enabled."));
+		return nullptr;
+	}
+
 	virtual void InitSoundEffectPresets() {}
 
 	/**
diff --git a/Engine/Source/Runtime/SignalProcessing/Private/DynamicsProcessor.cpp b/Engine/Source/Runtime/SignalProcessing/Private/DynamicsProcessor.cpp
--- a/Engine/Source/Runtime/SignalProcessing/Private/DynamicsProcessor.cpp
+++ b/Engine/Source/Runtime/SignalProcessing/Private/DynamicsProcessor.cpp
@@ -6,6 +6,7 @@ namespace Audio
 {
 	FDynamicsProcessor::FDynamicsProcessor()
 		: ProcessingMode(EDynamicsProcessingMode::Compressor)
+		, EnvelopeFollowerPeakMode(EPeakMode::RootMeanSquared)
 		, LookaheedDelayMsec(10.0f)
 		, AttackTimeMsec(20.0f)
 		, ReleaseTimeMsec(1000.0f)
@@ -15,7 +16,6 @@ namespace Audio
 		, InputGain(1.0f)
 		, OutputGain(1.0f)
 		, KeyGain(1.0f)
-		, NumChannels(0)
 		, LinkMode(EDynamicsProcessorChannelLinkMode::Disabled)
 		, bIsAnalogMode(true)
 		, bKeyAuditionEnabled(false)
@@ -30,16 +30,16 @@ namespace Audio
 	{
 	}
 
-	void FDynamicsProcessor::Init(const float SampleRate, const int32 InNumChannels)
+	void FDynamicsProcessor::Init(const float InSampleRate, const int32 InNumChannels)
 	{
-		NumChannels = InNumChannels;
+		SampleRate = InSampleRate;
 
+		SetNumChannels(InNumChannels);
+		SetKeyNumChannels(InNumChannels);
+	
 		LookaheadDelay.Reset();
 		LookaheadDelay.AddDefaulted(InNumChannels);
 
-		EnvFollower.Reset();
-		EnvFollower.AddDefaulted(InNumChannels);
-
 		for (int32 Channel = 0; Channel < InNumChannels; ++Channel)
 		{
 			LookaheadDelay[Channel].Init(SampleRate, 0.1f);
@@ -58,6 +58,16 @@ namespace Audio
 		Gain.AddZeroed(InNumChannels);
 	}
 
+	int32 FDynamicsProcessor::GetNumChannels() const
+	{
+		return Gain.Num();
+	}
+
+	int32 FDynamicsProcessor::GetKeyNumChannels() const
+	{
+		return EnvFollower.Num();
+	}
+	
 	void FDynamicsProcessor::SetLookaheadMsec(const float InLookAheadMsec)
 	{
 		LookaheedDelayMsec = InLookAheadMsec;
@@ -146,6 +156,36 @@ namespace Audio
 		InputLowshelfFilter.SetGainDB(InGainDb);
 	}
 
+	void FDynamicsProcessor::SetKeyNumChannels(const int32 InNumChannels)
+	{
+		if (InNumChannels != EnvFollower.Num())
+		{
+			EnvFollower.Reset();
+			EnvFollower.AddDefaulted(InNumChannels);
+
+			for (int32 Channel = 0; Channel < InNumChannels; ++Channel)
+			{
+				EnvFollower[Channel].Init(SampleRate, AttackTimeMsec, ReleaseTimeMsec, EnvelopeFollowerPeakMode, bIsAnalogMode);
+			}
+		}
+
+		if (InNumChannels != InputLowshelfFilter.GetNumChannels())
+		{
+			InputLowshelfFilter.Init(SampleRate, InNumChannels, EBiquadFilter::LowShelf);
+		}
+
+		if (InNumChannels != InputHighshelfFilter.GetNumChannels())
+		{
+			InputHighshelfFilter.Init(SampleRate, InNumChannels, EBiquadFilter::HighShelf);
+		}
+
+		if (InNumChannels != DetectorOuts.Num())
+		{
+			DetectorOuts.Reset();
+			DetectorOuts.AddZeroed(InNumChannels);
+		}
+	}
+
 	void FDynamicsProcessor::SetOutputGain(const float InOutputGainDb)
 	{
 		OutputGain = ConvertToLinear(InOutputGainDb);
@@ -165,8 +205,30 @@ namespace Audio
 		}
 	}
 
+	void FDynamicsProcessor::SetNumChannels(const int32 InNumChannels)
+	{
+		if (InNumChannels != Gain.Num())
+		{
+			Gain.Reset();
+			Gain.AddZeroed(InNumChannels);
+		}
+
+		if (InNumChannels != LookaheadDelay.Num())
+		{
+			LookaheadDelay.Reset();
+			LookaheadDelay.AddDefaulted(InNumChannels);
+
+			for (int32 Channel = 0; Channel < InNumChannels; ++Channel)
+			{
+				LookaheadDelay[Channel].Init(SampleRate, 0.1f);
+				LookaheadDelay[Channel].SetDelayMsec(LookaheedDelayMsec);
+			}
+		}
+	}
+
 	void FDynamicsProcessor::SetPeakMode(const EPeakMode::Type InEnvelopeFollowerModeType)
 	{
+		EnvelopeFollowerPeakMode = InEnvelopeFollowerModeType;
 		for (int32 Channel = 0; Channel < EnvFollower.Num(); ++Channel)
 		{
 			EnvFollower[Channel].SetMode(InEnvelopeFollowerModeType);
@@ -180,23 +242,66 @@ namespace Audio
 
 	void FDynamicsProcessor::ProcessAudioFrame(const float* InFrame, float* OutFrame, const float* InKeyFrame)
 	{
-		checkSlow(NumChannels == DetectorOuts.Num());
-		checkSlow(NumChannels == Gain.Num());
+		if (ProcessKeyFrame(InKeyFrame, OutFrame))
+		{
+			const int32 NumChannels = GetNumChannels();
+			for (int32 Channel = 0; Channel < NumChannels; ++Channel)
+			{
+				// Write and read into the look ahead delay line.
+				// We apply the compression output of the direct input to the output of this delay line
+				// This way sharp transients can be "caught" with the gain.
+				float LookaheadOutput = LookaheadDelay[Channel].ProcessAudioSample(InFrame[Channel]);
+
+				// Write into the output with the computed gain value
+				OutFrame[Channel] = Gain[Channel] * LookaheadOutput * OutputGain * InputGain;
+			}
+		}
+	}
+
+	void FDynamicsProcessor::ProcessAudio(const float* InBuffer, const int32 InNumSamples, float* OutBuffer, const float* InKeyBuffer)
+	{
+		const int32 NumChannels = GetNumChannels();
+		const int32 KeyNumChannels = GetKeyNumChannels();
 
+		if (InKeyBuffer)
+		{
+			int32 KeySampleIndex = 0;
+			for (int32 SampleIndex = 0; SampleIndex < InNumSamples; SampleIndex += NumChannels)
+			{
+				const float* KeyFrame = &InKeyBuffer[KeySampleIndex];
+				ProcessAudioFrame(&InBuffer[SampleIndex], &OutBuffer[SampleIndex], KeyFrame);
+				KeySampleIndex += KeyNumChannels;
+			}
+		}
+		else
+		{
+			for (int32 SampleIndex = 0; SampleIndex < InNumSamples; SampleIndex += NumChannels)
+			{
+				const float* KeyFrame = &InBuffer[SampleIndex];
+				ProcessAudioFrame(&InBuffer[SampleIndex], &OutBuffer[SampleIndex], KeyFrame);
+			}
+		}
+	}
+
+	bool FDynamicsProcessor::ProcessKeyFrame(const float* InKeyFrame, float* OutFrame)
+	{
 		// Get detector outputs
-		const float* DetectorIn = InKeyFrame;
-		if (NumChannels > 0)
+		const float* KeyIn = InKeyFrame;
+
+		const int32 KeyNumChannels = GetKeyNumChannels();
+		const int32 NumChannels = GetNumChannels();
+		if (KeyNumChannels > 0)
 		{
 			if (bKeyLowshelfEnabled)
 			{
-				InputLowshelfFilter.ProcessAudioFrame(DetectorIn, DetectorOuts.GetData());
-				DetectorIn = DetectorOuts.GetData();
+				InputLowshelfFilter.ProcessAudioFrame(KeyIn, DetectorOuts.GetData());
+				KeyIn = DetectorOuts.GetData();
 			}
 
 			if (bKeyHighshelfEnabled)
 			{
-				InputHighshelfFilter.ProcessAudioFrame(DetectorIn, DetectorOuts.GetData());
-				DetectorIn = DetectorOuts.GetData();
+				InputHighshelfFilter.ProcessAudioFrame(KeyIn, DetectorOuts.GetData());
+				KeyIn = DetectorOuts.GetData();
 			}
 		}
 
@@ -205,82 +310,68 @@ namespace Audio
 		{
 			for (int32 Channel = 0; Channel < NumChannels; ++Channel)
 			{
-				OutFrame[Channel] = DetectorGain * DetectorIn[Channel];
+				const int32 KeyIndex = Channel % KeyNumChannels;
+				OutFrame[Channel] = DetectorGain * KeyIn[KeyIndex];
 			}
+
+			return false;
 		}
-		else
+
+		for (int32 Channel = 0; Channel < KeyNumChannels; ++Channel)
 		{
-			for (int32 Channel = 0; Channel < NumChannels; ++Channel)
-			{
-				DetectorOuts[Channel] = EnvFollower[Channel].ProcessAudio(DetectorGain * DetectorIn[Channel]);
-			}
+			DetectorOuts[Channel] = EnvFollower[Channel].ProcessAudio(DetectorGain * KeyIn[Channel]);
+		}
 
-			switch (LinkMode)
+		switch (LinkMode)
+		{
+			case EDynamicsProcessorChannelLinkMode::Average:
 			{
-				case EDynamicsProcessorChannelLinkMode::Average:
+				float KeyOutLinked = 0.0f;
+				for (int32 Channel = 0; Channel < KeyNumChannels; ++Channel)
 				{
-					float DetectorOutLinked = 0.0f;
-					for (int32 Channel = 0; Channel < NumChannels; ++Channel)
-					{
-						DetectorOutLinked += DetectorOuts[Channel];
-					}
-					DetectorOutLinked /= static_cast<float>(NumChannels);
-					const float DetectorOutLinkedDb = ConvertToDecibels(DetectorOutLinked);
-					const float ComputedGain = ComputeGain(DetectorOutLinkedDb);
-					for (int32 Channel = 0; Channel < NumChannels; ++Channel)
-					{
-						Gain[Channel] = ComputedGain;
-					}
+					KeyOutLinked += DetectorOuts[Channel];
 				}
-				break;
-
-				case EDynamicsProcessorChannelLinkMode::Peak:
+				KeyOutLinked /= static_cast<float>(KeyNumChannels);
+				const float DetectorOutLinkedDb = ConvertToDecibels(KeyOutLinked);
+				const float ComputedGain = ComputeGain(DetectorOutLinkedDb);
+				for (int32 Channel = 0; Channel < NumChannels; ++Channel)
 				{
-					float DetectorOutLinked = FMath::Max<float>(DetectorOuts);
-					const float DetectorOutLinkedDb = ConvertToDecibels(DetectorOutLinked);
-					const float ComputedGain = ComputeGain(DetectorOutLinkedDb);
-					for (int32 Channel = 0; Channel < NumChannels; ++Channel)
-					{
-						Gain[Channel] = ComputedGain;
-					}
+					Gain[Channel] = ComputedGain;
 				}
-				break;
+			}
+			break;
 
-				case EDynamicsProcessorChannelLinkMode::Disabled:
-				default:
+			case EDynamicsProcessorChannelLinkMode::Peak:
+			{
+				float KeyOutLinked = FMath::Max<float>(DetectorOuts);
+				const float KeyOutLinkedDb = ConvertToDecibels(KeyOutLinked);
+				const float ComputedGain = ComputeGain(KeyOutLinkedDb);
+				for (int32 Channel = 0; Channel < NumChannels; ++Channel)
 				{
-					// compute gain individually per channel
-					for (int32 Channel = 0; Channel < NumChannels; ++Channel)
-					{
-						// convert to decibels
-						const float DetectorOutLinkedDb = ConvertToDecibels(DetectorOuts[Channel]);
-						const float ComputedGain = ComputeGain(DetectorOutLinkedDb);
-						Gain[Channel] = ComputedGain;
-					}
+					Gain[Channel] = ComputedGain;
 				}
-				break;
 			}
+			break;
 
-			for (int32 Channel = 0; Channel < NumChannels; ++Channel)
+			case EDynamicsProcessorChannelLinkMode::Disabled:
+			default:
 			{
-				// Write and read into the look ahead delay line.
-				// We apply the compression output of the direct input to the output of this delay line
-				// This way sharp transients can be "caught" with the gain.
-				float LookaheadOutput = LookaheadDelay[Channel].ProcessAudioSample(InFrame[Channel]);
+				// Compute gain individually per channel and wrap if
+				// channel count is greater than key channel count.
+				for (int32 Channel = 0; Channel < NumChannels; ++Channel)
+				{
+					const int32 KeyIndex = Channel % KeyNumChannels;
+					float ChannelGain = DetectorOuts[KeyIndex];
 
-				// Write into the output with the computed gain value
-				OutFrame[Channel] = Gain[Channel] * LookaheadOutput * OutputGain * InputGain;
+					const float KeyOutDb = ConvertToDecibels(ChannelGain);
+					const float ComputedGain = ComputeGain(KeyOutDb);
+					Gain[KeyIndex] = ComputedGain;
+				}
 			}
+			break;
 		}
-	}
 
-	void FDynamicsProcessor::ProcessAudio(const float* InBuffer, const int32 InNumSamples, float* OutBuffer, const float* InKeyBuffer)
-	{
-		for (int32 SampleIndex = 0; SampleIndex < InNumSamples; SampleIndex += NumChannels)
-		{
-			const float* KeyFrame = InKeyBuffer ? &InKeyBuffer[SampleIndex] : &InBuffer[SampleIndex];
-			ProcessAudioFrame(&InBuffer[SampleIndex], &OutBuffer[SampleIndex], KeyFrame);
-		}
+		return true;
 	}
 
 	float FDynamicsProcessor::ComputeGain(const float InEnvFollowerDb)
diff --git a/Engine/Source/Runtime/SignalProcessing/Private/Filter.cpp b/Engine/Source/Runtime/SignalProcessing/Private/Filter.cpp
--- a/Engine/Source/Runtime/SignalProcessing/Private/Filter.cpp
+++ b/Engine/Source/Runtime/SignalProcessing/Private/Filter.cpp
@@ -35,7 +35,7 @@ namespace Audio
 		SampleRate = InSampleRate;
 		NumChannels = InNumChannels;
 		FilterType = InFilterType;
-        Frequency = ClampCutoffFrequency(InCutoffFrequency);
+		Frequency = ClampCutoffFrequency(InCutoffFrequency);
 		Bandwidth = InBandwidth;
 		GainDB = InGainDB;
 
@@ -49,6 +49,11 @@ namespace Audio
 		CalculateBiquadCoefficients();
 	}
 
+	int32 FBiquadFilter::GetNumChannels() const
+	{
+		return NumChannels;
+	}
+
 	void FBiquadFilter::ProcessAudioFrame(const float* InFrame, float* OutFrame)
 	{
 		for (int32 Channel = 0; Channel < NumChannels; ++Channel)
diff --git a/Engine/Source/Runtime/SignalProcessing/Public/DSP/DynamicsProcessor.h b/Engine/Source/Runtime/SignalProcessing/Public/DSP/DynamicsProcessor.h
--- a/Engine/Source/Runtime/SignalProcessing/Public/DSP/DynamicsProcessor.h
+++ b/Engine/Source/Runtime/SignalProcessing/Public/DSP/DynamicsProcessor.h
@@ -38,7 +38,10 @@ namespace Audio
 		FDynamicsProcessor();
 		~FDynamicsProcessor();
 
-		void Init(const float SampleRate, const int32 NumChannels = 2);
+		void Init(const float InSampleRate, const int32 InNumChannels = 2);
+
+		int32 GetNumChannels() const;
+		int32 GetKeyNumChannels() const;
 
 		void SetLookaheadMsec(const float InLookAheadMsec);
 		void SetAttackTime(const float InAttackTimeMsec);
@@ -55,6 +58,8 @@ namespace Audio
 		void SetKeyLowshelfCutoffFrequency(const float InCutoffFreq);
 		void SetKeyLowshelfEnabled(const bool bInEnabled);
 		void SetKeyLowshelfGain(const float InGainDb);
+		void SetKeyNumChannels(const int32 InNumChannels);
+		void SetNumChannels(const int32 InNumChannels);
 		void SetOutputGain(const float InOutputGainDb);
 		void SetChannelLinkMode(const EDynamicsProcessorChannelLinkMode InLinkMode);
 		void SetAnalogMode(const bool bInIsAnalogMode);
@@ -69,6 +74,10 @@ namespace Audio
 
 		float ComputeGain(const float InEnvFollowerDb);
 
+		// Process key frame, returning true if should continue processing
+		// (Returns false in audition mode and writes straight to output).
+		bool ProcessKeyFrame(const float* InKeyFrame, float* OutFrame);
+
 		// (Optional) Low-pass filter for input signal
 		FBiquadFilter InputLowshelfFilter;
 
@@ -77,6 +86,9 @@ namespace Audio
 
 		EDynamicsProcessingMode::Type ProcessingMode;
 
+		// Peak mode of envelope followers
+		EPeakMode::Type EnvelopeFollowerPeakMode;
+
 		// Lookahead delay lines
 		TArray<FDelay> LookaheadDelay;
 
@@ -119,8 +131,8 @@ namespace Audio
 		// Gain of key detector signal in dB
 		float KeyGain;
 
-		// Number of channels to use for the dynamics processor
-		int32 NumChannels;
+		// Sample rate of both key and input (must match)
+		float SampleRate;
 
 		// Whether or not input channels are linked, and if so, how to calculate gain
 		EDynamicsProcessorChannelLinkMode LinkMode;
diff --git a/Engine/Source/Runtime/SignalProcessing/Public/DSP/Filter.h b/Engine/Source/Runtime/SignalProcessing/Public/DSP/Filter.h
--- a/Engine/Source/Runtime/SignalProcessing/Public/DSP/Filter.h
+++ b/Engine/Source/Runtime/SignalProcessing/Public/DSP/Filter.h
@@ -40,6 +40,9 @@ namespace Audio
 		// Initialize the filter
 		void Init(const float InSampleRate, const int32 InNumChannels, const EBiquadFilter::Type InType, const float InCutoffFrequency = 20000.0f, const float InBandwidth = 2.0f, const float InGain = 0.0f);
 
+		// Returns number of channels initialized with
+		int32 GetNumChannels() const;
+
 		// Resets the filter state
 		void Reset();
 
@@ -72,9 +75,9 @@ namespace Audio
 		// Function computes biquad coefficients based on current filter settings
 		void CalculateBiquadCoefficients();
 
-        // Function used to clamp the cutoff frequency.
-        float ClampCutoffFrequency(float InCutoffFrequency);
-        
+		// Function used to clamp the cutoff frequency.
+		float ClampCutoffFrequency(float InCutoffFrequency);
+
 		// What kind of filter to use when computing coefficients
 		EBiquadFilter::Type FilterType;
 
