Fix for broken reverb in surround

diff --git a/Engine/Source/Runtime/AudioMixer/Private/AudioMixerSubmix.cpp b/Engine/Source/Runtime/AudioMixer/Private/AudioMixerSubmix.cpp
--- a/Engine/Source/Runtime/AudioMixer/Private/AudioMixerSubmix.cpp
+++ b/Engine/Source/Runtime/AudioMixer/Private/AudioMixerSubmix.cpp
@@ -1314,8 +1314,6 @@ namespace Audio
 				InputData.AudioBuffer = &DownmixedBuffer;
 				SubmixEffect->ProcessAudio(InputData, OutputData);
 
-				// Mix back up to the input channel count when we copy the effect output to the input
-				DownmixBuffer(ChannelCountOverride, DownmixedBuffer, NumChannels, InAudioBuffer);
 			}
 			else
 			{
@@ -1323,11 +1321,11 @@ namespace Audio
 				InputData.NumChannels = NumChannels;
 				InputData.AudioBuffer = &InAudioBuffer;
 				SubmixEffect->ProcessAudio(InputData, OutputData);
-
-				// Copy the output to the input
-				FMemory::Memcpy((void*)InAudioBuffer.GetData(), (void*)OutputData.AudioBuffer->GetData(), sizeof(float) * NumSamples);
 			}
 
+			// Copy the output to the input
+			FMemory::Memcpy((void*)InAudioBuffer.GetData(), (void*)OutputData.AudioBuffer->GetData(), sizeof(float) * NumSamples);
+
 			// Mix in the dry signal directly
 			const float DryLevel = SubmixEffect->GetDryLevel();
 			if (DryLevel > 0.0f)
