Adding clamping on asset import for cases where import formats allow audio to go above unity gain.

diff --git a/Engine/Source/Runtime/AudioMixer/Private/SoundFileIO/SoundFileIO.cpp b/Engine/Source/Runtime/AudioMixer/Private/SoundFileIO/SoundFileIO.cpp
--- a/Engine/Source/Runtime/AudioMixer/Private/SoundFileIO/SoundFileIO.cpp
+++ b/Engine/Source/Runtime/AudioMixer/Private/SoundFileIO/SoundFileIO.cpp
@@ -106,6 +106,22 @@ namespace Audio
 		Error = InputSoundDataReader->ReadSamples(ProcessBuffer.GetData(), ProcessBufferSamples, InputSamplesRead);
 		check(Error == ESoundFileError::Type::NONE);
 
+		// Normalize and clamp the input decoded audio
+		if (bPerformPeakNormalization)
+		{
+			for (int32 Sample = 0; Sample < InputSamplesRead; ++Sample)
+			{
+				ProcessBuffer[Sample] = FMath::Clamp(ProcessBuffer[Sample] / MaxValue, -1.0f, 1.0f);
+			}
+		}
+		else
+		{
+			for (int32 Sample = 0; Sample < InputSamplesRead; ++Sample)
+			{
+				ProcessBuffer[Sample] = FMath::Clamp(ProcessBuffer[Sample], -1.0f, 1.0f);
+			}
+		}
+
 		SoundFileCount SamplesWritten = 0;
 
 		while (InputSamplesRead == ProcessBuffer.Num())
@@ -118,12 +134,19 @@ namespace Audio
 			Error = InputSoundDataReader->ReadSamples(ProcessBuffer.GetData(), ProcessBufferSamples, InputSamplesRead);
 			check(Error == ESoundFileError::Type::NONE);
 
-			// ... normalize the samples if we're told to
+			// Normalize and clamp the samples
 			if (bPerformPeakNormalization)
 			{
 				for (int32 Sample = 0; Sample < InputSamplesRead; ++Sample)
 				{
-					ProcessBuffer[Sample] /= MaxValue;
+					ProcessBuffer[Sample] = FMath::Clamp(ProcessBuffer[Sample] / MaxValue, -1.0f, 1.0f);
+				}
+			}
+			else
+			{
+				for (int32 Sample = 0; Sample < InputSamplesRead; ++Sample)
+				{
+					ProcessBuffer[Sample] = FMath::Clamp(ProcessBuffer[Sample], -1.0f, 1.0f);
 				}
 			}
 		}
