Adding new architecture for realtime analysis using audio buses.

- Implemented a realtime analyzer for loudness.
- Added some utilities to multithreaded audio patching
- Added some utilities to audio buses

diff --git a/Engine/Plugins/Runtime/AudioSynesthesia/Source/AudioSynesthesia/Classes/AudioSynesthesia.h b/Engine/Plugins/Runtime/AudioSynesthesia/Source/AudioSynesthesia/Classes/AudioSynesthesia.h
new file mode 100644
--- /dev/null
+++ b/Engine/Plugins/Runtime/AudioSynesthesia/Source/AudioSynesthesia/Classes/AudioSynesthesia.h
@@ -0,0 +1,22 @@
+// Copyright Epic Games, Inc. All Rights Reserved.
+#pragma once
+
+#include "AudioAnalyzer.h"
+#include "AudioSynesthesia.generated.h"
+
+/** UAudioSynesthesiaSettings
+ *
+ * Defines asset actions for derived UAudioSynthesiaSettings subclasses.
+ */
+UCLASS(Abstract, Blueprintable)
+class AUDIOSYNESTHESIA_API UAudioSynesthesiaSettings : public UAudioAnalyzerSettings
+{
+	GENERATED_BODY()
+
+public:
+
+#if WITH_EDITOR
+	const TArray<FText>& GetAssetActionSubmenus() const;
+	FColor GetTypeColor() const override;
+#endif
+};
\ No newline at end of file
diff --git a/Engine/Plugins/Runtime/AudioSynesthesia/Source/AudioSynesthesia/Classes/Loudness.h b/Engine/Plugins/Runtime/AudioSynesthesia/Source/AudioSynesthesia/Classes/Loudness.h
new file mode 100644
--- /dev/null
+++ b/Engine/Plugins/Runtime/AudioSynesthesia/Source/AudioSynesthesia/Classes/Loudness.h
@@ -0,0 +1,155 @@
+// Copyright Epic Games, Inc. All Rights Reserved.
+
+#pragma once
+
+#include "AudioSynesthesia.h"
+#include "Loudness.generated.h"
+
+/** ELoudnessCurveTypeEnum
+ *
+ * Enumeration of available equal loudness curves. Loudness curves based on IEC 61672-1:2013 Electroacoustics - Sound level meters - Part 1: Specifications.
+ */
+UENUM(BlueprintType)
+enum class ELoudnessCurveTypeEnum : uint8
+{
+	/** Loudness Curve A Weighting. Commonly used in environmental sound measurements. Best for low level sounds. */
+	A 		UMETA(DisplayName = "A"),
+	/** Loudness Curve B Weighting. Relative to "A", gives more precedence to frequencies below 1kHz. */
+	B		UMETA(DisplayName = "B"),
+	/** Loudness Curve C Weighting. Relative to "A" and "B", gives more precedence to frequencies below 1kHz. */
+	C		UMETA(DisplayName = "C"),
+	/** Loudness Curve D Weighting. Similar to "B" but with an emphasis on presence in the 2kHz-6KHz frequency range. */
+	D		UMETA(DisplayName = "D"),
+	/** No loudness curve weighting. */
+	None	UMETA(DisplayName = "None"),
+};
+
+
+/** ULoudnessSettings
+ *
+ * Settings for a ULoudness analyzer.
+ */
+UCLASS(Blueprintable)
+class AUDIOSYNESTHESIA_API ULoudnessSettings : public UAudioSynesthesiaSettings
+{
+	GENERATED_BODY()
+public:
+
+	ULoudnessSettings();
+
+	/** Number of seconds between loudness measurements */
+	UPROPERTY(EditAnywhere, BlueprintReadOnly, Category = AudioAnalyzer, meta = (ClampMin = "0.01", ClampMax = "0.25"))
+	float AnalysisPeriod;
+
+	/** Minimum analysis frequency for calculating loudness. */
+	UPROPERTY(EditAnywhere, BlueprintReadOnly, Category = AudioAnalyzer, meta = (ClampMin = "20.0", ClampMax = "20000"))
+	float MinimumFrequency;
+
+	/** Maximum analysis frequency for calculating loudness. */
+	UPROPERTY(EditAnywhere, BlueprintReadOnly, Category = AudioAnalyzer, meta = (ClampMin = "20.0", ClampMax = "20000"))
+	float MaximumFrequency;
+
+	/** Type of equal loudness curve used in calculations */
+	UPROPERTY(EditAnywhere, BlueprintReadOnly, Category = AudioAnalyzer)
+	ELoudnessCurveTypeEnum CurveType;
+
+	/** dB level to denote silence.  Used when calculating normalized loudness. */
+	UPROPERTY(EditAnywhere, BlueprintReadOnly, AdvancedDisplay, Category = AudioAnalyzer, meta = (ClampMin = "-100.0", ClampMax = "0"))
+	float NoiseFloorDb;
+
+	/** dB level to denote silence.  Used when calculating normalized loudness. */
+	UPROPERTY(EditAnywhere, BlueprintReadOnly, AdvancedDisplay, Category = AudioAnalyzer, meta = (ClampMin = "-100.0", ClampMax = "0"))
+	float ExpectedMaxLoudness;
+
+	/** Convert ULoudnessSettings to FLoudnessSettings */
+	TUniquePtr<Audio::IAnalyzerSettings> GetSettings(const int32 InSampleRate, const int32 InNumChannels) const;
+
+#if WITH_EDITOR
+	virtual FText GetAssetActionName() const override;
+
+	virtual UClass* GetSupportedClass() const override;
+#endif
+};
+
+/** The results of the loudness analyis. */
+USTRUCT(BlueprintType)
+struct AUDIOSYNESTHESIA_API FLoudnessResults
+{
+	GENERATED_USTRUCT_BODY()
+	
+	// The raw loudness value in dB
+	UPROPERTY(EditAnywhere, BlueprintReadOnly, Category = AudioAnalyzer)
+	float Loudness = 0.0f;
+
+	// The normalized loudness
+	UPROPERTY(EditAnywhere, BlueprintReadOnly, Category = AudioAnalyzer)
+	float NormalizedLoudness = 0.0f;
+
+	// The normalized loudness
+	UPROPERTY(EditAnywhere, BlueprintReadOnly, Category = AudioAnalyzer)
+	float PerceptualEnergy = 0.0f;
+
+	// The time in seconds since analysis began
+	UPROPERTY(EditAnywhere, BlueprintReadOnly, Category = AudioAnalyzer)
+	float TimeSeconds = 0.0f;
+};
+
+/** Delegate to receive all overall loudness results (time-stamped in an array) since last delegate call. */
+DECLARE_DYNAMIC_MULTICAST_DELEGATE_OneParam(FOnOverallLoudnessResults, const TArray<FLoudnessResults>&, OverallLoudnessResults);
+
+/** Delegate to receive only the most recent overall loudness result. */
+DECLARE_DYNAMIC_MULTICAST_DELEGATE_OneParam(FOnLatestOverallLoudnessResults, const FLoudnessResults&, LatestOverallLoudnessResults);
+
+/** Delegate to receive all loudness results per channel (time-stamped in an array) since last delegate call. */
+DECLARE_DYNAMIC_MULTICAST_DELEGATE_TwoParams(FOnPerChannelLoudnessResults, int32, ChannelIndex, const TArray<FLoudnessResults>&, LoudnessResults);
+
+/** Delegate to receive only the most recent overall loudness result per channel. */
+DECLARE_DYNAMIC_MULTICAST_DELEGATE_TwoParams(FOnLatestPerChannelLoudnessResults, int32, ChannelIndex, const FLoudnessResults&, LatestLoudnessResults);
+
+/** ULoudnessAnalyzer
+ *
+ * ULoudnessAnalyzer calculates the temporal evolution of perceptual loudness for a given
+ * audio bus in real-time. Loudness is available for individual channels or the overall audio bus. Normalized
+ * loudness values convert the range to 0.0 to 1.0 where 0.0 is the noise floor and 1.0 is the
+ * maximum loudness of the particular sound.
+ */
+UCLASS(Blueprintable)
+class AUDIOSYNESTHESIA_API ULoudnessAnalyzer : public UAudioAnalyzer
+{
+	GENERATED_BODY()
+
+public:
+
+	ULoudnessAnalyzer();
+
+	/** The settings for the audio analyzer.  */
+	UPROPERTY(EditAnywhere, BlueprintReadWrite, Category = AudioAnalyzer)
+	ULoudnessSettings* Settings;
+
+	/** Delegate to receive all overall loudness results since last delegate call. */
+	UPROPERTY(BlueprintAssignable)
+	FOnOverallLoudnessResults OnOverallLoudnessResults;
+
+	/** Delegate to receive all loudness results, per-channel, since last delegate call. */
+	UPROPERTY(BlueprintAssignable)
+	FOnPerChannelLoudnessResults OnPerChannelLoudnessResults;
+
+	/** Delegate to receive the latest overall loudness results. */
+	UPROPERTY(BlueprintAssignable)
+	FOnLatestOverallLoudnessResults OnLatestOverallLoudnessResults;
+
+	/** Delegate to receive the latest per-channel loudness results. */
+	UPROPERTY(BlueprintAssignable)
+	FOnLatestPerChannelLoudnessResults OnLatestPerChannelLoudnessResults;
+
+	/** Convert ULoudnessSettings to FLoudnessSettings */
+	TUniquePtr<Audio::IAnalyzerSettings> GetSettings(const int32 InSampleRate, const int32 InNumChannels) const override;
+
+	/** Broadcasts results to any delegates if hooked up. */
+	void BroadcastResults() override;
+
+protected:
+
+	/** Return the name of the IAudioAnalyzerFactory associated with this UAudioAnalyzer */
+	FName GetAnalyzerFactoryName() const override;
+};
\ No newline at end of file
diff --git a/Engine/Plugins/Runtime/AudioSynesthesia/Source/AudioSynesthesia/Classes/LoudnessNRT.h b/Engine/Plugins/Runtime/AudioSynesthesia/Source/AudioSynesthesia/Classes/LoudnessNRT.h
--- a/Engine/Plugins/Runtime/AudioSynesthesia/Source/AudioSynesthesia/Classes/LoudnessNRT.h
+++ b/Engine/Plugins/Runtime/AudioSynesthesia/Source/AudioSynesthesia/Classes/LoudnessNRT.h
@@ -18,7 +18,7 @@ enum class ELoudnessNRTCurveTypeEnum : uint8
 	B		UMETA(DisplayName="B"),
 	/** Loudness Curve C Weighting. Relative to "A" and "B", gives more precedence to frequencies below 1kHz. */
 	C		UMETA(DisplayName="C"),
-	/** Loudness Curve D Weighting. Simlar to "B" but with an emphasis on presence in the 2kHz-6KHz frequency range. */
+	/** Loudness Curve D Weighting. Similar to "B" but with an emphasis on presence in the 2kHz-6KHz frequency range. */
 	D		UMETA(DisplayName="D"),
 	/** No loudness curve weighting. */
 	None	UMETA(DisplayName="None"),
diff --git a/Engine/Plugins/Runtime/AudioSynesthesia/Source/AudioSynesthesia/Private/AudioSynesthesia.cpp b/Engine/Plugins/Runtime/AudioSynesthesia/Source/AudioSynesthesia/Private/AudioSynesthesia.cpp
new file mode 100644
--- /dev/null
+++ b/Engine/Plugins/Runtime/AudioSynesthesia/Source/AudioSynesthesia/Private/AudioSynesthesia.cpp
@@ -0,0 +1,21 @@
+// Copyright Epic Games, Inc. All Rights Reserved.
+
+#include "AudioSynesthesia.h"
+
+
+#if WITH_EDITOR
+const TArray<FText>& UAudioSynesthesiaSettings::GetAssetActionSubmenus() const
+{
+	static const TArray<FText> SubMenus
+	{
+		NSLOCTEXT("AssetTypeActions", "AssetTypeActions_AssetSoundAnalysisSubmenu", "Analysis")
+	};
+
+	return SubMenus;
+}
+
+FColor UAudioSynesthesiaSettings::GetTypeColor() const
+{
+	return FColor(150.0f, 200.0f, 200.0f);
+}
+#endif
\ No newline at end of file
diff --git a/Engine/Plugins/Runtime/AudioSynesthesia/Source/AudioSynesthesia/Private/Loudness.cpp b/Engine/Plugins/Runtime/AudioSynesthesia/Source/AudioSynesthesia/Private/Loudness.cpp
new file mode 100644
--- /dev/null
+++ b/Engine/Plugins/Runtime/AudioSynesthesia/Source/AudioSynesthesia/Private/Loudness.cpp
@@ -0,0 +1,174 @@
+// Copyright Epic Games, Inc. All Rights Reserved.
+
+#include "Loudness.h"
+#include "LoudnessFactory.h"
+#include "InterpolateSorted.h"
+#include "AudioSynesthesiaLog.h"
+
+ULoudnessSettings::ULoudnessSettings()
+	: AnalysisPeriod(0.01f)
+	, MinimumFrequency(20.f)
+	, MaximumFrequency(20000.f)
+	, CurveType(ELoudnessCurveTypeEnum::D)
+	, NoiseFloorDb(-60.f)
+	, ExpectedMaxLoudness(0.0f)
+{}
+
+TUniquePtr<Audio::IAnalyzerSettings> ULoudnessSettings::GetSettings(const int32 InSampleRate, const int32 InNumChannels) const
+{
+	TUniquePtr<Audio::FLoudnessSettings> Settings = MakeUnique<Audio::FLoudnessSettings>();
+
+	Settings->AnalysisPeriod = AnalysisPeriod;
+	Settings->MinAnalysisFrequency = MinimumFrequency;
+	Settings->MaxAnalysisFrequency = MaximumFrequency;
+
+	switch (CurveType)
+	{
+	case ELoudnessCurveTypeEnum::A:
+		Settings->LoudnessCurveType = Audio::ELoudnessCurveType::A;
+		break;
+
+	case ELoudnessCurveTypeEnum::B:
+		Settings->LoudnessCurveType = Audio::ELoudnessCurveType::B;
+		break;
+
+	case ELoudnessCurveTypeEnum::C:
+		Settings->LoudnessCurveType = Audio::ELoudnessCurveType::C;
+		break;
+
+	case ELoudnessCurveTypeEnum::D:
+		Settings->LoudnessCurveType = Audio::ELoudnessCurveType::D;
+		break;
+
+	case ELoudnessCurveTypeEnum::None:
+	default:
+		Settings->LoudnessCurveType = Audio::ELoudnessCurveType::None;
+		break;
+	}
+
+	return Settings;
+}
+
+#if WITH_EDITOR
+FText ULoudnessSettings::GetAssetActionName() const
+{
+	return NSLOCTEXT("AssetTypeActions", "AssetTypeActions_AssetSoundSynesthesiaLoudnessSettings", "Synesthesia Real-Time Settings (Loudness)");
+}
+
+UClass* ULoudnessSettings::GetSupportedClass() const
+{
+	return ULoudnessSettings::StaticClass();
+}
+#endif
+
+ULoudnessAnalyzer::ULoudnessAnalyzer()
+{
+	Settings = CreateDefaultSubobject<ULoudnessSettings>(TEXT("DefaultLoudnessSettings"));
+}
+
+TUniquePtr<Audio::IAnalyzerSettings> ULoudnessAnalyzer::GetSettings(const int32 InSampleRate, const int32 InNumChannels) const
+{
+	TUniquePtr<Audio::IAnalyzerSettings> AnalyzerSettings;
+
+	if (Settings)
+	{
+		AnalyzerSettings = Settings->GetSettings(InSampleRate, InNumChannels);
+	}
+
+	return AnalyzerSettings;
+}
+
+static TArray<FLoudnessResults> ConvertToBlueprintResults(ULoudnessSettings* Settings, const TArray<Audio::FLoudnessEntry>& InLoudnessArray)
+{
+	// Helper function for "Sort" (higher priority sorts last).
+	struct FCompareLoudnessResults
+	{
+		FORCEINLINE bool operator()(const FLoudnessResults& A, const FLoudnessResults& B) const
+		{
+			return A.TimeSeconds < B.TimeSeconds;
+		}
+	};
+
+	TArray<FLoudnessResults> ResultsArray;
+
+	float LoudnessRange = FMath::Max(Settings->ExpectedMaxLoudness - Settings->NoiseFloorDb, 0.001f);
+
+	for (const Audio::FLoudnessEntry& LoudnessEntry : InLoudnessArray)
+	{
+		FLoudnessResults NewResults;
+		NewResults.Loudness = LoudnessEntry.Loudness;
+		NewResults.NormalizedLoudness = FMath::Clamp((LoudnessEntry.Loudness - Settings->NoiseFloorDb) / LoudnessRange, 0.0f, 1.0f);
+		NewResults.PerceptualEnergy = LoudnessEntry.Energy;
+		NewResults.TimeSeconds = LoudnessEntry.Timestamp;
+
+		ResultsArray.Add(NewResults);
+	}
+
+	// Sort by priority (lowest priority first).
+	ResultsArray.Sort(FCompareLoudnessResults());
+	return MoveTemp(ResultsArray);
+}
+
+void ULoudnessAnalyzer::BroadcastResults()
+{
+	TUniquePtr<const Audio::FLoudnessResult> LoudnessResults = GetResults<Audio::FLoudnessResult>();
+
+	int32 NumChannels = LoudnessResults->GetNumChannels();
+
+	if (NumChannels > 0)
+	{
+		bool bIsOnOverallLoudnessResultsBound = OnOverallLoudnessResults.IsBound();
+		bool bIsOnLatestOverallLoudnessResultsBound = OnLatestOverallLoudnessResults.IsBound();
+		if (bIsOnOverallLoudnessResultsBound || bIsOnLatestOverallLoudnessResultsBound)
+		{
+			const TArray<Audio::FLoudnessEntry>& OverallLoudnessArray = LoudnessResults->GetLoudnessArray();
+			if (OverallLoudnessArray.Num() > 0)
+			{
+				TArray<FLoudnessResults> Results = ConvertToBlueprintResults(Settings, OverallLoudnessArray);
+				check(Results.Num() > 0);
+
+				if (bIsOnOverallLoudnessResultsBound)
+				{
+					OnOverallLoudnessResults.Broadcast(Results);
+				}
+				if (bIsOnLatestOverallLoudnessResultsBound)
+				{
+					FLoudnessResults& Latest = Results[Results.Num() - 1];
+					OnLatestOverallLoudnessResults.Broadcast(Latest);
+				}
+			}
+		}
+
+		bool bIsOnPerChannelLoudnessResultsBound = OnPerChannelLoudnessResults.IsBound();
+		bool bIsOnLatestPerChannelLoudnessResultsBound = OnLatestPerChannelLoudnessResults.IsBound();
+
+		if (bIsOnPerChannelLoudnessResultsBound || bIsOnLatestPerChannelLoudnessResultsBound)
+		{
+			for (int32 ChannelIndex = 0; ChannelIndex < NumChannels; ++ChannelIndex)
+			{
+				const TArray<Audio::FLoudnessEntry>& LoudnessArray = LoudnessResults->GetChannelLoudnessArray(ChannelIndex);
+				if (LoudnessArray.Num() > 0)
+				{
+					TArray<FLoudnessResults> Results = ConvertToBlueprintResults(Settings, LoudnessArray);
+					check(Results.Num() > 0);
+
+					if (bIsOnPerChannelLoudnessResultsBound)
+					{
+						OnPerChannelLoudnessResults.Broadcast(ChannelIndex, Results);
+					}
+
+					if (bIsOnLatestPerChannelLoudnessResultsBound)
+					{
+						OnLatestPerChannelLoudnessResults.Broadcast(ChannelIndex, Results[Results.Num() - 1]);
+					}
+				}
+			}
+		}
+	}
+}
+
+FName ULoudnessAnalyzer::GetAnalyzerFactoryName() const
+{
+	static const FName FactoryName(TEXT("LoudnessFactory"));
+	return FactoryName;
+}
\ No newline at end of file
diff --git a/Engine/Plugins/Runtime/AudioSynesthesia/Source/AudioSynesthesiaCore/Private/AudioSynesthesiaCoreModule.cpp b/Engine/Plugins/Runtime/AudioSynesthesia/Source/AudioSynesthesiaCore/Private/AudioSynesthesiaCoreModule.cpp
--- a/Engine/Plugins/Runtime/AudioSynesthesia/Source/AudioSynesthesiaCore/Private/AudioSynesthesiaCoreModule.cpp
+++ b/Engine/Plugins/Runtime/AudioSynesthesia/Source/AudioSynesthesiaCore/Private/AudioSynesthesiaCoreModule.cpp
@@ -9,6 +9,7 @@
 #include "LoudnessNRTFactory.h"
 #include "ConstantQNRTFactory.h"
 #include "OnsetNRTFactory.h"
+#include "LoudnessFactory.h"
 
 DEFINE_LOG_CATEGORY(LogAudioSynesthesiaCore);
 
@@ -20,23 +21,29 @@ namespace Audio
 			void StartupModule()
 			{
 				// Register factories on startup
-				IModularFeatures::Get().RegisterModularFeature(FLoudnessNRTFactory::GetModularFeatureName(), &LoudnessFactory);
-				IModularFeatures::Get().RegisterModularFeature(FConstantQNRTFactory::GetModularFeatureName(), &ConstantQFactory);
-				IModularFeatures::Get().RegisterModularFeature(FOnsetNRTFactory::GetModularFeatureName(), &OnsetFactory);
+				IModularFeatures::Get().RegisterModularFeature(FLoudnessNRTFactory::GetModularFeatureName(), &LoudnessNRTFactory);
+				IModularFeatures::Get().RegisterModularFeature(FConstantQNRTFactory::GetModularFeatureName(), &ConstantQNRTFactory);
+				IModularFeatures::Get().RegisterModularFeature(FOnsetNRTFactory::GetModularFeatureName(), &OnsetNRTFactory);
+
+				IModularFeatures::Get().RegisterModularFeature(FLoudnessFactory::GetModularFeatureName(), &LoudnessFactory);
 			}
 
 			void ShutdownModule()
 			{
 				// Unregister factories on shutdown
+				IModularFeatures::Get().UnregisterModularFeature(FLoudnessNRTFactory::GetModularFeatureName(), &LoudnessNRTFactory);
+				IModularFeatures::Get().UnregisterModularFeature(FConstantQNRTFactory::GetModularFeatureName(), &ConstantQNRTFactory);
+				IModularFeatures::Get().UnregisterModularFeature(FOnsetNRTFactory::GetModularFeatureName(), &OnsetNRTFactory);
+
 				IModularFeatures::Get().UnregisterModularFeature(FLoudnessNRTFactory::GetModularFeatureName(), &LoudnessFactory);
-				IModularFeatures::Get().UnregisterModularFeature(FConstantQNRTFactory::GetModularFeatureName(), &ConstantQFactory);
-				IModularFeatures::Get().UnregisterModularFeature(FOnsetNRTFactory::GetModularFeatureName(), &OnsetFactory);
 			}
 
 		private:
-			FLoudnessNRTFactory LoudnessFactory;
-			FConstantQNRTFactory ConstantQFactory;
-			FOnsetNRTFactory OnsetFactory;
+			FLoudnessNRTFactory LoudnessNRTFactory;
+			FConstantQNRTFactory ConstantQNRTFactory;
+			FOnsetNRTFactory OnsetNRTFactory;
+
+			FLoudnessFactory LoudnessFactory;
 	};
 
 }
diff --git a/Engine/Plugins/Runtime/AudioSynesthesia/Source/AudioSynesthesiaCore/Private/LoudnessFactory.cpp b/Engine/Plugins/Runtime/AudioSynesthesia/Source/AudioSynesthesiaCore/Private/LoudnessFactory.cpp
new file mode 100644
--- /dev/null
+++ b/Engine/Plugins/Runtime/AudioSynesthesia/Source/AudioSynesthesiaCore/Private/LoudnessFactory.cpp
@@ -0,0 +1,148 @@
+// Copyright Epic Games, Inc. All Rights Reserved.
+
+#include "LoudnessFactory.h"
+#include "AudioSynesthesiaCustomVersion.h"
+#include "DSP/SlidingWindow.h"
+
+namespace Audio
+{
+	// Use for overall loudness storage
+	const int32 FLoudnessResult::ChannelIndexOverall = INDEX_NONE;
+
+	FLoudnessResult::FLoudnessResult()
+	:	DurationInSeconds(0.f)
+	{
+	}
+
+	void FLoudnessResult::Add(const FLoudnessEntry& InDatum)
+	{
+		// Store loudness data in appropriate channel
+		TArray<FLoudnessEntry>& LoudnessArray = ChannelLoudnessArrays.FindOrAdd(InDatum.Channel);
+		LoudnessArray.Add(InDatum);
+
+		// Update intervals per channel
+		FFloatInterval& Interval = ChannelLoudnessIntervals.FindOrAdd(InDatum.Channel);
+		Interval.Include(InDatum.Loudness);
+	}
+
+	const TArray<FLoudnessEntry>& FLoudnessResult::GetChannelLoudnessArray(int32 ChannelIdx) const
+	{
+		return ChannelLoudnessArrays[ChannelIdx];
+	}
+
+	const TArray<FLoudnessEntry>& FLoudnessResult::GetLoudnessArray() const
+	{
+		return ChannelLoudnessArrays[ChannelIndexOverall];
+	}
+
+	int32 FLoudnessResult::GetNumChannels() const 
+	{
+		if (ChannelLoudnessArrays.Num() > 0)
+		{
+			// Channel loudness arrays have -1 reserved for overall loudness
+			return ChannelLoudnessArrays.Num() - 1;
+		}
+		return 0;
+	}
+
+	FLoudnessWorker::FLoudnessWorker(const FAnalyzerParameters& InParams, const FLoudnessSettings& InAnalyzerSettings)
+		: NumChannels(InParams.NumChannels)
+		, NumAnalyzedBuffers(0)
+		, NumHopFrames(0)
+		, SampleRate(InParams.SampleRate)
+	{
+		check(NumChannels > 0);
+		check(SampleRate > 0);
+
+		// Include NumChannels in calculating window size because windows are generated 
+		// with interleaved samples and deinterleaved later. 
+		int32 NumWindowSamples = InAnalyzerSettings.FFTSize * NumChannels;
+
+		NumHopFrames = FMath::CeilToInt(InAnalyzerSettings.AnalysisPeriod * InParams.SampleRate);
+		NumHopFrames = FMath::Max(1, NumHopFrames);
+
+		int32 NumHopSamples = NumHopFrames * NumChannels;
+
+		InternalBuffer = MakeUnique<TSlidingBuffer<float>>(NumWindowSamples, NumHopSamples);
+
+		Analyzer = MakeUnique<FMultichannelLoudnessAnalyzer>(InParams.SampleRate, InAnalyzerSettings);
+	}
+
+	void FLoudnessWorker::Analyze(TArrayView<const float> InAudio, IAnalyzerResult* OutResult) 
+	{
+		FLoudnessResult* LoudnessResult = static_cast<FLoudnessResult*>(OutResult);
+		check(nullptr != LoudnessResult);
+
+		TAutoSlidingWindow<float> SlidingWindow(*InternalBuffer, InAudio, InternalWindow);
+
+		// Loop through entire array of input samples. 
+		for (const TArray<float>& Window : SlidingWindow)
+		{
+			AnalyzeWindow(Window, *LoudnessResult);
+		}
+	}
+
+	void FLoudnessWorker::AnalyzeWindow(TArrayView<const float> InWindow, FLoudnessResult& OutResult)
+	{
+		TArray<float> ChannelPerceptualEnergy;
+		float OverallPerceptualEnergy = 0.0f;
+
+		// Calculate perceptual energy
+		OverallPerceptualEnergy = Analyzer->CalculatePerceptualEnergy(InWindow, NumChannels, ChannelPerceptualEnergy);
+
+		// Convert to perceptual energy to loudness
+		TArray<float> ChannelLoudness;
+		float OverallLoudness = 0.f;
+
+		OverallLoudness = FLoudnessAnalyzer::ConvertPerceptualEnergyToLoudness(OverallPerceptualEnergy);
+
+		for (float PerceptualEnergy : ChannelPerceptualEnergy)
+		{
+			ChannelLoudness.Add(FLoudnessAnalyzer::ConvertPerceptualEnergyToLoudness(PerceptualEnergy));
+		}
+
+		// Calculate timestamp
+		float Timestamp = ((NumAnalyzedBuffers * NumHopFrames) + (0.5f * Analyzer->GetSettings().FFTSize)) / SampleRate;
+
+		// Add overall FLoudnessDatum to result
+		const FLoudnessEntry OverallLoudnessEntry = { FLoudnessResult::ChannelIndexOverall, Timestamp, OverallPerceptualEnergy, OverallLoudness };
+		OutResult.Add(OverallLoudnessEntry);
+
+		// Add channel FLoudnessDatum to result
+		for (int32 ChannelIdx = 0; ChannelIdx < NumChannels; ChannelIdx++)
+		{
+			const FLoudnessEntry ChannelEntry = { ChannelIdx, Timestamp, ChannelPerceptualEnergy[ChannelIdx], ChannelLoudness[ChannelIdx] };
+
+			OutResult.Add(ChannelEntry);
+		}
+
+		NumAnalyzedBuffers++;
+	}
+
+	FName FLoudnessFactory::GetName() const 
+	{
+		static FName FactoryName(TEXT("LoudnessFactory"));
+		return FactoryName;
+	}
+
+	FString FLoudnessFactory::GetTitle() const
+	{
+		return TEXT("Real-Time Loudness Analyzer");
+	}
+
+	TUniquePtr<IAnalyzerResult> FLoudnessFactory::NewResult() const
+	{
+		TUniquePtr<FLoudnessResult> Result = MakeUnique<FLoudnessResult>();
+		return Result;
+	}
+
+	TUniquePtr<IAnalyzerWorker> FLoudnessFactory::NewWorker(const FAnalyzerParameters& InParams, const IAnalyzerSettings* InSettings) const
+	{
+		const FLoudnessSettings* LoudnessSettings = static_cast<const FLoudnessSettings*>(InSettings);
+
+		check(nullptr != LoudnessSettings);
+
+		return MakeUnique<FLoudnessWorker>(InParams, *LoudnessSettings);
+	}
+}
+
diff --git a/Engine/Plugins/Runtime/AudioSynesthesia/Source/AudioSynesthesiaCore/Private/LoudnessNRTFactory.cpp b/Engine/Plugins/Runtime/AudioSynesthesia/Source/AudioSynesthesiaCore/Private/LoudnessNRTFactory.cpp
--- a/Engine/Plugins/Runtime/AudioSynesthesia/Source/AudioSynesthesiaCore/Private/LoudnessNRTFactory.cpp
+++ b/Engine/Plugins/Runtime/AudioSynesthesia/Source/AudioSynesthesiaCore/Private/LoudnessNRTFactory.cpp
@@ -22,7 +22,7 @@ namespace Audio
 	}
 
 	// Use for overall loudness storage
-	const int32 FLoudnessNRTResult::ChannelIndexOverall = -1;
+	const int32 FLoudnessNRTResult::ChannelIndexOverall = INDEX_NONE;
 
 
 	FLoudnessNRTResult::FLoudnessNRTResult()
diff --git a/Engine/Plugins/Runtime/AudioSynesthesia/Source/AudioSynesthesiaCore/Public/LoudnessFactory.h b/Engine/Plugins/Runtime/AudioSynesthesia/Source/AudioSynesthesiaCore/Public/LoudnessFactory.h
new file mode 100644
--- /dev/null
+++ b/Engine/Plugins/Runtime/AudioSynesthesia/Source/AudioSynesthesiaCore/Public/LoudnessFactory.h
@@ -0,0 +1,110 @@
+// Copyright Epic Games, Inc. All Rights Reserved.
+
+#pragma once
+
+#include "LoudnessAnalyzer.h"
+#include "IAudioAnalyzerInterface.h"
+#include "DSP/SlidingWindow.h"
+
+namespace Audio
+{
+	/**
+	 * Contains settings for loudness analyzer.
+	 */
+	class AUDIOSYNESTHESIACORE_API FLoudnessSettings : public IAnalyzerSettings, public FLoudnessAnalyzerSettings
+	{
+		public:
+
+		/** Number of seconds between loudness measurements */
+		float AnalysisPeriod = 0.01f;
+	};
+
+	/**
+	 * Holds the loudness values per a time step.
+	 */
+	struct FLoudnessEntry
+	{
+		int32 Channel = 0;
+		float Timestamp = 0.f;
+		float Energy = 0.f;
+		float Loudness = 0.f;
+	};
+
+	/** 
+	 * FLoudnessResult contains the temporal evolution of loudness for the last buffer of audio analyzed.
+	 */
+	class AUDIOSYNESTHESIACORE_API FLoudnessResult : public IAnalyzerResult
+	{
+	public:
+		/**
+		 * Denotes the overall loudness as opposed individual channel indices.
+		 */
+		static const int32 ChannelIndexOverall;
+
+		FLoudnessResult();
+
+		/** Appends an FLoudnessDatum to the container. */
+		void Add(const FLoudnessEntry& InDatum);
+
+		/** Returns const reference to FLoudnessDatum array for individual channel. */
+		const TArray<FLoudnessEntry>& GetChannelLoudnessArray(int32 ChannelIdx) const;
+
+		/** Returns const reference to FLoudnessDatum array associated with overall loudness. */
+		const TArray<FLoudnessEntry>& GetLoudnessArray() const;
+
+		/** Returns the number of channels. */
+		int32 GetNumChannels() const;
+
+	private:
+		float DurationInSeconds;
+		TMap<int32, TArray<FLoudnessEntry> > ChannelLoudnessArrays;
+		TMap<int32, FFloatInterval> ChannelLoudnessIntervals;
+	};
+
+	/**
+	 * FLoudnessWorker performs loudness analysis on input sample buffers.
+	 */
+	class AUDIOSYNESTHESIACORE_API FLoudnessWorker : public IAnalyzerWorker
+	{
+	public:
+		/** Construct a worker */
+		FLoudnessWorker(const FAnalyzerParameters& InParams, const FLoudnessSettings& InAnalyzerSettings);
+
+		/** Analyzes input sample buffer and updates result. */
+		virtual void Analyze(TArrayView<const float> InAudio, IAnalyzerResult* OutResult) override;
+
+	private:
+
+		/** Analyze a single window. */
+		void AnalyzeWindow(TArrayView<const float> InWindow, FLoudnessResult& OutResult);
+
+		int32 NumChannels = 0;
+		int32 NumAnalyzedBuffers = 0;
+		int32 NumHopFrames = 0;
+		int32 SampleRate = 0;
+		TArray<float> InternalWindow;
+		TUniquePtr<TSlidingBuffer<float>> InternalBuffer;
+		TUniquePtr<FMultichannelLoudnessAnalyzer> Analyzer;
+	};
+
+	/**
+	 * Defines the Loudness analyzer and creates related classes.
+	 */
+	class AUDIOSYNESTHESIACORE_API FLoudnessFactory : public IAnalyzerFactory
+	{
+		public:
+
+		/** Name of specific analyzer type. */
+		virtual FName GetName() const override;
+
+		/** Human readable name of analyzer. */
+		virtual FString GetTitle() const override;
+
+		/** Creates a new FLoudnessNRTResult */
+		virtual TUniquePtr<IAnalyzerResult> NewResult() const override;
+
+		/**  Creates a new FLoudnessWorker. This expects IAnalyzerNRTSettings to be a valid pointer to a FLoudnessSettings object. */
+		virtual TUniquePtr<IAnalyzerWorker> NewWorker(const FAnalyzerParameters& InParams, const IAnalyzerSettings* InSettings) const override;
+	};
+}
+
diff --git a/Engine/Plugins/Runtime/AudioSynesthesia/Source/AudioSynesthesiaCore/Public/LoudnessNRTFactory.h b/Engine/Plugins/Runtime/AudioSynesthesia/Source/AudioSynesthesiaCore/Public/LoudnessNRTFactory.h
--- a/Engine/Plugins/Runtime/AudioSynesthesia/Source/AudioSynesthesiaCore/Public/LoudnessNRTFactory.h
+++ b/Engine/Plugins/Runtime/AudioSynesthesia/Source/AudioSynesthesiaCore/Public/LoudnessNRTFactory.h
@@ -8,8 +8,7 @@
 
 namespace Audio
 {
-	/** FLoudnessNRTSettings
-	 *
+	/**
 	 * Contains settings for loudness analyzer.
 	 */
 	class AUDIOSYNESTHESIACORE_API FLoudnessNRTSettings : public IAnalyzerNRTSettings, public FLoudnessAnalyzerSettings
@@ -24,9 +23,8 @@ namespace Audio
 	};
 
 
-	/** FLoudnessNRTResult
-	 *
-	 * Holds the loundess values per a time step.
+	/** 
+	 * Holds the loudness values per a time step. 
 	 */
 	struct FLoudnessDatum
 	{
@@ -35,87 +33,63 @@ namespace Audio
 		float Energy = 0.f;
 		float Loudness = 0.f;
 	};
+	
 	/** De/Serialize single loudness datum into archive. */
 	AUDIOSYNESTHESIACORE_API FArchive &operator <<(FArchive& Ar, FLoudnessDatum& Datum);
 
-	/** FLoudnessNRTResult
-	 *
-	 * FLoudnessNRTResult contains the temporal evolution of loudness.
-	 */
+	/** FLoudnessNRTResult contains the temporal evolution of loudness. */
 	class AUDIOSYNESTHESIACORE_API FLoudnessNRTResult : public IAnalyzerNRTResult
 	{
 	public:
-		/** ChannelIndexOverall
-		 *
-		 * Denotes the overall loudness channel index as opposed individual channel indices.
-		 */
+		/** Denotes the overall loudness channel index as opposed individual channel indices. */
 		static const int32 ChannelIndexOverall;
 
 		FLoudnessNRTResult();
 
-		/**
-		 * Defines how to serialize result.
-		 */
+		/** Defines how to serialize result. */
 		virtual void Serialize(FArchive& Archive) override;
 
-		/**
-		 * Appends an FLoudnessDatum to the container.
-		 */
+		/** Appends an FLoudnessDatum to the container. */
 		void Add(const FLoudnessDatum& InDatum);
 
 		/** Returns true if this object data for the given channel index */
 		bool ContainsChannel(int32 InChannelIndex) const;
 
-		/**
-		 * Returns const reference to FLoudnessDatum array for individual channel.
-		 */
+		/** Returns const reference to FLoudnessDatum array for individual channel. */
 		const TArray<FLoudnessDatum>& GetChannelLoudnessArray(int32 ChannelIdx) const;
 
-		/**
-		 * Returns const reference to FLoudnessDatum array associated with overall loudness.
-		 */
+		/** Returns const reference to FLoudnessDatum array associated with overall loudness. */
 		const TArray<FLoudnessDatum>& GetLoudnessArray() const;
 
-		/**
-		 * Returns range in dB of overall loudness result given the noise floor.
-		 */
+		/** Returns range in dB of overall loudness result given the noise floor. */
 		float GetLoudnessRange(float InNoiseFloor) const;
 
-
 		/** Returns range in dB of loudness result given the noise floor. */
 		float GetChannelLoudnessRange(int32 InChannelIdx, float InNoiseFloor) const;
 
-		/** Returns the channel indices availabe in result. */
+		/** Returns the channel indices available in result. */
 		void GetChannels(TArray<int32>& OutChannels) const;
 
-		/** Gets the duration of the anlayzed audio. */
+		/** Gets the duration of the analyzed audio. */
 		virtual float GetDurationInSeconds() const override;
 
-		/** Sets the duration of the anlayzed audio. */
+		/** Sets the duration of the analyzed audio. */
 		void SetDurationInSeconds(float InDuration);
 
-		/**
-		 * Returns true if FLoudnessDatum arrays are sorted in chronologically ascending order via their timestamp.
-		 */
+		/** Returns true if FLoudnessDatum arrays are sorted in chronologically ascending order via their timestamp. */
 		bool IsSortedChronologically() const;
 
-		/**
-		 * Sorts FLoudnessDatum arrays in chronologically ascnding order via their timestamp.
-		 */
+		/** Sorts FLoudnessDatum arrays in chronologically ascending order via their timestamp. */
 		void SortChronologically();
 
 	private:
 		float DurationInSeconds;
-
 		TMap<int32, TArray<FLoudnessDatum> > ChannelLoudnessArrays;
-
 		TMap<int32, FFloatInterval> ChannelLoudnessIntervals;
-
 		bool bIsSortedChronologically;
 	};
 
-	/** FLoudnessNRTWorker
-	 *
+	/** 
 	 * FLoudnessNRTWorker performs loudness analysis on input sample buffers.
 	 */
 	class AUDIOSYNESTHESIACORE_API FLoudnessNRTWorker : public IAnalyzerNRTWorker
@@ -143,18 +117,13 @@ namespace Audio
 		int32 NumAnalyzedBuffers;
 		int32 NumHopFrames;
 		int32 NumFrames;
-
 		float SampleRate;
-
 		TArray<float> InternalWindow;
-
 		TUniquePtr<TSlidingBuffer<float>> InternalBuffer;
-
 		TUniquePtr<FMultichannelLoudnessAnalyzer> Analyzer;
 	};
 
-	/** FLoudnessNRTFactory
-	 *
+	/**
 	 * Defines the LoudnessNRT analyzer and creates related classes.
 	 */
 	class AUDIOSYNESTHESIACORE_API FLoudnessNRTFactory : public IAnalyzerNRTFactory
@@ -170,10 +139,7 @@ namespace Audio
 		/** Creates a new FLoudnessNRTResult */
 		virtual TUniquePtr<IAnalyzerNRTResult> NewResult() const override;
 
-		/** 
-		 * Creates a new FLoudnessNRTWorker. This expects IAnalyzerNRTSettings to be a valid pointer to
-		 * a FLoudnessNRTSettings object.
-		 */
+		/** Creates a new FLoudnessNRTWorker. This expects IAnalyzerNRTSettings to be a valid pointer to a FLoudnessNRTSettings object. */
 		virtual TUniquePtr<IAnalyzerNRTWorker> NewWorker(const FAnalyzerNRTParameters& InParams, const IAnalyzerNRTSettings* InSettings) const override;
 	};
 }
diff --git a/Engine/Plugins/Runtime/AudioSynesthesia/Source/AudioSynesthesiaEditor/Classes/AudioSynesthesiaSettingsFactory.h b/Engine/Plugins/Runtime/AudioSynesthesia/Source/AudioSynesthesiaEditor/Classes/AudioSynesthesiaSettingsFactory.h
new file mode 100644
--- /dev/null
+++ b/Engine/Plugins/Runtime/AudioSynesthesia/Source/AudioSynesthesiaEditor/Classes/AudioSynesthesiaSettingsFactory.h
@@ -0,0 +1,28 @@
+// Copyright Epic Games, Inc. All Rights Reserved.
+
+#pragma once
+
+#include "AudioSynesthesia.h"
+#include "Factories/Factory.h"
+#include "AudioSynesthesiaSettingsFactory.generated.h"
+
+/** UAudioSynesthesiaSettingsFactory
+ *
+ * UAudioSynesthesiaSettingsFactory implements the factory for UAudioSynesthesiaSettings assets.
+ */
+UCLASS(MinimalAPI, hidecategories=Object)
+class UAudioSynesthesiaSettingsFactory : public UFactory
+{
+	GENERATED_UCLASS_BODY()
+
+	/** The type of audio analyzer settings that will be created */
+	UPROPERTY(EditAnywhere, Category = AudioSynesthesiaFactory)
+	TSubclassOf<UAudioSynesthesiaSettings> AudioSynesthesiaSettingsClass;
+
+	//~ Begin UFactory Interface
+	virtual bool ConfigureProperties() override;
+	virtual UObject* FactoryCreateNew(UClass* Class,UObject* InParent,FName Name,EObjectFlags Flags,UObject* Context, FFeedbackContext* Warn) override;
+	//~ Begin UFactory Interface	
+};
+
+
diff --git a/Engine/Plugins/Runtime/AudioSynesthesia/Source/AudioSynesthesiaEditor/Private/AssetTypeActions_AudioSynesthesiaSettings.cpp b/Engine/Plugins/Runtime/AudioSynesthesia/Source/AudioSynesthesiaEditor/Private/AssetTypeActions_AudioSynesthesiaSettings.cpp
new file mode 100644
--- /dev/null
+++ b/Engine/Plugins/Runtime/AudioSynesthesia/Source/AudioSynesthesiaEditor/Private/AssetTypeActions_AudioSynesthesiaSettings.cpp
@@ -0,0 +1,81 @@
+// Copyright Epic Games, Inc. All Rights Reserved.
+
+#include "AssetTypeActions_AudioSynesthesiaSettings.h"
+#include "Developer/AssetTools/Public/AssetTypeCategories.h"
+#include "AudioSynesthesia.h"
+
+FAssetTypeActions_AudioSynesthesiaSettings::FAssetTypeActions_AudioSynesthesiaSettings(UAudioSynesthesiaSettings* InSynesthesiaSettings)
+	: SynesthesiaSettings(InSynesthesiaSettings)
+{
+}
+
+bool FAssetTypeActions_AudioSynesthesiaSettings::CanFilter()
+{
+	// If no paired settings pointer provided, we filter as its a base class.
+	// Otherwise, we do not as this bloats the filter list.
+	return SynesthesiaSettings == nullptr;
+}
+
+FText FAssetTypeActions_AudioSynesthesiaSettings::GetName() const
+{
+	if (SynesthesiaSettings)
+	{
+		const FText AssetActionName = SynesthesiaSettings->GetAssetActionName();
+		if (AssetActionName.IsEmpty())
+		{
+			FString ClassName;
+			SynesthesiaSettings->GetClass()->GetName(ClassName);
+			return FText::FromString(ClassName);
+		}
+
+		return AssetActionName;
+	}
+
+	static const FText DefaultAssetActionName = NSLOCTEXT("AssetTypeActions", "AssetTypeActions_AssetSoundSynesthesiaSettings", "Synesthesia Settings");
+	return DefaultAssetActionName;
+}
+
+FColor FAssetTypeActions_AudioSynesthesiaSettings::GetTypeColor() const 
+{
+	if (!SynesthesiaSettings)
+	{
+		return FColor(100.0f, 50.0f, 100.0f);
+	}
+	return SynesthesiaSettings->GetTypeColor();
+}
+
+const TArray<FText>& FAssetTypeActions_AudioSynesthesiaSettings::GetSubMenus() const
+{
+	static const TArray<FText> SubMenus
+	{
+		NSLOCTEXT("AssetTypeActions", "AssetTypeActions_AssetSoundAnalysisSubmenu", "Analysis")
+	};
+
+	if (!SynesthesiaSettings)
+	{
+		return SubMenus;
+	}
+
+	return SynesthesiaSettings->GetAssetActionSubmenus();
+}
+
+UClass* FAssetTypeActions_AudioSynesthesiaSettings::GetSupportedClass() const
+{
+	if (SynesthesiaSettings)
+	{
+		if (UClass* SupportedClass = SynesthesiaSettings->GetSupportedClass())
+		{
+			return SupportedClass;
+		}
+
+		return SynesthesiaSettings->GetClass();
+	}
+
+	return UAudioSynesthesiaSettings::StaticClass();
+}
+
+uint32 FAssetTypeActions_AudioSynesthesiaSettings::GetCategories() 
+{
+	return EAssetTypeCategories::Sounds; 
+}
+
diff --git a/Engine/Plugins/Runtime/AudioSynesthesia/Source/AudioSynesthesiaEditor/Private/AssetTypeActions_AudioSynesthesiaSettings.h b/Engine/Plugins/Runtime/AudioSynesthesia/Source/AudioSynesthesiaEditor/Private/AssetTypeActions_AudioSynesthesiaSettings.h
new file mode 100644
--- /dev/null
+++ b/Engine/Plugins/Runtime/AudioSynesthesia/Source/AudioSynesthesiaEditor/Private/AssetTypeActions_AudioSynesthesiaSettings.h
@@ -0,0 +1,27 @@
+// Copyright Epic Games, Inc. All Rights Reserved.
+
+#pragma once
+
+#include "CoreMinimal.h"
+#include "AssetTypeActions_Base.h"
+
+class UAudioSynesthesiaSettings;
+
+/** FAssetTypeActions_AudioSynesthesiaSettings */
+class FAssetTypeActions_AudioSynesthesiaSettings : public FAssetTypeActions_Base
+{
+public:
+	FAssetTypeActions_AudioSynesthesiaSettings(UAudioSynesthesiaSettings* InSynesthesiaSettings);
+
+	//~ Begin FAssetTypeActions_Base
+	virtual bool CanFilter() override;
+	virtual FText GetName() const override;
+	virtual FColor GetTypeColor() const override;
+	virtual UClass* GetSupportedClass() const override;
+	virtual uint32 GetCategories() override;
+	virtual const TArray<FText>& GetSubMenus() const override;
+	//~ End FAssetTypeActions_Base
+
+private:
+	UAudioSynesthesiaSettings* SynesthesiaSettings;
+};
diff --git a/Engine/Plugins/Runtime/AudioSynesthesia/Source/AudioSynesthesiaEditor/Private/AudioSynesthesiaEditorModule.cpp b/Engine/Plugins/Runtime/AudioSynesthesia/Source/AudioSynesthesiaEditor/Private/AudioSynesthesiaEditorModule.cpp
--- a/Engine/Plugins/Runtime/AudioSynesthesia/Source/AudioSynesthesiaEditor/Private/AudioSynesthesiaEditorModule.cpp
+++ b/Engine/Plugins/Runtime/AudioSynesthesia/Source/AudioSynesthesiaEditor/Private/AudioSynesthesiaEditorModule.cpp
@@ -6,8 +6,11 @@
 #include "AssetToolsModule.h"
 #include "UObject/UObjectIterator.h"
 #include "AudioSynesthesiaNRT.h"
+#include "AudioSynesthesia.h"
 #include "AssetTypeActions_AudioSynesthesiaNRT.h"
 #include "AssetTypeActions_AudioSynesthesiaNRTSettings.h"
+#include "AssetTypeActions_AudioSynesthesiaSettings.h"
+
 
 DEFINE_LOG_CATEGORY(LogAudioSynesthesiaEditor);
 
@@ -30,6 +33,7 @@ public:
 	{
 		RegisterAudioSynesthesiaAssetActions<UAudioSynesthesiaNRT, FAssetTypeActions_AudioSynesthesiaNRT>();
 		RegisterAudioSynesthesiaAssetActions<UAudioSynesthesiaNRTSettings, FAssetTypeActions_AudioSynesthesiaNRTSettings>();
+		RegisterAudioSynesthesiaAssetActions<UAudioSynesthesiaSettings, FAssetTypeActions_AudioSynesthesiaSettings>();
 	}
 
 private:
diff --git a/Engine/Plugins/Runtime/AudioSynesthesia/Source/AudioSynesthesiaEditor/Private/AudioSynesthesiaSettingsFactory.cpp b/Engine/Plugins/Runtime/AudioSynesthesia/Source/AudioSynesthesiaEditor/Private/AudioSynesthesiaSettingsFactory.cpp
new file mode 100644
--- /dev/null
+++ b/Engine/Plugins/Runtime/AudioSynesthesia/Source/AudioSynesthesiaEditor/Private/AudioSynesthesiaSettingsFactory.cpp
@@ -0,0 +1,61 @@
+// Copyright Epic Games, Inc. All Rights Reserved.
+
+#include "AudioSynesthesiaSettingsFactory.h"
+
+#include "ClassViewerModule.h"
+#include "ClassViewerFilter.h"
+#include "Kismet2/SClassPickerDialog.h"
+#include "Modules/ModuleManager.h"
+#include "AudioSynesthesiaClassFilter.h"
+
+UAudioSynesthesiaSettingsFactory::UAudioSynesthesiaSettingsFactory(const FObjectInitializer& ObjectInitializer)
+	: Super(ObjectInitializer)
+{
+	SupportedClass = UAudioSynesthesiaSettings::StaticClass();
+
+	bCreateNew = true;
+	bEditorImport = false;
+	bEditAfterNew = true;
+	AudioSynesthesiaSettingsClass = nullptr;
+}
+#define LOCTEXT_NAMESPACE "AudioSynesthesiaEditor"
+
+bool UAudioSynesthesiaSettingsFactory::ConfigureProperties()
+{
+	AudioSynesthesiaSettingsClass = nullptr;
+
+	// Load the classviewer module to display a class picker
+	FClassViewerModule& ClassViewerModule = FModuleManager::LoadModuleChecked<FClassViewerModule>("ClassViewer");
+
+	FClassViewerInitializationOptions Options;
+	Options.Mode = EClassViewerMode::ClassPicker;
+
+	TSharedPtr<Audio::FAssetClassParentFilter> Filter = MakeShareable(new Audio::FAssetClassParentFilter);
+	Options.ClassFilter = Filter;
+
+	Filter->DisallowedClassFlags = CLASS_Abstract | CLASS_Deprecated | CLASS_NewerVersionExists;
+	Filter->AllowedChildrenOfClasses.Add(UAudioSynesthesiaSettings::StaticClass());
+
+	const FText TitleText = LOCTEXT("CreateAudioSynesthesiaSettingOptions", "Pick Synesthesia Settings Class");
+	UClass* ChosenClass = nullptr;
+	const bool bPressedOk = SClassPickerDialog::PickClass(TitleText, Options, ChosenClass, UAudioSynesthesiaSettings::StaticClass());
+
+	if (bPressedOk)
+	{
+		AudioSynesthesiaSettingsClass = ChosenClass;
+	}
+
+	return bPressedOk;
+}
+
+UObject* UAudioSynesthesiaSettingsFactory::FactoryCreateNew(UClass* InClass, UObject* InParent, FName InName, EObjectFlags Flags, UObject* Context, FFeedbackContext* Warn)
+{
+	UAudioSynesthesiaSettings* NewAudioSynesthesiaSettings = nullptr;
+	if (AudioSynesthesiaSettingsClass != nullptr)
+	{
+		NewAudioSynesthesiaSettings = NewObject<UAudioSynesthesiaSettings>(InParent, AudioSynesthesiaSettingsClass, InName, Flags);
+	}
+	return NewAudioSynesthesiaSettings;
+}
+
+#undef LOCTEXT_NAMESPACE
diff --git a/Engine/Source/Runtime/AudioAnalyzer/AudioAnalyzer.Build.cs b/Engine/Source/Runtime/AudioAnalyzer/AudioAnalyzer.Build.cs
--- a/Engine/Source/Runtime/AudioAnalyzer/AudioAnalyzer.Build.cs
+++ b/Engine/Source/Runtime/AudioAnalyzer/AudioAnalyzer.Build.cs
@@ -13,8 +13,9 @@ namespace UnrealBuildTool.Rules
 				{
                     "InputCore",
                     "CoreUObject",
-                    "Engine"
-                }
+                    "Engine",
+					"AudioMixer"
+				}
 			);
 
 			PublicDependencyModuleNames.AddRange(
diff --git a/Engine/Source/Runtime/AudioAnalyzer/Classes/AudioAnalyzer.h b/Engine/Source/Runtime/AudioAnalyzer/Classes/AudioAnalyzer.h
new file mode 100644
--- /dev/null
+++ b/Engine/Source/Runtime/AudioAnalyzer/Classes/AudioAnalyzer.h
@@ -0,0 +1,162 @@
+// Copyright Epic Games, Inc. All Rights Reserved.
+
+#pragma once
+
+#include "CoreMinimal.h"
+#include "IAudioAnalyzerInterface.h"
+#include "AudioAnalyzerAsset.h"
+#include "AudioAnalyzerFacade.h"
+#include "Sound/AudioBus.h"
+#include "DSP/MultithreadedPatching.h"
+#include "AudioAnalyzer.generated.h"
+
+class UAudioAnalyzerSubsystem;
+
+/** UAudioAnalyzerSettings
+ *
+ * UAudioAnalyzerSettings provides a way to store and reuse existing analyzer settings
+ * across multiple analyzers. 
+ *
+ */
+UCLASS(Abstract, EditInlineNew, BlueprintType)
+class AUDIOANALYZER_API UAudioAnalyzerSettings : public UAudioAnalyzerAssetBase
+{
+	GENERATED_BODY()
+};
+
+typedef TSharedPtr<Audio::IAnalyzerResult, ESPMode::ThreadSafe> FAnalyzerResultSharedPtr;
+
+class FAudioAnalyzeTask : public FNonAbandonableTask
+{
+public:
+	FAudioAnalyzeTask(TUniquePtr<Audio::FAnalyzerFacade>& InAnalyzerFacade, int32 InSampleRate, int32 InNumChannels);
+	
+	// Give the task the audio data to analyze
+	void SetAudioBuffer(TArray<float>&& InAudioData);
+
+	// Move the audio buffer back out
+	TArray<float>&& GetAudioBuffer() { return MoveTemp(AudioData); }
+
+	// Does the task work
+	void DoWork();
+
+	FORCEINLINE TStatId GetStatId() const { RETURN_QUICK_DECLARE_CYCLE_STAT(AudioAnalyzeTask, STATGROUP_ThreadPoolAsyncTasks); }
+
+	// Get the results from the task
+	TUniquePtr<Audio::IAnalyzerResult> GetResults() { return MoveTemp(Results); }
+
+private:
+	TUniquePtr<Audio::FAnalyzerFacade> AnalyzerFacade;
+	int32 SampleRate = 0;
+	int32 NumChannels = 0;
+	TArray<float> AudioData;
+	TUniquePtr<Audio::IAnalyzerResult> Results;
+};
+
+/** UAudioAnalyzer
+ *
+ * UAudioAnalyzer performs analysis on an audio bus using specific settings and exposes the results via blueprints.
+ *
+ * Subclasses of UAudioAnalyzer must implement GetAnalyzerFactoryName() to associate
+ * the UAudioAnalyzer asset with an IAudioAnalyzerFactory implementation.
+ *
+ * To support blueprint access, subclasses can implement UFUNCTIONs to expose the data
+ * returned by GetResult().
+ */
+UCLASS(Abstract, EditInlineNew, BlueprintType)
+class AUDIOANALYZER_API UAudioAnalyzer : public UObject
+{
+	GENERATED_BODY()
+	
+	friend class UAudioAnalyzerSubsystem;
+
+public:
+	/**
+	 * ID to keep track of results. Useful for tracking most recent result when performing
+	 * asynchronous processing.
+	 */
+	typedef int32 FResultId;
+
+	/** Thread safe shared point to result object. */
+	typedef TSharedPtr<Audio::IAnalyzerResult, ESPMode::ThreadSafe> FResultSharedPtr;
+
+	/** The UAudioBus which is analyzed in real-time. */
+	UPROPERTY(Transient)
+	UAudioBus* AudioBus;
+
+	/** Starts analyzing audio from the given audio bus. Optionally override the audio bus desired to analyze. */
+	UFUNCTION(BlueprintCallable, BlueprintCosmetic, Category = AudioAnalyzer, meta = (WorldContext = "WorldContextObject"))
+	void StartAnalyzing(const UObject* WorldContextObject, UAudioBus* AudioBusToAnalyze);
+
+	/** Stops analyzing audio. */
+	UFUNCTION(BlueprintCallable, BlueprintCosmetic, Category = AudioAnalyzer, meta = (WorldContext = "WorldContextObject"))
+	void StopAnalyzing(const UObject* WorldContextObject);
+
+	/**
+	 * Implementations can override this method to create settings objects
+	 * specific for their analyzer.
+	 */
+	virtual TUniquePtr<Audio::IAnalyzerSettings> GetSettings(const int32 InSampleRate, const int32 InNumChannels) const;
+
+	/** Function to broadcast results. */
+	virtual void BroadcastResults() {}
+
+	//~ Begin UObject Interface.
+	virtual void BeginDestroy() override;
+	//~ End UObject interface
+
+private: 
+	/** Returns if this audio analyzer has enough audio queued up and is ready for analysis. */
+	bool IsReadyForAnalysis() const;
+
+	/** Does the actual analysis and casts the results to the derived classes type. Returns true if there results to broadcast. */
+	bool DoAnalysis();
+
+protected:
+
+	template<class ResultType>
+	TUniquePtr<ResultType> GetResults()
+	{
+		return TUniquePtr<ResultType>(static_cast<ResultType*>(ResultsInternal.Release()));
+	}
+
+	/* Subclasses must override this method in order to inform this object which AnalyzerFactory to use for analysis */
+	virtual FName GetAnalyzerFactoryName() const PURE_VIRTUAL(UAudioAnalyzer::GetAnalyzerFactoryName, return FName(););
+
+	/** How many frames of audio to wait before analyzing the audio. */
+	int32 NumFramesPerBufferToAnalyze = 1024;
+
+private:
+
+	// Returns UAudioAnalyzerSettings* if property points to a valid UAudioAnalyzerSettings, otherwise returns nullptr.
+	UAudioAnalyzerSettings* GetSettingsFromProperty(FProperty* Property);
+
+	// Audio analysis subsystem used with this audio analyzer
+	UPROPERTY(Transient)
+	UAudioAnalyzerSubsystem* AudioAnalyzerSubsystem;
+
+	// Output patch for retrieving audio from audio bus for analysis
+	Audio::FPatchOutputStrongPtr PatchOutputStrongPtr;
+
+	// Analysis task
+	TSharedPtr<FAsyncTask<FAudioAnalyzeTask>> AnalysisTask;
+
+	// The analyzer facade used for async tasks
+	TUniquePtr<Audio::FAnalyzerFacade> AnalyzerFacade;
+
+	// Cached results of previous analysis task. Can be invalid.
+	TUniquePtr<Audio::IAnalyzerResult> ResultsInternal;
+
+	// Scratch buffer used for copying data from patch output and fed to analysis task
+	TArray<float> AnalysisBuffer;
+
+	// The sample of the audio renderer
+	int32 AudioMixerSampleRate = 0;
+
+	// The number of channels for the audio bus
+	int32 NumBusChannels = 0;
+
+	// The analyzer factory to use
+	Audio::IAnalyzerFactory* AnalyzerFactory = nullptr;
+};
+
diff --git a/Engine/Source/Runtime/AudioAnalyzer/Classes/AudioAnalyzerAsset.h b/Engine/Source/Runtime/AudioAnalyzer/Classes/AudioAnalyzerAsset.h
--- a/Engine/Source/Runtime/AudioAnalyzer/Classes/AudioAnalyzerAsset.h
+++ b/Engine/Source/Runtime/AudioAnalyzer/Classes/AudioAnalyzerAsset.h
@@ -7,12 +7,14 @@
 #include "CoreMinimal.h"
 #include "AudioAnalyzerAsset.generated.h"
 
-/** UAudioAnalyzerAsset
+DECLARE_MULTICAST_DELEGATE(FAnalyzeAudioDelegate);
+
+/** UAudioAnalyzerAssetBase
  *
- * UAudioAnalyzerAsset provides the interface for controlling asset actions within the editor.
+ * UAudioAnalyzerAssetBase provides the base interface for controlling asset actions within the editor.
  */
 UCLASS(Abstract, EditInlineNew)
-class AUDIOANALYZER_API UAudioAnalyzerAsset : public UObject
+class AUDIOANALYZER_API UAudioAnalyzerAssetBase : public UObject
 {
 	GENERATED_BODY()
 
diff --git a/Engine/Source/Runtime/AudioAnalyzer/Classes/AudioAnalyzerNRT.h b/Engine/Source/Runtime/AudioAnalyzer/Classes/AudioAnalyzerNRT.h
--- a/Engine/Source/Runtime/AudioAnalyzer/Classes/AudioAnalyzerNRT.h
+++ b/Engine/Source/Runtime/AudioAnalyzer/Classes/AudioAnalyzerNRT.h
@@ -10,18 +10,18 @@
 #include "Sound/SoundWave.h"
 #include "AudioAnalyzerNRT.generated.h"
 
-DECLARE_MULTICAST_DELEGATE( FAnalyzeAudioDelegate );
+
 
 /** UAudioAnalyzerNRTSettings
  *
  * UAudioAnalyzerNRTSettings provides a way to store and reuse existing analyzer settings
- * across multipler analyzers. This class provides the interface and functionality to 
+ * across multiple analyzers. This class provides the interface and functionality to 
  * automatically trigger reanalysis of audio across all analyzers associated with this 
- * settingwhen when a UPROPERTY in this setting object is edited.
+ * setting when when a UPROPERTY in this setting object is edited.
  *
  */
 UCLASS(Abstract, EditInlineNew, BlueprintType)
-class AUDIOANALYZER_API UAudioAnalyzerNRTSettings : public UAudioAnalyzerAsset
+class AUDIOANALYZER_API UAudioAnalyzerNRTSettings : public UAudioAnalyzerAssetBase
 {
 	GENERATED_BODY()
 
@@ -60,7 +60,7 @@ class AUDIOANALYZER_API UAudioAnalyzerNRTSettings : public UAudioAnalyzerAsset
  * returned by GetResult().
  */
 UCLASS(Abstract, EditInlineNew, BlueprintType)
-class AUDIOANALYZER_API UAudioAnalyzerNRT : public UAudioAnalyzerAsset
+class AUDIOANALYZER_API UAudioAnalyzerNRT : public UAudioAnalyzerAssetBase
 {
 	GENERATED_BODY()
 
diff --git a/Engine/Source/Runtime/AudioAnalyzer/Classes/AudioAnalyzerSubsystem.h b/Engine/Source/Runtime/AudioAnalyzer/Classes/AudioAnalyzerSubsystem.h
new file mode 100644
--- /dev/null
+++ b/Engine/Source/Runtime/AudioAnalyzer/Classes/AudioAnalyzerSubsystem.h
@@ -0,0 +1,43 @@
+// Copyright Epic Games, Inc. All Rights Reserved.
+
+#pragma once
+
+#include "CoreMinimal.h"
+#include "AudioAnalyzer.h"
+#include "Subsystems/WorldSubsystem.h"
+#include "Tickable.h"
+#include "AudioAnalyzerSubsystem.generated.h"
+
+class UWorld;
+
+/** 
+* Class manages querying analysis results from various audio analyzers.
+*/
+UCLASS(DisplayName = "Audio Analysis")
+class AUDIOANALYZER_API UAudioAnalyzerSubsystem : public UWorldSubsystem, public FTickableGameObject
+{
+	GENERATED_BODY()
+
+public:
+	UAudioAnalyzerSubsystem();
+	~UAudioAnalyzerSubsystem();
+
+	//~ Begin FTickableGameObject Interface
+	virtual void Tick(float DeltaTime) override;
+	virtual bool IsTickable() const override;
+	virtual TStatId GetStatId() const override;
+	//~ End FTickableGameObject Interface
+
+	void RegisterAudioAnalyzer(UAudioAnalyzer* InAnalyzer);
+	void UnregisterAudioAnalyzer(UAudioAnalyzer* InAnalyzer);
+
+	static UAudioAnalyzerSubsystem* Get(UWorld* World);
+
+private:
+
+	UPROPERTY(Transient);
+	TArray<UAudioAnalyzer*> AudioAnalyzers;
+};
+
+
+
diff --git a/Engine/Source/Runtime/AudioAnalyzer/Private/AudioAnalyzer.cpp b/Engine/Source/Runtime/AudioAnalyzer/Private/AudioAnalyzer.cpp
new file mode 100644
--- /dev/null
+++ b/Engine/Source/Runtime/AudioAnalyzer/Private/AudioAnalyzer.cpp
@@ -0,0 +1,220 @@
+// Copyright Epic Games, Inc. All Rights Reserved.
+
+
+#include "AudioAnalyzer.h"
+#include "AudioAnalyzerFacade.h"
+#include "AudioAnalyzerModule.h"
+#include "Async/Async.h"
+#include "AudioMixerDevice.h"
+#include "AudioDeviceManager.h"
+#include "AudioAnalyzerSubsystem.h"
+
+
+FAudioAnalyzeTask::FAudioAnalyzeTask(TUniquePtr<Audio::FAnalyzerFacade>& InAnalyzerFacade, int32 InSampleRate, int32 InNumChannels)
+	: AnalyzerFacade(InAnalyzerFacade.Release())
+	, SampleRate(InSampleRate)
+	, NumChannels(InNumChannels)
+{
+}
+
+void FAudioAnalyzeTask::SetAudioBuffer(TArray<float>&& InAudioData)
+{
+	AudioData = MoveTemp(InAudioData);
+}
+
+void FAudioAnalyzeTask::DoWork()
+{
+	check(AnalyzerFacade);
+	Results = AnalyzerFacade->AnalyzeAudioBuffer(AudioData, NumChannels, SampleRate);
+}
+
+void UAudioAnalyzer::StartAnalyzing(const UObject* WorldContextObject, UAudioBus* AudioBusToAnalyze)
+{
+	if (!AudioBusToAnalyze)
+	{
+		UE_LOG(LogAudioAnalyzer, Error, TEXT("Unable to analyze audio without an audio bus to analyze."));
+		return;
+	}
+
+	// Retrieve the world and if it's not a valid world, don't do anything
+	UWorld* ThisWorld = GEngine->GetWorldFromContextObject(WorldContextObject, EGetWorldErrorMode::LogAndReturnNull);
+	if (!ThisWorld || !ThisWorld->bAllowAudioPlayback || ThisWorld->GetNetMode() == NM_DedicatedServer)
+	{
+		return;
+	}
+
+	// Get the audio mixer device to get audio mixer sample rate and be able to play audio buses
+	Audio::FMixerDevice* MixerDevice = FAudioDeviceManager::GetAudioMixerDeviceFromWorldContext(WorldContextObject);
+	if (!MixerDevice)
+	{
+		UE_LOG(LogAudioAnalyzer, Error, TEXT("Audio analyzer only works with the audio mixer."));
+		return;
+	}
+
+	// Retrieve the audio analyzer subsystem if it's not already retrieved
+	if (!AudioAnalyzerSubsystem)
+	{
+		AudioAnalyzerSubsystem = UAudioAnalyzerSubsystem::Get(ThisWorld);
+	}
+
+	// Store the audio bus we're analyzing
+ 	AudioBus = AudioBusToAnalyze;
+
+ 	NumBusChannels = AudioBus->GetNumChannels();
+ 	check(NumBusChannels > 0);
+ 	
+	AudioMixerSampleRate = (int32)MixerDevice->GetSampleRate();
+ 
+ 	// Get the analyzer factory to use for our analysis
+	if (!AnalyzerFactory)
+	{
+ 		AnalyzerFactory = Audio::GetAnalyzerFactory(GetAnalyzerFactoryName());
+	}
+	checkf(AnalyzerFactory != nullptr, TEXT("Need to register the factory as a modular feature for the analyzer '%s'."), *GetAnalyzerFactoryName().ToString());
+
+ 	// Start the audio bus. This won't do anythign if the bus is already started elsewhere.
+ 	uint32 AudioBusId = AudioBus->GetUniqueID();
+ 	int32 NumChannels = (int32)AudioBus->AudioBusChannels + 1;
+ 	MixerDevice->StartAudioBus(AudioBusId, NumChannels, false);
+ 
+ 	// Get an output patch for the audio bus
+ 	PatchOutputStrongPtr = MixerDevice->AddPatchForAudioBus(AudioBusId);
+
+	// Register this audio analyzer with the audio analyzer subsystem
+	// The subsystem will query this analyzer to see if it has enough audio to perform analysis.
+	// If it does, it'll report the results (of any previous analysis) and ask us to start analyzing the audio.		
+	AudioAnalyzerSubsystem->RegisterAudioAnalyzer(this);
+
+	// Setup the analyzer facade here once
+	AnalyzerFacade = MakeUnique<Audio::FAnalyzerFacade>(GetSettings(AudioMixerSampleRate, NumBusChannels), AnalyzerFactory);
+}
+
+void UAudioAnalyzer::StopAnalyzing(const UObject* WorldContextObject)
+{
+	if (AudioAnalyzerSubsystem)
+	{
+		// Unregister the audio analyzer from the analyzer subsystem
+		AudioAnalyzerSubsystem->UnregisterAudioAnalyzer(this);
+
+		// Null out our patch output
+		PatchOutputStrongPtr = nullptr;
+	}
+}
+
+bool UAudioAnalyzer::IsReadyForAnalysis() const
+{
+	check(AudioBus != nullptr);
+
+	// Only allow one worker task to happen at a time
+	if (AnalysisTask.IsValid() && !AnalysisTask->IsWorkDone())
+	{
+		return false;
+	}
+
+	// Only ready if we've got a patch output and there's enough audio queued up
+	if (PatchOutputStrongPtr.IsValid())
+	{
+		int32 NumSamplesAvailable = PatchOutputStrongPtr->GetNumSamplesAvailable();
+		int32 NumAudioBusChannels = AudioBus->GetNumChannels();
+		check(NumAudioBusChannels > 0);
+		int32 NumFramesAvailable = NumSamplesAvailable / NumAudioBusChannels;
+		return NumFramesAvailable >= NumFramesPerBufferToAnalyze;
+	}
+	return false;
+}
+
+bool UAudioAnalyzer::DoAnalysis()
+{
+	bool bHasResults = false;
+
+	// Make sure the task is finished before we do any more analysis
+	if (AnalysisTask.IsValid())
+	{
+		check(AnalysisTask->IsWorkDone());
+
+		// Do final ensure completion before reusing the task
+		AnalysisTask->EnsureCompletion();
+
+		// Report results from previous analysis
+		FAudioAnalyzeTask& Task = AnalysisTask->GetTask();
+		ResultsInternal = Task.GetResults();
+
+		// Retrieve the analysis buffer so we can reuse the memory when we copy to it from the patch output
+		AnalysisBuffer = Task.GetAudioBuffer();
+
+		bHasResults = true;
+	}
+	else
+	{
+		// Create a task if one doesn't exist
+		AnalysisTask = TSharedPtr<FAsyncTask<FAudioAnalyzeTask>>(new FAsyncTask<FAudioAnalyzeTask>(AnalyzerFacade, AudioMixerSampleRate, NumBusChannels));
+	}
+
+	// Make sure our task is valid and done
+	check(AnalysisTask.IsValid() && AnalysisTask->IsDone());
+
+	// Copy the audio to be analyzed from the patch output
+	int32 NumSamplesAvailable = PatchOutputStrongPtr->GetNumSamplesAvailable();
+
+	AnalysisBuffer.Reset();
+	AnalysisBuffer.AddUninitialized(NumSamplesAvailable);
+	PatchOutputStrongPtr->PopAudio(AnalysisBuffer.GetData(), NumSamplesAvailable, false);
+
+	check(AudioMixerSampleRate != 0);
+	check(NumBusChannels != 0);
+
+	FAudioAnalyzeTask& Task = AnalysisTask->GetTask();
+	Task.SetAudioBuffer(MoveTemp(AnalysisBuffer));
+
+ 	AnalysisTask->StartBackgroundTask();
+
+	return bHasResults;
+}
+
+void UAudioAnalyzer::BeginDestroy()
+{
+	Super::BeginDestroy();
+
+	if (AnalysisTask.IsValid())
+	{
+		AnalysisTask->EnsureCompletion();
+	}
+
+	// Unregister this audio analyzer from the subsystem since we won't be needing anymore ticks
+	if (AudioAnalyzerSubsystem)
+	{
+		AudioAnalyzerSubsystem->UnregisterAudioAnalyzer(this);
+	}
+}
+
+UAudioAnalyzerSettings* UAudioAnalyzer::GetSettingsFromProperty(FProperty* Property)
+{
+	if (nullptr == Property)
+	{
+		return nullptr;
+	}
+
+	if (Property->IsA(FObjectPropertyBase::StaticClass()))
+	{
+		FObjectPropertyBase* ObjectPropertyBase = CastFieldChecked<FObjectPropertyBase>(Property);
+
+		if (nullptr == ObjectPropertyBase)
+		{
+			return nullptr;
+		}
+
+		if (ObjectPropertyBase->PropertyClass->IsChildOf(UAudioAnalyzerSettings::StaticClass()))
+		{
+			UObject* PropertyObject = ObjectPropertyBase->GetObjectPropertyValue_InContainer(this);
+			return Cast<UAudioAnalyzerSettings>(PropertyObject);
+		}
+	}
+
+	return nullptr;
+}
+
+TUniquePtr<Audio::IAnalyzerSettings> UAudioAnalyzer::GetSettings(const int32 InSampleRate, const int32 InNumChannels) const
+{
+	return MakeUnique<Audio::IAnalyzerSettings>();
+}
+
diff --git a/Engine/Source/Runtime/AudioAnalyzer/Private/AudioAnalyzerFacade.cpp b/Engine/Source/Runtime/AudioAnalyzer/Private/AudioAnalyzerFacade.cpp
new file mode 100644
--- /dev/null
+++ b/Engine/Source/Runtime/AudioAnalyzer/Private/AudioAnalyzerFacade.cpp
@@ -0,0 +1,82 @@
+// Copyright Epic Games, Inc. All Rights Reserved.
+
+#include "AudioAnalyzerFacade.h"
+
+#include "AudioAnalyzerModule.h"
+#include "IAudioAnalyzerInterface.h"
+#include "Features/IModularFeatures.h"
+
+namespace Audio
+{
+	IAnalyzerFactory* GetAnalyzerFactory(FName InFactoryName)
+	{
+		// Get all analyzer nrt factories implementations.
+		TArray<IAnalyzerFactory*> RegisteredFactories = IModularFeatures::Get().GetModularFeatureImplementations<IAnalyzerFactory>(IAnalyzerFactory::GetModularFeatureName());
+
+		// Get the factory of interest by matching the name.
+		TArray<IAnalyzerFactory*> MatchingFactories = RegisteredFactories.FilterByPredicate([InFactoryName](IAnalyzerFactory* Factory) { check(nullptr != Factory); return Factory->GetName() == InFactoryName; });
+
+		if (0 == MatchingFactories.Num())
+		{
+			// There is a likely programming error if the factory is not found. 
+			UE_LOG(LogAudioAnalyzer, Error, TEXT("Failed to find factory of type '%s' with name '%s'"), *IAnalyzerFactory::GetModularFeatureName().ToString(), *InFactoryName.ToString());
+
+			return nullptr;
+		}
+
+		if (MatchingFactories.Num() > 1)
+		{
+			// If multiple factories with the same name exist, the first one in the array will be used. 
+			UE_LOG(LogAudioAnalyzer, Warning, TEXT("Found multiple factories of type '%s' with name '%s'. Factory names should be unique."), *IAnalyzerFactory::GetModularFeatureName().ToString(), *InFactoryName.ToString());
+		}
+
+		return MatchingFactories[0];
+	}
+
+	FAnalyzerFacade::FAnalyzerFacade(TUniquePtr<IAnalyzerSettings> InSettings, IAnalyzerFactory* InFactory)
+		: Settings(MoveTemp(InSettings))
+		, Factory(InFactory)
+	{
+	}
+
+	TUniquePtr<IAnalyzerResult> FAnalyzerFacade::AnalyzeAudioBuffer(const TArray<float>& InAudioBuffer, int32 InNumChannels, float InSampleRate)
+	{
+		if (nullptr == Factory)
+		{
+			UE_LOG(LogAudioAnalyzer, Error, TEXT("Cannot analyze audio due to null factory"));
+
+			return TUniquePtr<IAnalyzerResult>();
+		}
+
+		// Create result and worker from factory
+		if (!Worker.IsValid())
+		{
+			FAnalyzerParameters AnalyzerParameters(InSampleRate, InNumChannels);
+			Worker = Factory->NewWorker(AnalyzerParameters, Settings.Get());
+		}
+
+		// Check that worker created successfully
+		if (!Worker.IsValid())
+		{
+			UE_LOG(LogAudioAnalyzer, Error, TEXT("Failed to create IAnalyzerWorker with factory of type '%s' with name '%s'"), *IAnalyzerFactory::GetModularFeatureName().ToString(), *Factory->GetName().ToString());
+
+			return TUniquePtr<IAnalyzerResult>();
+		}
+
+		TUniquePtr<IAnalyzerResult> Result = Factory->NewResult();
+
+		// Check that result created successfully
+		if (!Result.IsValid())
+		{
+			UE_LOG(LogAudioAnalyzer, Error, TEXT("Failed to create IAnalyzerResult with factory of type '%s' with name '%s'"), *Audio::IAnalyzerFactory::GetModularFeatureName().ToString(), *Factory->GetName().ToString());
+
+			return TUniquePtr<IAnalyzerResult>();
+		}
+
+		// Perform and finalize audio analysis.
+		Worker->Analyze(MakeArrayView(InAudioBuffer), Result.Get());
+
+		return Result;
+	}
+}
+
diff --git a/Engine/Source/Runtime/AudioAnalyzer/Private/AudioAnalyzerNRT.cpp b/Engine/Source/Runtime/AudioAnalyzer/Private/AudioAnalyzerNRT.cpp
--- a/Engine/Source/Runtime/AudioAnalyzer/Private/AudioAnalyzerNRT.cpp
+++ b/Engine/Source/Runtime/AudioAnalyzer/Private/AudioAnalyzerNRT.cpp
@@ -11,15 +11,15 @@
 
 namespace 
 {
-	class FAudioAnalyzeTask : public FNonAbandonableTask
+	class FAudioAnalyzeNRTTask : public FNonAbandonableTask
 	{
-		friend class FAutoDeleteAsyncTask<FAudioAnalyzeTask>;
+		friend class FAutoDeleteAsyncTask<FAudioAnalyzeNRTTask>;
 
 		public:
-			FAudioAnalyzeTask(
+			FAudioAnalyzeNRTTask(
 					TWeakObjectPtr<UAudioAnalyzerNRT> InAnalyzerUObject, 
 					const UAudioAnalyzerNRT::FResultId InResultId,
-					TUniquePtr<Audio::FAnalyzerNRTBatch>&& InAnalyzerFacade, 
+					TUniquePtr<Audio::FAnalyzerNRTFacade>&& InAnalyzerFacade, 
 					TArray<uint8>&& InRawWaveData,
 					int32 InNumChannels,
 					float InSampleRate)
@@ -45,12 +45,12 @@ namespace
 				});
 			}
 
-			FORCEINLINE TStatId GetStatId() const { RETURN_QUICK_DECLARE_CYCLE_STAT(AudioAnalyzeTask, STATGROUP_ThreadPoolAsyncTasks); }
+			FORCEINLINE TStatId GetStatId() const { RETURN_QUICK_DECLARE_CYCLE_STAT(AudioAnalyzeNRTTask, STATGROUP_ThreadPoolAsyncTasks); }
 
 		private:
 			TWeakObjectPtr<UAudioAnalyzerNRT> AnalyzerUObject;
 			const UAudioAnalyzerNRT::FResultId ResultId;
-			TUniquePtr<Audio::FAnalyzerNRTBatch> AnalyzerFacade;
+			TUniquePtr<Audio::FAnalyzerNRTFacade> AnalyzerFacade;
 			TArray<uint8> RawWaveData;
 			int32 NumChannels;
 			float SampleRate;
@@ -155,13 +155,13 @@ void UAudioAnalyzerNRT::AnalyzeAudio()
 		}
 		
 		// Create analyzer helper object
-		TUniquePtr<Audio::FAnalyzerNRTBatch> BatchAnalyzer = MakeUnique<Audio::FAnalyzerNRTBatch>(GetSettings(SampleRate, NumChannels), GetAnalyzerNRTFactoryName());
+		TUniquePtr<Audio::FAnalyzerNRTFacade> BatchAnalyzer = MakeUnique<Audio::FAnalyzerNRTFacade>(GetSettings(SampleRate, NumChannels), GetAnalyzerNRTFactoryName());
 
 		// Use weak reference in case this object is deleted before analysis is done
 		TWeakObjectPtr<UAudioAnalyzerNRT> AnalyzerPtr(this);
 		
 		// Create and start async task. Parentheses avoids memory leak warnings from static analysis.
-		(new FAutoDeleteAsyncTask<FAudioAnalyzeTask>(AnalyzerPtr, ThisResultId, MoveTemp(BatchAnalyzer), MoveTemp(RawWaveData), NumChannels, SampleRate))->StartBackgroundTask();
+		(new FAutoDeleteAsyncTask<FAudioAnalyzeNRTTask>(AnalyzerPtr, ThisResultId, MoveTemp(BatchAnalyzer), MoveTemp(RawWaveData), NumChannels, SampleRate))->StartBackgroundTask();
 	}
 	else
 	{
diff --git a/Engine/Source/Runtime/AudioAnalyzer/Private/AudioAnalyzerNRTFacade.cpp b/Engine/Source/Runtime/AudioAnalyzer/Private/AudioAnalyzerNRTFacade.cpp
--- a/Engine/Source/Runtime/AudioAnalyzer/Private/AudioAnalyzerNRTFacade.cpp
+++ b/Engine/Source/Runtime/AudioAnalyzer/Private/AudioAnalyzerNRTFacade.cpp
@@ -15,7 +15,7 @@ namespace Audio
 		TArray<IAnalyzerNRTFactory*> RegisteredFactories = IModularFeatures::Get().GetModularFeatureImplementations<IAnalyzerNRTFactory>(IAnalyzerNRTFactory::GetModularFeatureName());
 
 		// Get the factory of interest by matching the name.
-		TArray<Audio::IAnalyzerNRTFactory*> MatchingFactories = RegisteredFactories.FilterByPredicate([InFactoryName](IAnalyzerNRTFactory* Factory) { check(nullptr != Factory); return Factory->GetName() == InFactoryName; });
+		TArray<IAnalyzerNRTFactory*> MatchingFactories = RegisteredFactories.FilterByPredicate([InFactoryName](IAnalyzerNRTFactory* Factory) { check(nullptr != Factory); return Factory->GetName() == InFactoryName; });
 
 		if (0 == MatchingFactories.Num())
 		{
@@ -35,13 +35,13 @@ namespace Audio
 	}
 
 	/**********************************************************/
-	/*************** FAnalyzerNRTBatch ************************/
+	/*************** FAnalyzerNRTFacade ************************/
 	/**********************************************************/
 
 	/**
-	 * Create an FAnalyzerNRTBatch with the analyzer settings and factory name.
+	 * Create an FAnalyzerNRTFacade with the analyzer settings and factory name.
 	 */
-	FAnalyzerNRTBatch::FAnalyzerNRTBatch(TUniquePtr<IAnalyzerNRTSettings> InSettings, const FName& InFactoryName)
+	FAnalyzerNRTFacade::FAnalyzerNRTFacade(TUniquePtr<IAnalyzerNRTSettings> InSettings, const FName& InFactoryName)
 	: Settings(MoveTemp(InSettings))
 	, FactoryName(InFactoryName)
 	{}
@@ -49,7 +49,7 @@ namespace Audio
 	/**
 	 * Analyze an entire PCM16 encoded audio object.  Audio for the entire sound should be contained within InRawWaveData.
 	 */
-	TUniquePtr<IAnalyzerNRTResult> FAnalyzerNRTBatch::AnalyzePCM16Audio(const TArray<uint8>& InRawWaveData, int32 InNumChannels, float InSampleRate)
+	TUniquePtr<IAnalyzerNRTResult> FAnalyzerNRTFacade::AnalyzePCM16Audio(const TArray<uint8>& InRawWaveData, int32 InNumChannels, float InSampleRate)
 	{
 		// Get analyzer factory
 		IAnalyzerNRTFactory* Factory = GetAnalyzerNRTFactory(FactoryName);
@@ -61,7 +61,7 @@ namespace Audio
 			return TUniquePtr<IAnalyzerNRTResult>();
 		}
 
-		// Create result and worker from fractory
+		// Create result and worker from factory
 		FAnalyzerNRTParameters AnalyzerNRTParameters(InSampleRate, InNumChannels);
 
 		TUniquePtr<IAnalyzerNRTResult> Result = Factory->NewResult();
diff --git a/Engine/Source/Runtime/AudioAnalyzer/Private/AudioAnalyzerSubsystem.cpp b/Engine/Source/Runtime/AudioAnalyzer/Private/AudioAnalyzerSubsystem.cpp
new file mode 100644
--- /dev/null
+++ b/Engine/Source/Runtime/AudioAnalyzer/Private/AudioAnalyzerSubsystem.cpp
@@ -0,0 +1,66 @@
+// Copyright Epic Games, Inc. All Rights Reserved.
+
+#include "AudioAnalyzerSubsystem.h"
+#include "Engine/World.h"
+
+UAudioAnalyzerSubsystem::UAudioAnalyzerSubsystem()
+{
+}
+
+UAudioAnalyzerSubsystem::~UAudioAnalyzerSubsystem()
+{
+	AudioAnalyzers.Reset();
+}
+
+UAudioAnalyzerSubsystem* UAudioAnalyzerSubsystem::Get(UWorld* World)
+{
+	if (World)
+	{
+		return World->GetSubsystem<UAudioAnalyzerSubsystem>();
+	}
+	return nullptr;
+}
+
+TStatId UAudioAnalyzerSubsystem::GetStatId() const
+{
+	RETURN_QUICK_DECLARE_CYCLE_STAT(UAudioAnalyzerSubsystem, STATGROUP_Tickables);
+}
+
+void UAudioAnalyzerSubsystem::Tick(float DeltaTime)
+{
+	// Loop through all analyzers and if they're ready to analyze, do it
+	for (UAudioAnalyzer* Analyzer : AudioAnalyzers)
+	{
+		if (Analyzer->IsReadyForAnalysis())
+		{
+			if (Analyzer->DoAnalysis())
+			{
+				Analyzer->BroadcastResults();
+			}
+		}
+	}
+}
+
+bool UAudioAnalyzerSubsystem::IsTickable() const
+{
+	// As soon as any analyzers are ready, say we're ok to be ticked
+	for (UAudioAnalyzer* Analyzer : AudioAnalyzers)
+	{
+		if (Analyzer->IsReadyForAnalysis())
+		{
+			return true;
+		}
+	}
+
+	return false;
+}
+
+void UAudioAnalyzerSubsystem::RegisterAudioAnalyzer(UAudioAnalyzer* InAnalyzer)
+{
+	AudioAnalyzers.AddUnique(InAnalyzer);
+}
+
+void UAudioAnalyzerSubsystem::UnregisterAudioAnalyzer(UAudioAnalyzer* InAnalyzer)
+{
+	AudioAnalyzers.Remove(InAnalyzer);
+}
\ No newline at end of file
diff --git a/Engine/Source/Runtime/AudioAnalyzer/Private/IAudioAnalyzerNRTInterface.cpp b/Engine/Source/Runtime/AudioAnalyzer/Private/IAudioAnalyzerNRTInterface.cpp
deleted file mode 100644
diff --git a/Engine/Source/Runtime/AudioAnalyzer/Public/AudioAnalyzerFacade.h b/Engine/Source/Runtime/AudioAnalyzer/Public/AudioAnalyzerFacade.h
new file mode 100644
--- /dev/null
+++ b/Engine/Source/Runtime/AudioAnalyzer/Public/AudioAnalyzerFacade.h
@@ -0,0 +1,38 @@
+// Copyright Epic Games, Inc. All Rights Reserved.
+
+#pragma once
+
+#include "CoreMinimal.h"
+#include "IAudioAnalyzerInterface.h"
+
+namespace Audio
+{
+
+	IAnalyzerFactory* GetAnalyzerFactory(FName InFactoryName);
+
+	/** FAnalyzer
+	 *
+	 * FAnalyzer provides a simplified interface for performing analysis on audio buffers from the same AudioAnalyzer.
+	 */
+	class AUDIOANALYZER_API FAnalyzerFacade
+	{
+		public:
+			/**
+			 * Create an FAnalyzerBatch with the analyzer settings and factory name.
+			 */
+			FAnalyzerFacade(TUniquePtr<IAnalyzerSettings> InSettings, IAnalyzerFactory* InFactory);
+
+			/**
+			 * Analyze the audio buffer.
+			 */
+			TUniquePtr<IAnalyzerResult> AnalyzeAudioBuffer(const TArray<float>& InAudioBuffer, int32 InNumChannels, float InSampleRate);
+
+		private:
+
+			TUniquePtr<IAnalyzerSettings> Settings;
+			TUniquePtr<IAnalyzerWorker> Worker;
+
+			IAnalyzerFactory* Factory;
+			FName FactoryName;
+	};
+}
diff --git a/Engine/Source/Runtime/AudioAnalyzer/Public/AudioAnalyzerNRTFacade.h b/Engine/Source/Runtime/AudioAnalyzer/Public/AudioAnalyzerNRTFacade.h
--- a/Engine/Source/Runtime/AudioAnalyzer/Public/AudioAnalyzerNRTFacade.h
+++ b/Engine/Source/Runtime/AudioAnalyzer/Public/AudioAnalyzerNRTFacade.h
@@ -10,18 +10,18 @@ namespace Audio
 
 	IAnalyzerNRTFactory* GetAnalyzerNRTFactory(FName InFactoryName);
 
-	/** FAnalyzerNRTBatch
+	/** FAnalyzerNRTFacade
 	 *
-	 * FAnalyzerNRTBatch provides a simplified interface for running 
+	 * FAnalyzerNRTFacade provides a simplified interface for running 
 	 * analyzer factories over complete audio resources.
 	 */
-	class AUDIOANALYZER_API FAnalyzerNRTBatch
+	class AUDIOANALYZER_API FAnalyzerNRTFacade
 	{
 		public:
 			/**
-			 * Create an FAnalyzerNRTBatch with the analyzer settings and factory name.
+			 * Create an FAnalyzerNRTFacade with the analyzer settings and factory name.
 			 */
-			FAnalyzerNRTBatch(TUniquePtr<IAnalyzerNRTSettings> InSettings, const FName& InFactoryName);
+			FAnalyzerNRTFacade(TUniquePtr<IAnalyzerNRTSettings> InSettings, const FName& InFactoryName);
 
 			/**
 			 * Analyze an entire PCM16 encoded audio object.  Audio for the entire sound should be contained within InRawWaveData.
diff --git a/Engine/Source/Runtime/AudioAnalyzer/Public/IAudioAnalyzerInterface.h b/Engine/Source/Runtime/AudioAnalyzer/Public/IAudioAnalyzerInterface.h
new file mode 100644
--- /dev/null
+++ b/Engine/Source/Runtime/AudioAnalyzer/Public/IAudioAnalyzerInterface.h
@@ -0,0 +1,96 @@
+// Copyright Epic Games, Inc. All Rights Reserved.
+
+#pragma once
+
+#include "CoreMinimal.h"
+#include "Features/IModularFeature.h"
+
+namespace Audio
+{
+	/** FAnalyzerParameters
+	 *
+	 * These parameters are pass to an IAnalyzerFactory when creating
+	 * a new IAnalyzerWorker or IAnalyzerResult
+	 */
+	struct FAnalyzerParameters
+	{
+	public:
+		int32 SampleRate;
+		int32 NumChannels;
+
+		FAnalyzerParameters(int32 InSampleRate = 0, int32 InNumChannels = 0)
+			: SampleRate(InSampleRate)
+			, NumChannels(InNumChannels)
+		{}
+	};
+
+	/** IAnalyzerSettings
+	 *
+	 * This interface defines the required methods for real-time
+	 * analyzer settings.
+	 */
+	class IAnalyzerSettings
+	{
+	public:
+		virtual ~IAnalyzerSettings() {};
+	};
+
+	/** IAnalyzerResult
+	 *
+	 * This interface defines the required methods for real-time
+	 * analyzer results.
+	 */
+	class IAnalyzerResult
+	{
+	public:
+		virtual ~IAnalyzerResult() {};
+	};
+
+	// An audio analyzer worker
+	class IAnalyzerWorker
+	{
+	public:
+		virtual ~IAnalyzerWorker() {};
+
+		virtual void Analyze(TArrayView<const float> InAudio, IAnalyzerResult* OutResult) = 0;
+	};
+
+	// Audio analyzer factory. Primarily used for creating a worker.
+	class IAnalyzerFactory : public IModularFeature
+	{
+	public:
+		virtual ~IAnalyzerFactory() {};
+
+		// Supplied unique name of IAnalyzerFactory to enable querying of added 
+		// analyzer factories
+		static FName GetModularFeatureName()
+		{
+			static FName AudioExtFeatureName = FName(TEXT("AudioAnalyzerPlugin"));
+			return AudioExtFeatureName;
+		}
+
+		// Name of specific analyzer type.
+		virtual FName GetName() const = 0;
+
+		// Human readable name of analyzer.
+		virtual FString GetTitle() const = 0;
+
+		// Create a new result.
+		virtual TUniquePtr<IAnalyzerResult> NewResult() const = 0;
+
+		// Convenience function to create a new shared result by calling NewResult.
+		template<ESPMode Mode = ESPMode::Fast>
+		TSharedPtr<IAnalyzerResult, Mode> NewResultShared() const
+		{
+			TUniquePtr<Audio::IAnalyzerResult> Result = NewResult();
+
+			return TSharedPtr<Audio::IAnalyzerResult, Mode>(Result.Release());
+		}
+
+		// Create a new worker.
+		virtual TUniquePtr<IAnalyzerWorker> NewWorker(const FAnalyzerParameters& InParams, const IAnalyzerSettings* InSettings) const = 0;
+	};
+}
+
+
+
diff --git a/Engine/Source/Runtime/AudioMixer/Private/AudioMixerBlueprintLibrary.cpp b/Engine/Source/Runtime/AudioMixer/Private/AudioMixerBlueprintLibrary.cpp
--- a/Engine/Source/Runtime/AudioMixer/Private/AudioMixerBlueprintLibrary.cpp
+++ b/Engine/Source/Runtime/AudioMixer/Private/AudioMixerBlueprintLibrary.cpp
@@ -11,38 +11,11 @@
 #include "AudioCompressionSettingsUtils.h"
 #include "Async/Async.h"
 #include "Sound/SoundEffectPreset.h"
+#include "AudioDeviceManager.h"
 
 // This is our global recording task:
 static TUniquePtr<Audio::FAudioRecordingData> RecordingData;
 
-
-static FAudioDevice* GetAudioDeviceFromWorldContext(const UObject* WorldContextObject)
-{
-	UWorld* ThisWorld = GEngine->GetWorldFromContextObject(WorldContextObject, EGetWorldErrorMode::LogAndReturnNull);
-	if (!ThisWorld || !ThisWorld->bAllowAudioPlayback || ThisWorld->GetNetMode() == NM_DedicatedServer)
-	{
-		return nullptr;
-	}
-
-	return ThisWorld->GetAudioDevice().GetAudioDevice();
-}
-
-static Audio::FMixerDevice* GetAudioMixerDeviceFromWorldContext(const UObject* WorldContextObject)
-{
-	if (FAudioDevice* AudioDevice = GetAudioDeviceFromWorldContext(WorldContextObject))
-	{
-		if (!AudioDevice->IsAudioMixerEnabled())
-		{
-			return nullptr;
-		}
-		else
-		{
-			return static_cast<Audio::FMixerDevice*>(AudioDevice);
-		}
-	}
-	return nullptr;
-}
-
 void UAudioMixerBlueprintLibrary::AddMasterSubmixEffect(const UObject* WorldContextObject, USoundEffectSubmixPreset* SubmixEffectPreset)
 {
 	if (!SubmixEffectPreset)
@@ -51,7 +24,7 @@ void UAudioMixerBlueprintLibrary::AddMasterSubmixEffect(const UObject* WorldCont
 		return;
 	}
 
-	if (Audio::FMixerDevice* MixerDevice = GetAudioMixerDeviceFromWorldContext(WorldContextObject))
+	if (Audio::FMixerDevice* MixerDevice = FAudioDeviceManager::GetAudioMixerDeviceFromWorldContext(WorldContextObject))
 	{
 		FSoundEffectSubmixInitData InitData;
 		InitData.SampleRate = MixerDevice->GetSampleRate();
@@ -75,7 +48,7 @@ void UAudioMixerBlueprintLibrary::RemoveMasterSubmixEffect(const UObject* WorldC
 		return;
 	}
 
-	if (Audio::FMixerDevice* MixerDevice = GetAudioMixerDeviceFromWorldContext(WorldContextObject))
+	if (Audio::FMixerDevice* MixerDevice = FAudioDeviceManager::GetAudioMixerDeviceFromWorldContext(WorldContextObject))
 	{
 		// Get the unique id for the preset object on the game thread. Used to refer to the object on audio render thread.
 		uint32 SubmixPresetUniqueId = SubmixEffectPreset->GetUniqueID();
@@ -86,7 +59,7 @@ void UAudioMixerBlueprintLibrary::RemoveMasterSubmixEffect(const UObject* WorldC
 
 void UAudioMixerBlueprintLibrary::ClearMasterSubmixEffects(const UObject* WorldContextObject)
 {
-	if (Audio::FMixerDevice* MixerDevice = GetAudioMixerDeviceFromWorldContext(WorldContextObject))
+	if (Audio::FMixerDevice* MixerDevice = FAudioDeviceManager::GetAudioMixerDeviceFromWorldContext(WorldContextObject))
 	{
 		MixerDevice->ClearMasterSubmixEffects();
 	}
@@ -99,7 +72,7 @@ int32 UAudioMixerBlueprintLibrary::AddSubmixEffect(const UObject* WorldContextOb
 		return 0;
 	}
 
-	if (Audio::FMixerDevice* MixerDevice = GetAudioMixerDeviceFromWorldContext(WorldContextObject))
+	if (Audio::FMixerDevice* MixerDevice = FAudioDeviceManager::GetAudioMixerDeviceFromWorldContext(WorldContextObject))
 	{
 		FSoundEffectSubmixInitData InitData;
 		InitData.SampleRate = MixerDevice->GetSampleRate();
@@ -116,7 +89,7 @@ int32 UAudioMixerBlueprintLibrary::AddSubmixEffect(const UObject* WorldContextOb
 
 void UAudioMixerBlueprintLibrary::RemoveSubmixEffectPreset(const UObject* WorldContextObject, USoundSubmix* InSoundSubmix, USoundEffectSubmixPreset* InSubmixEffectPreset)
 {
-	if (Audio::FMixerDevice* MixerDevice = GetAudioMixerDeviceFromWorldContext(WorldContextObject))
+	if (Audio::FMixerDevice* MixerDevice = FAudioDeviceManager::GetAudioMixerDeviceFromWorldContext(WorldContextObject))
 	{
 		uint32 SubmixPresetUniqueId = InSubmixEffectPreset->GetUniqueID();
 		MixerDevice->RemoveSubmixEffect(InSoundSubmix, SubmixPresetUniqueId);
@@ -125,7 +98,7 @@ void UAudioMixerBlueprintLibrary::RemoveSubmixEffectPreset(const UObject* WorldC
 
 void UAudioMixerBlueprintLibrary::RemoveSubmixEffectPresetAtIndex(const UObject* WorldContextObject, USoundSubmix* InSoundSubmix, int32 SubmixChainIndex)
 {
-	if (Audio::FMixerDevice* MixerDevice = GetAudioMixerDeviceFromWorldContext(WorldContextObject))
+	if (Audio::FMixerDevice* MixerDevice = FAudioDeviceManager::GetAudioMixerDeviceFromWorldContext(WorldContextObject))
 	{
 		MixerDevice->RemoveSubmixEffectAtIndex(InSoundSubmix, SubmixChainIndex);
 	}
@@ -138,7 +111,7 @@ void UAudioMixerBlueprintLibrary::ReplaceSoundEffectSubmix(const UObject* WorldC
 		return;
 	}
 
-	if (Audio::FMixerDevice* MixerDevice = GetAudioMixerDeviceFromWorldContext(WorldContextObject))
+	if (Audio::FMixerDevice* MixerDevice = FAudioDeviceManager::GetAudioMixerDeviceFromWorldContext(WorldContextObject))
 	{
 		FSoundEffectSubmixInitData InitData;
 		InitData.SampleRate = MixerDevice->GetSampleRate();
@@ -152,7 +125,7 @@ void UAudioMixerBlueprintLibrary::ReplaceSoundEffectSubmix(const UObject* WorldC
 
 void UAudioMixerBlueprintLibrary::ClearSubmixEffects(const UObject* WorldContextObject, USoundSubmix* InSoundSubmix)
 {
-	if (Audio::FMixerDevice* MixerDevice = GetAudioMixerDeviceFromWorldContext(WorldContextObject))
+	if (Audio::FMixerDevice* MixerDevice = FAudioDeviceManager::GetAudioMixerDeviceFromWorldContext(WorldContextObject))
 	{
 		MixerDevice->ClearSubmixEffects(InSoundSubmix);
 	}
@@ -160,7 +133,7 @@ void UAudioMixerBlueprintLibrary::ClearSubmixEffects(const UObject* WorldContext
 
 void UAudioMixerBlueprintLibrary::SetSubmixEffectChainOverride(const UObject* WorldContextObject, USoundSubmix* InSoundSubmix, TArray<USoundEffectSubmixPreset*> InSubmixEffectPresetChain, float InFadeTimeSec)
 {
-	if (Audio::FMixerDevice* MixerDevice = GetAudioMixerDeviceFromWorldContext(WorldContextObject))
+	if (Audio::FMixerDevice* MixerDevice = FAudioDeviceManager::GetAudioMixerDeviceFromWorldContext(WorldContextObject))
 	{
 		TArray<FSoundEffectSubmixPtr> NewSubmixEffectPresetChain;
 
@@ -182,7 +155,7 @@ void UAudioMixerBlueprintLibrary::SetSubmixEffectChainOverride(const UObject* Wo
 
 void UAudioMixerBlueprintLibrary::ClearSubmixEffectChainOverride(const UObject* WorldContextObject, USoundSubmix* InSoundSubmix, float InFadeTimeSec)
 {
-	if (Audio::FMixerDevice* MixerDevice = GetAudioMixerDeviceFromWorldContext(WorldContextObject))
+	if (Audio::FMixerDevice* MixerDevice = FAudioDeviceManager::GetAudioMixerDeviceFromWorldContext(WorldContextObject))
 	{
 		MixerDevice->ClearSubmixEffectChainOverride(InSoundSubmix, InFadeTimeSec);
 	}
@@ -190,7 +163,7 @@ void UAudioMixerBlueprintLibrary::ClearSubmixEffectChainOverride(const UObject*
 
 void UAudioMixerBlueprintLibrary::StartRecordingOutput(const UObject* WorldContextObject, float ExpectedDuration, USoundSubmix* SubmixToRecord)
 {
-	if (Audio::FMixerDevice* MixerDevice = GetAudioMixerDeviceFromWorldContext(WorldContextObject))
+	if (Audio::FMixerDevice* MixerDevice = FAudioDeviceManager::GetAudioMixerDeviceFromWorldContext(WorldContextObject))
 	{
 		MixerDevice->StartRecording(SubmixToRecord, ExpectedDuration);
 	}
@@ -207,7 +180,7 @@ USoundWave* UAudioMixerBlueprintLibrary::StopRecordingOutput(const UObject* Worl
 		UE_LOG(LogAudioMixer, Warning, TEXT("Abandoning existing write operation. If you'd like to export multiple submix recordings at the same time, use Start/Finish Recording Submix Output instead."));
 	}
 
-	if (Audio::FMixerDevice* MixerDevice = GetAudioMixerDeviceFromWorldContext(WorldContextObject))
+	if (Audio::FMixerDevice* MixerDevice = FAudioDeviceManager::GetAudioMixerDeviceFromWorldContext(WorldContextObject))
 	{
 		float SampleRate;
 		float ChannelCount;
@@ -262,7 +235,7 @@ USoundWave* UAudioMixerBlueprintLibrary::StopRecordingOutput(const UObject* Worl
 
 void UAudioMixerBlueprintLibrary::PauseRecordingOutput(const UObject* WorldContextObject, USoundSubmix* SubmixToPause)
 {
-	if (Audio::FMixerDevice* MixerDevice = GetAudioMixerDeviceFromWorldContext(WorldContextObject))
+	if (Audio::FMixerDevice* MixerDevice = FAudioDeviceManager::GetAudioMixerDeviceFromWorldContext(WorldContextObject))
 	{
 		MixerDevice->PauseRecording(SubmixToPause);
 	}
@@ -274,7 +247,7 @@ void UAudioMixerBlueprintLibrary::PauseRecordingOutput(const UObject* WorldConte
 
 void UAudioMixerBlueprintLibrary::ResumeRecordingOutput(const UObject* WorldContextObject, USoundSubmix* SubmixToResume)
 {
-	if (Audio::FMixerDevice* MixerDevice = GetAudioMixerDeviceFromWorldContext(WorldContextObject))
+	if (Audio::FMixerDevice* MixerDevice = FAudioDeviceManager::GetAudioMixerDeviceFromWorldContext(WorldContextObject))
 	{
 		MixerDevice->ResumeRecording(SubmixToResume);
 	}
@@ -286,7 +259,7 @@ void UAudioMixerBlueprintLibrary::ResumeRecordingOutput(const UObject* WorldCont
 
 void UAudioMixerBlueprintLibrary::StartAnalyzingOutput(const UObject* WorldContextObject, USoundSubmix* SubmixToAnalyze, EFFTSize FFTSize, EFFTPeakInterpolationMethod InterpolationMethod, EFFTWindowType WindowType, float HopSize, EAudioSpectrumType AudioSpectrumType)
 {
-	if (Audio::FMixerDevice* MixerDevice = GetAudioMixerDeviceFromWorldContext(WorldContextObject))
+	if (Audio::FMixerDevice* MixerDevice = FAudioDeviceManager::GetAudioMixerDeviceFromWorldContext(WorldContextObject))
 	{
 		FSoundSpectrumAnalyzerSettings Settings = USoundSubmix::GetSpectrumAnalyzerSettings(FFTSize, InterpolationMethod, WindowType, HopSize, AudioSpectrumType);
 		MixerDevice->StartSpectrumAnalysis(SubmixToAnalyze, Settings);
@@ -299,7 +272,7 @@ void UAudioMixerBlueprintLibrary::StartAnalyzingOutput(const UObject* WorldConte
 
 void UAudioMixerBlueprintLibrary::StopAnalyzingOutput(const UObject* WorldContextObject, USoundSubmix* SubmixToStopAnalyzing)
 {
-	if (Audio::FMixerDevice* MixerDevice = GetAudioMixerDeviceFromWorldContext(WorldContextObject))
+	if (Audio::FMixerDevice* MixerDevice = FAudioDeviceManager::GetAudioMixerDeviceFromWorldContext(WorldContextObject))
 	{
 		MixerDevice->StopSpectrumAnalysis(SubmixToStopAnalyzing);
 	}
@@ -417,7 +390,7 @@ TArray<FSoundSubmixSpectralAnalysisBandSettings> UAudioMixerBlueprintLibrary::Ma
 
 void UAudioMixerBlueprintLibrary::GetMagnitudeForFrequencies(const UObject* WorldContextObject, const TArray<float>& Frequencies, TArray<float>& Magnitudes, USoundSubmix* SubmixToAnalyze /*= nullptr*/)
 {
-	if (Audio::FMixerDevice* MixerDevice = GetAudioMixerDeviceFromWorldContext(WorldContextObject))
+	if (Audio::FMixerDevice* MixerDevice = FAudioDeviceManager::GetAudioMixerDeviceFromWorldContext(WorldContextObject))
 	{
 		MixerDevice->GetMagnitudesForFrequencies(SubmixToAnalyze, Frequencies, Magnitudes);
 	}
@@ -429,7 +402,7 @@ void UAudioMixerBlueprintLibrary::GetMagnitudeForFrequencies(const UObject* Worl
 
 void UAudioMixerBlueprintLibrary::GetPhaseForFrequencies(const UObject* WorldContextObject, const TArray<float>& Frequencies, TArray<float>& Phases, USoundSubmix* SubmixToAnalyze /*= nullptr*/)
 {
-	if (Audio::FMixerDevice* MixerDevice = GetAudioMixerDeviceFromWorldContext(WorldContextObject))
+	if (Audio::FMixerDevice* MixerDevice = FAudioDeviceManager::GetAudioMixerDeviceFromWorldContext(WorldContextObject))
 	{
 		MixerDevice->GetPhasesForFrequencies(SubmixToAnalyze, Frequencies, Phases);
 	}
@@ -447,7 +420,7 @@ void UAudioMixerBlueprintLibrary::AddSourceEffectToPresetChain(const UObject* Wo
 		return;
 	}
 
-	if (Audio::FMixerDevice* MixerDevice = GetAudioMixerDeviceFromWorldContext(WorldContextObject))
+	if (Audio::FMixerDevice* MixerDevice = FAudioDeviceManager::GetAudioMixerDeviceFromWorldContext(WorldContextObject))
 	{
 		TArray<FSourceEffectChainEntry> Chain;
 
@@ -471,7 +444,7 @@ void UAudioMixerBlueprintLibrary::RemoveSourceEffectFromPresetChain(const UObjec
 		return;
 	}
 
-	if (Audio::FMixerDevice* MixerDevice = GetAudioMixerDeviceFromWorldContext(WorldContextObject))
+	if (Audio::FMixerDevice* MixerDevice = FAudioDeviceManager::GetAudioMixerDeviceFromWorldContext(WorldContextObject))
 	{
 		TArray<FSourceEffectChainEntry> Chain;
 
@@ -500,7 +473,7 @@ void UAudioMixerBlueprintLibrary::SetBypassSourceEffectChainEntry(const UObject*
 		return;
 	}
 
-	if (Audio::FMixerDevice* MixerDevice = GetAudioMixerDeviceFromWorldContext(WorldContextObject))
+	if (Audio::FMixerDevice* MixerDevice = FAudioDeviceManager::GetAudioMixerDeviceFromWorldContext(WorldContextObject))
 	{
 		TArray<FSourceEffectChainEntry> Chain;
 
@@ -528,7 +501,7 @@ int32 UAudioMixerBlueprintLibrary::GetNumberOfEntriesInSourceEffectChain(const U
 		return 0;
 	}
 
-	if (Audio::FMixerDevice* MixerDevice = GetAudioMixerDeviceFromWorldContext(WorldContextObject))
+	if (Audio::FMixerDevice* MixerDevice = FAudioDeviceManager::GetAudioMixerDeviceFromWorldContext(WorldContextObject))
 	{
 		TArray<FSourceEffectChainEntry> Chain;
 
@@ -592,7 +565,7 @@ float UAudioMixerBlueprintLibrary::TrimAudioCache(float InMegabytesToFree)
 
 void UAudioMixerBlueprintLibrary::StartAudioBus(const UObject* WorldContextObject, UAudioBus* AudioBus)
 {
-	if (Audio::FMixerDevice* MixerDevice = GetAudioMixerDeviceFromWorldContext(WorldContextObject))
+	if (Audio::FMixerDevice* MixerDevice = FAudioDeviceManager::GetAudioMixerDeviceFromWorldContext(WorldContextObject))
 	{
 		uint32 AudioBusId = AudioBus->GetUniqueID();
 		int32 NumChannels = (int32)AudioBus->AudioBusChannels + 1;
@@ -606,7 +579,7 @@ void UAudioMixerBlueprintLibrary::StartAudioBus(const UObject* WorldContextObjec
 
 void UAudioMixerBlueprintLibrary::StopAudioBus(const UObject* WorldContextObject, UAudioBus* AudioBus)
 {
-	if (Audio::FMixerDevice* MixerDevice = GetAudioMixerDeviceFromWorldContext(WorldContextObject))
+	if (Audio::FMixerDevice* MixerDevice = FAudioDeviceManager::GetAudioMixerDeviceFromWorldContext(WorldContextObject))
 	{
 		uint32 AudioBusId = AudioBus->GetUniqueID();
 		MixerDevice->StopAudioBus(AudioBusId);
@@ -619,7 +592,7 @@ void UAudioMixerBlueprintLibrary::StopAudioBus(const UObject* WorldContextObject
 
 bool UAudioMixerBlueprintLibrary::IsAudioBusActive(const UObject* WorldContextObject, UAudioBus* AudioBus)
 {
-	if (Audio::FMixerDevice* MixerDevice = GetAudioMixerDeviceFromWorldContext(WorldContextObject))
+	if (Audio::FMixerDevice* MixerDevice = FAudioDeviceManager::GetAudioMixerDeviceFromWorldContext(WorldContextObject))
 	{
 		uint32 AudioBusId = AudioBus->GetUniqueID();
 		return MixerDevice->IsAudioBusActive(AudioBusId);
diff --git a/Engine/Source/Runtime/AudioMixer/Private/AudioMixerBus.cpp b/Engine/Source/Runtime/AudioMixer/Private/AudioMixerBus.cpp
--- a/Engine/Source/Runtime/AudioMixer/Private/AudioMixerBus.cpp
+++ b/Engine/Source/Runtime/AudioMixer/Private/AudioMixerBus.cpp
@@ -14,6 +14,9 @@ namespace Audio
 		, bIsAutomatic(bInIsAutomatic)
 	{
 		SetNumOutputChannels(NumChannels);
+
+		// Make a patch bus input to push mixed audio sent to audio bus from source manager
+		AudioBusInput = PatchMixerSplitter.AddNewInput(4096, 1.0f);
 	}
 
 	void FMixerAudioBus::SetNumOutputChannels(int32 InNumOutputChannels)
@@ -122,26 +125,9 @@ namespace Audio
 					const int32 NumOutputFrames = SourceManager->GetNumOutputFrames();
 					const int32 NumSourceSamples = NumSourceChannels * NumOutputFrames;
 
-					// If source channels are 1 but the bus is 2 channels, we need to up-mix
-					if (NumSourceChannels == 1 && NumChannels == 2)
-					{
-						int32 BusSampleIndex = 0;
-						for (int32 SourceSampleIndex = 0; SourceSampleIndex < NumSourceSamples; ++SourceSampleIndex)
-						{
-							// Take half the source sample to up-mix to stereo
-							const float SourceSample = 0.5f * AudioBusSend.SendLevel * SourceBufferPtr[SourceSampleIndex];
-
-							// Mix in the source sample
-							for (int32 BusChannel = 0; BusChannel < NumChannels; ++BusChannel)
-							{
-								BusDataBufferPtr[BusSampleIndex++] += SourceSample;
-							}
-						}
-					}
-					// If the source channels is greater than num channels
-					else if (NumSourceChannels > NumChannels)
+					// If source channels are different than bus channels, we need to up-mix or down-mix
+					if (NumSourceChannels != NumChannels)
 					{
-						check(NumChannels == 1 || NumChannels == 2);
 						Audio::AlignedFloatBuffer ChannelMap;
 						SourceManager->Get2DChannelMap(AudioBusSend.SourceId, NumChannels, ChannelMap);
 
@@ -150,77 +136,41 @@ namespace Audio
 						Audio::DownmixBuffer(NumSourceChannels, NumChannels, SourceBufferPtr, DownmixedBuffer.GetData(), NumOutputFrames, ChannelMap.GetData());
 						Audio::MixInBufferFast(DownmixedBuffer.GetData(), BusDataBufferPtr, DownmixedBuffer.Num(), AudioBusSend.SendLevel);
 					}
-					// If they're the same channels, just mix it in
-					else if (ensureMsgf(NumSourceChannels == NumChannels, TEXT("NumSourceChannels=%d, NumChannels=%d"), NumSourceChannels, NumChannels))
+					else
 					{
 						Audio::MixInBufferFast(SourceBufferPtr, BusDataBufferPtr, NumOutputFrames * NumChannels, AudioBusSend.SendLevel);
 					}
 
 					// Push the mixed data to the patch splitter
-					PatchSplitter.PushAudio(BusDataBufferPtr, NumOutputFrames * NumChannels);
-				}				
+					AudioBusInput.PushAudio(BusDataBufferPtr, NumOutputFrames * NumChannels);
+
+					// Process the audio in the patch mixer splitter
+					PatchMixerSplitter.ProcessAudio();
+				}
 			}
 		}
 	}
 
-	void FMixerAudioBus::CopyCurrentBuffer(Audio::AlignedFloatBuffer& OutBuffer, int32 InNumFrames, int32 InNumOutputChannels, EMonoChannelUpmixMethod InMixMethod) const
+	void FMixerAudioBus::CopyCurrentBuffer(Audio::AlignedFloatBuffer& InChannelMap, int32 InNumOutputChannels, AlignedFloatBuffer& OutBuffer, int32 NumOutputFrames) const
 	{
 		const float* RESTRICT CurrentBuffer = GetCurrentBusBuffer();
 
-		if (InNumOutputChannels == 2 && NumChannels == 1)
-		{
-			switch (InMixMethod)
-			{
-			case EMonoChannelUpmixMethod::EqualPower:
-			{
-				static constexpr float Sqrt2Over2Gains[2] = { 0.707106781f, 0.707106781f };
-				MixMonoTo2ChannelsFast(CurrentBuffer, OutBuffer.GetData(), InNumFrames, Sqrt2Over2Gains);
-			}
-			break;
+		check(NumChannels != InNumOutputChannels);
+		Audio::AlignedFloatBuffer DownmixedBuffer;
+		DownmixedBuffer.AddUninitialized(NumOutputFrames * NumChannels);
 
-			case EMonoChannelUpmixMethod::Linear:
-			{
-				static constexpr float HalfGains[2] = { 0.5f, 0.5f };
-				MixMonoTo2ChannelsFast(CurrentBuffer, OutBuffer.GetData(), InNumFrames, HalfGains);
-			}
-			break;
+		Audio::DownmixBuffer(NumChannels, InNumOutputChannels, CurrentBuffer, DownmixedBuffer.GetData(), NumOutputFrames, InChannelMap.GetData());
 
-			case EMonoChannelUpmixMethod::FullVolume:
-			default:
-				MixMonoTo2ChannelsFast(CurrentBuffer, OutBuffer.GetData(), InNumFrames);
-				break;
-			}
-		}
-		else if (InNumOutputChannels == 1 && NumChannels == 2)
-		{
-			BufferSum2ChannelToMonoFast(CurrentBuffer, OutBuffer.GetData(), InNumFrames);
-			switch (InMixMethod)
-			{
-			case EMonoChannelUpmixMethod::EqualPower:
-			{
-				static constexpr float Sqrt2Gain = 1.41421356f;
-				MultiplyBufferByConstantInPlace(OutBuffer, Sqrt2Gain);
-			}
-			break;
+		Audio::MixInBufferFast(DownmixedBuffer.GetData(), OutBuffer.GetData(), DownmixedBuffer.Num(), 1.0f);
+	}
 
-			case EMonoChannelUpmixMethod::Linear:
-			{
-				static constexpr float DoubleGain = 2.0f;
-				MultiplyBufferByConstantInPlace(OutBuffer, DoubleGain);
-			}
-			break;
+	void FMixerAudioBus::CopyCurrentBuffer(int32 InNumOutputChannels, AlignedFloatBuffer& OutBuffer, int32 NumOutputFrames) const
+	{
+		const float* RESTRICT CurrentBuffer = GetCurrentBusBuffer();
 
-			case EMonoChannelUpmixMethod::FullVolume:
-			default:
-				break;
-			}
-		}
-		// Copy into the pre-distance attenuation buffer ptr as channel count should match
-		else
-		{
-			check(InNumOutputChannels == NumChannels);
-			FMemory::Memcpy(OutBuffer.GetData(), CurrentBuffer, sizeof(float) * InNumFrames * InNumOutputChannels);
-		}
+		check(NumChannels == InNumOutputChannels);
+
+		FMemory::Memcpy(OutBuffer.GetData(), CurrentBuffer, sizeof(float) * NumOutputFrames * InNumOutputChannels);
 	}
 
 	const float* FMixerAudioBus::GetCurrentBusBuffer() const
@@ -233,11 +183,22 @@ namespace Audio
 		return MixedSourceData[!CurrentBufferIndex].GetData();
 	}
 
-	FPatchOutputStrongPtr FMixerAudioBus::AddNewPatch(int32 MaxLatencyInSamples, float InGain)
+	void FMixerAudioBus::AddNewPatchOutput(FPatchOutputStrongPtr& InPatchOutputStrongPtr)
+	{
+		PatchMixerSplitter.AddNewOutput(InPatchOutputStrongPtr);
+	}
+
+	void FMixerAudioBus::AddNewPatchInput(FPatchInput& InPatchInput)
 	{
-		return PatchSplitter.AddNewPatch(MaxLatencyInSamples, InGain);
+		return PatchMixerSplitter.AddNewInput(InPatchInput);
 	}
 
+	void FMixerAudioBus::RemovePatchInput(const FPatchInput& PatchInput)
+	{
+		return PatchMixerSplitter.RemovePatch(PatchInput);
+	}
+
+
 
 }
 
diff --git a/Engine/Source/Runtime/AudioMixer/Private/AudioMixerDevice.cpp b/Engine/Source/Runtime/AudioMixer/Private/AudioMixerDevice.cpp
--- a/Engine/Source/Runtime/AudioMixer/Private/AudioMixerDevice.cpp
+++ b/Engine/Source/Runtime/AudioMixer/Private/AudioMixerDevice.cpp
@@ -2282,7 +2282,12 @@ namespace Audio
 				return;
 			}
 
-			ActiveAudioBuses_GameThread.Add(InAudioBusId);
+			FActiveBusData BusData;
+			BusData.BusId = InAudioBusId;
+			BusData.NumChannels = InNumChannels;
+			BusData.bIsAutomatic = bInIsAutomatic;
+
+			ActiveAudioBuses_GameThread.Add(InAudioBusId, BusData);
 
 			FAudioThread::RunCommandOnAudioThread([this, InAudioBusId, InNumChannels, bInIsAutomatic]()
 			{
@@ -2320,16 +2325,33 @@ namespace Audio
 		return SourceManager->IsAudioBusActive(InAudioBusId);
 	}
 
-	FPatchOutputStrongPtr FMixerDevice::AddPatchForAudioBus(uint32 InAudioBusId, float InPatchGain)
+	Audio::FPatchOutputStrongPtr FMixerDevice::AddPatchForAudioBus(uint32 InAudioBusId, float InPatchGain)
 	{
-		if (ensure(!IsInGameThread() && !IsInAudioThread()))
+		FPatchOutputStrongPtr StrongOutputPtr;
+		if (IsInGameThread())
 		{
-			return SourceManager->AddPatchForAudioBus(InAudioBusId, InPatchGain);
+			FActiveBusData* BusData = ActiveAudioBuses_GameThread.Find(InAudioBusId);
+			if (BusData)
+			{
+				int32 NumOutputFrames = SourceManager->GetNumOutputFrames();
+				int32 BusNumChannels = BusData->NumChannels;
+				StrongOutputPtr = MakeShareable(new FPatchOutput(NumOutputFrames * BusData->NumChannels, InPatchGain));
+
+				FAudioThread::RunCommandOnAudioThread([this, InAudioBusId, StrongOutputPtr]() mutable
+				{
+					SourceManager->AddPatchOutputForAudioBus(InAudioBusId, StrongOutputPtr);
+				});
+			}
+			else
+			{
+				UE_LOG(LogAudioMixer, Warning, TEXT("Unable to add a patch output for audio bus because audio bus id '%d' is not active."), InAudioBusId);
+			}
 		}
 		else
 		{
-			return nullptr;
+			UE_LOG(LogAudioMixer, Warning, TEXT("AddPatchForAudioBus can only be called from the game thread."));
 		}
+		return StrongOutputPtr;
 	}
 
 	Audio::FPatchOutputStrongPtr FMixerDevice::AddPatchForSubmix(uint32 InObjectId, float InPatchGain)
diff --git a/Engine/Source/Runtime/AudioMixer/Private/AudioMixerSource.cpp b/Engine/Source/Runtime/AudioMixer/Private/AudioMixerSource.cpp
--- a/Engine/Source/Runtime/AudioMixer/Private/AudioMixerSource.cpp
+++ b/Engine/Source/Runtime/AudioMixer/Private/AudioMixerSource.cpp
@@ -405,10 +405,12 @@ namespace Audio
 					if (SoundSourceBus->AudioBus)
 					{
 						InitParams.AudioBusId = SoundSourceBus->AudioBus->GetUniqueID();
+						InitParams.AudioBusChannels = (int32)SoundSourceBus->AudioBus->GetNumChannels();
 					}
 					else
 					{
 						InitParams.AudioBusId = WaveData->GetUniqueID();
+						InitParams.AudioBusChannels = WaveData->NumChannels;
 					}
 
 					if (!WaveData->IsLooping())
@@ -574,11 +576,12 @@ namespace Audio
 			for (FSoundSourceBusSendInfo& SendInfo : WaveInstance->BusSends[BusSendType])
 			{
 				// Avoid redoing duplicate code for sending audio to source bus or audio bus. Most of it is the same other than the bus id.
-				auto SetupBusSend = [this](TArray<FInitAudioBusSend>* AudioBusSends, const FSoundSourceBusSendInfo& InSendInfo, int32 InBusSendType, uint32 InBusId)
+				auto SetupBusSend = [this](TArray<FInitAudioBusSend>* AudioBusSends, const FSoundSourceBusSendInfo& InSendInfo, int32 InBusSendType, uint32 InBusId, int32 InBusChannels)
 				{
 					FInitAudioBusSend BusSend;
 					BusSend.AudioBusId = InBusId;
 					BusSend.SendLevel = InSendInfo.SendLevel;
+					BusSend.BusChannels = InBusChannels;
 
 					if (AudioBusSends)
 					{
@@ -622,29 +625,33 @@ namespace Audio
 				if (SendInfo.SoundSourceBus)
 				{						
 					uint32 BusId;
+					int32 BusChannels;
 
 					// Either use the bus id of the source bus's audio bus id if it was specified
 					if (SendInfo.SoundSourceBus->AudioBus)
 					{
 						BusId = SendInfo.SoundSourceBus->AudioBus->GetUniqueID();
+						BusChannels = (int32)SendInfo.SoundSourceBus->AudioBus->GetNumChannels();
 					}
 					else
 					{
 						// otherwise, use the id of the source bus itself (for an automatic source bus)
 						BusId = SendInfo.SoundSourceBus->GetUniqueID();
+						BusChannels = SendInfo.SoundSourceBus->NumChannels;
 					}
 
 					// Call lambda w/ the correctly derived bus id
-					SetupBusSend(OutAudioBusSends, SendInfo, BusSendType, BusId);
+					SetupBusSend(OutAudioBusSends, SendInfo, BusSendType, BusId, BusChannels);
 				}
 
 				if (SendInfo.AudioBus)
 				{
 					// Only need to send audio to just the specified audio bus
 					uint32 BusId = SendInfo.AudioBus->GetUniqueID();
+					int32 BusChannels = (int32)SendInfo.AudioBus->AudioBusChannels + 1;
 
 					// Note we will be sending audio to both the specified source bus and the audio bus with the same send level
-					SetupBusSend(OutAudioBusSends, SendInfo, BusSendType, BusId);
+					SetupBusSend(OutAudioBusSends, SendInfo, BusSendType, BusId, BusChannels);
 				}
 			}
 		}
diff --git a/Engine/Source/Runtime/AudioMixer/Private/AudioMixerSourceManager.cpp b/Engine/Source/Runtime/AudioMixer/Private/AudioMixerSourceManager.cpp
--- a/Engine/Source/Runtime/AudioMixer/Private/AudioMixerSourceManager.cpp
+++ b/Engine/Source/Runtime/AudioMixer/Private/AudioMixerSourceManager.cpp
@@ -852,7 +852,7 @@ namespace Audio
 					else
 					{
 						// If the bus is not registered, make a new entry. This will default to an automatic audio bus until explicitly made manual later.
-						TSharedPtr<FMixerAudioBus> NewAudioBus = TSharedPtr<FMixerAudioBus>(new FMixerAudioBus(this, true, InitParams.NumInputChannels));
+						TSharedPtr<FMixerAudioBus> NewAudioBus = TSharedPtr<FMixerAudioBus>(new FMixerAudioBus(this, true, InitParams.AudioBusChannels));
 						NewAudioBus->AddInstanceId(SourceId, InitParams.NumInputChannels);
 
 						AudioBuses.Add(InitParams.AudioBusId, NewAudioBus);
@@ -881,7 +881,7 @@ namespace Audio
 					else
 					{
 						// If the bus is not registered, make a new entry. This will default to an automatic audio bus until explicitly made manual later.
-						TSharedPtr<FMixerAudioBus> NewAudioBus(new FMixerAudioBus(this, true, FMath::Min(2, InitParams.NumInputChannels)));
+						TSharedPtr<FMixerAudioBus> NewAudioBus(new FMixerAudioBus(this, true, AudioBusSend.BusChannels));
 
 						// Add a send to it. This will not have a bus instance id (i.e. won't output audio), but 
 						// we register the send anyway in the event that this bus does play, we'll know to send this
@@ -1066,15 +1066,28 @@ namespace Audio
 		return AudioBusIds_AudioThread.Contains(InAudioBusId);
 	}
 
-	FPatchOutputStrongPtr FMixerSourceManager::AddPatchForAudioBus(uint32 InAudioBusId, float PatchGain)
+	void FMixerSourceManager::AddPatchOutputForAudioBus(uint32 InAudioBusId, FPatchOutputStrongPtr& InPatchOutputStrongPtr)
 	{
-		AUDIO_MIXER_CHECK_AUDIO_PLAT_THREAD(MixerDevice);
-		TSharedPtr<FMixerAudioBus> AudioBusPtr = AudioBuses.FindRef(InAudioBusId);
-		if (AudioBusPtr.IsValid())
+		AudioMixerThreadCommand([this, InAudioBusId, InPatchOutputStrongPtr]() mutable
 		{
-			return AudioBusPtr->AddNewPatch(NumOutputFrames * AudioBusPtr->GetNumChannels(), PatchGain);
-		}
-		return nullptr;
+			TSharedPtr<FMixerAudioBus> AudioBusPtr = AudioBuses.FindRef(InAudioBusId);
+			if (AudioBusPtr.IsValid())
+			{
+				AudioBusPtr->AddNewPatchOutput(InPatchOutputStrongPtr);
+			}
+		});
+	}
+
+	void FMixerSourceManager::AddPatchInputForAudioBus(uint32 InAudioBusId, FPatchInput& InPatchInput)
+	{
+		AudioMixerThreadCommand([this, InAudioBusId, InPatchInput]() mutable
+		{
+			TSharedPtr<FMixerAudioBus> AudioBusPtr = AudioBuses.FindRef(InAudioBusId);
+			if (AudioBusPtr.IsValid())
+			{
+				AudioBusPtr->AddNewPatchInput(InPatchInput);
+			}
+		});
 	}
 
 	void FMixerSourceManager::Play(const int32 SourceId)
@@ -1787,7 +1800,19 @@ namespace Audio
 					}
 
 					SourceInfo.NumFramesPlayed += NumFramesPlayed;
-					AudioBusPtr->CopyCurrentBuffer(SourceInfo.PreDistanceAttenuationBuffer, NumFramesPlayed, SourceInfo.NumInputChannels);
+
+					// Retrieve the channel map of going from the audio bus channel count to the source channel count since they may not match
+					int32 NumAudioBusChannels = AudioBusPtr->GetNumChannels();
+					if (NumAudioBusChannels != SourceInfo.NumInputChannels)
+					{
+						Audio::AlignedFloatBuffer ChannelMap;
+						MixerDevice->Get2DChannelMap(SourceInfo.bIsVorbis, AudioBusPtr->GetNumChannels(), SourceInfo.NumInputChannels, SourceInfo.bIsCenterChannelOnly, ChannelMap);
+						AudioBusPtr->CopyCurrentBuffer(ChannelMap, SourceInfo.NumInputChannels, SourceInfo.PreDistanceAttenuationBuffer, NumFramesPlayed);
+					}
+					else
+					{
+						AudioBusPtr->CopyCurrentBuffer(SourceInfo.NumInputChannels, SourceInfo.PreDistanceAttenuationBuffer, NumFramesPlayed);
+					}
 				}
 			}
 			else
@@ -2646,6 +2671,7 @@ namespace Audio
 
 	const float* FMixerSourceManager::GetPreviousAudioBusBuffer(const int32 AudioBusId) const
 	{
+		// This is only called from within a scope-lock
 		const TSharedPtr<FMixerAudioBus> AudioBusPtr = AudioBuses.FindRef(AudioBusId);
 		if (AudioBusPtr.IsValid())
 		{
diff --git a/Engine/Source/Runtime/AudioMixer/Private/AudioMixerSourceManager.h b/Engine/Source/Runtime/AudioMixer/Private/AudioMixerSourceManager.h
--- a/Engine/Source/Runtime/AudioMixer/Private/AudioMixerSourceManager.h
+++ b/Engine/Source/Runtime/AudioMixer/Private/AudioMixerSourceManager.h
@@ -99,6 +99,7 @@ namespace Audio
 	{
 		uint32 AudioBusId = INDEX_NONE;
 		float SendLevel = 0.0f;
+		int32 BusChannels = 0;
 	};
 
 	struct FMixerSourceVoiceInitParams
@@ -108,6 +109,7 @@ namespace Audio
 		TArray<FMixerSourceSubmixSend> SubmixSends;
 		TArray<FInitAudioBusSend> AudioBusSends[(int32)EBusSendType::Count];
 		uint32 AudioBusId = INDEX_NONE;
+		int32 AudioBusChannels = 0;
 		float SourceBusDuration = 0.0f;
 		uint32 SourceEffectChainId = INDEX_NONE;
 		TArray<FSourceEffectChainEntry> SourceEffectChain;
@@ -163,11 +165,20 @@ namespace Audio
 		void ReleaseSourceId(const int32 SourceId);
 		void InitSource(const int32 SourceId, const FMixerSourceVoiceInitParams& InitParams);
 
-		// Creates an audio bus manually. Returns an audio bus Id.
+		// Creates and starts an audio bus manually.
 		void StartAudioBus(uint32 InAudioBusId, int32 InNumChannels, bool bInIsAutomatic);
+
+		// Stops an audio bus manually
 		void StopAudioBus(uint32 InAudioBusId);
+
+		// Queries if an audio bus is active
 		bool IsAudioBusActive(uint32 InAudioBusId);
-		FPatchOutputStrongPtr AddPatchForAudioBus(uint32 InAudioBusId, float PatchGain);
+
+		// Adds a patch output for an audio bus
+		void AddPatchOutputForAudioBus(uint32 InAudioBusId, FPatchOutputStrongPtr& InPatchOutputStrongPtr);
+
+		// Adds a patch input for an audio bus
+		void AddPatchInputForAudioBus(uint32 InAudioBusId, FPatchInput& InPatchInput);
 
 		void Play(const int32 SourceId);
 		void Stop(const int32 SourceId);
diff --git a/Engine/Source/Runtime/AudioMixer/Public/AudioMixerBus.h b/Engine/Source/Runtime/AudioMixer/Public/AudioMixerBus.h
--- a/Engine/Source/Runtime/AudioMixer/Public/AudioMixerBus.h
+++ b/Engine/Source/Runtime/AudioMixer/Public/AudioMixerBus.h
@@ -22,8 +22,14 @@ namespace Audio
 	class FMixerAudioBus
 	{
 	public:
-		// Allow anybody to add a patch to retrieve the audio from any thread.
-		FPatchOutputStrongPtr AddNewPatch(int32 MaxLatencyInSamples, float InGain);
+		// Allow anybody to add a pre-existing patch output object to the audio bus
+		void AddNewPatchOutput(FPatchOutputStrongPtr& InPatchOutputStrongPtr);
+
+		// Allow anybody to write audio into this audio bus from any thread.
+		void AddNewPatchInput(FPatchInput& InPatchInput);
+
+		// Allow anybody to write audio into this audio bus from any thread.
+		void RemovePatchInput(const FPatchInput& InPatchInput);
 
 	private:
 
@@ -52,7 +58,7 @@ namespace Audio
 		// Removes the source id from this bus. Returns true if there are no more instances or sends.
 		bool RemoveInstanceId(const int32 InSourceId);
 
-		// Adds a buss end to the bus
+		// Adds a bus send to the bus
 		void AddSend(EBusSendType BusSendType, const FAudioBusSend& InBusSend);
 
 		// Removes the source instance from this bus's send list
@@ -68,7 +74,8 @@ namespace Audio
 		void MixBuffer();
 
 		// Copies the current internal buffer to a provided output buffer. Only supports mono or stereo input/output formats.
-		void CopyCurrentBuffer(AlignedFloatBuffer& OutBuffer, int32 InNumFrames, int32 InNumOutputChannels, EMonoChannelUpmixMethod InMixMethod = EMonoChannelUpmixMethod::EqualPower) const;
+		void CopyCurrentBuffer(Audio::AlignedFloatBuffer& InChannelMap, int32 InNumOutputChannels, AlignedFloatBuffer& OutBuffer, int32 NumOutputFrames) const;
+		void CopyCurrentBuffer(int32 InNumOutputChannels, AlignedFloatBuffer& OutBuffer, int32 NumOutputFrames) const;
 
 		// If this bus was constructed before
 		void SetNumOutputChannels(int32 InNumOutputChannels);
@@ -98,12 +105,15 @@ namespace Audio
 		// Owning soruce manager
 		FMixerSourceManager* SourceManager;
 
-		Audio::FPatchSplitter PatchSplitter;
+		// Multiple places can produce and consume from audio buses
+		Audio::FPatchInput AudioBusInput;
+		Audio::FPatchMixerSplitter PatchMixerSplitter;
 
 		// Was created manually, not via source buses.
 		bool bIsAutomatic;
 
 		friend FMixerSourceManager;
+		friend FMixerSubmix;
 	};
 
 }
diff --git a/Engine/Source/Runtime/AudioMixer/Public/AudioMixerDevice.h b/Engine/Source/Runtime/AudioMixer/Public/AudioMixerDevice.h
--- a/Engine/Source/Runtime/AudioMixer/Public/AudioMixerDevice.h
+++ b/Engine/Source/Runtime/AudioMixer/Public/AudioMixerDevice.h
@@ -255,7 +255,7 @@ namespace Audio
 		void StartAudioBus(uint32 InAudioBusId, int32 InNumChannels, bool bInIsAutomatic);
 		void StopAudioBus(uint32 InAudioBusId);
 		bool IsAudioBusActive(uint32 InAudioBusId);
-		FPatchOutputStrongPtr AddPatchForAudioBus(uint32 InAudioBusId, float InPatchGain);
+		FPatchOutputStrongPtr AddPatchForAudioBus(uint32 InAudioBusId, float InPatchGain = 1.0f);
 
 		// Clock Manager for quantized event handling on Audio Render Thread
 		FQuartzClockManager QuantizedEventClockManager;
@@ -311,8 +311,15 @@ namespace Audio
 
 		TArray<TStrongObjectPtr<UAudioBus>> DefaultAudioBuses;
 
+		struct FActiveBusData
+		{
+			int32 BusId = 0;
+			int32 NumChannels = 0;
+			bool bIsAutomatic = false;
+		};
+
 		// The active audio bus list accessible on the game thread
-		TArray<int32> ActiveAudioBuses_GameThread;
+		TMap<int32, FActiveBusData> ActiveAudioBuses_GameThread;
 
 		/** Ptr to the platform interface, which handles streaming audio to the hardware device. */
 		IAudioMixerPlatformInterface* AudioMixerPlatform;
diff --git a/Engine/Source/Runtime/Engine/Classes/Sound/AudioBus.h b/Engine/Source/Runtime/Engine/Classes/Sound/AudioBus.h
--- a/Engine/Source/Runtime/Engine/Classes/Sound/AudioBus.h
+++ b/Engine/Source/Runtime/Engine/Classes/Sound/AudioBus.h
@@ -12,8 +12,11 @@ class FAudioDevice;
 UENUM(BlueprintType)
 enum class EAudioBusChannels : uint8
 {
-	Mono,
-	Stereo,
+	Mono = 0,
+	Stereo = 1,
+	Quad = 3,
+	FivePointOne = 5,
+	SevenPointOne = 7
 };
 
 // Function to retrieve an audio bus buffer given a handle
@@ -21,7 +24,7 @@ enum class EAudioBusChannels : uint8
 
 // An audio bus is an object which represents an audio patch chord. Audio can be sent to it. It can be sonified using USoundSourceBuses.
 // Instances of the audio bus are created in the audio engine. 
-UCLASS(hidecategories= (Compression, SoundWave, Streaming, Subtitles, Sound, Info, ImportSettings), ClassGroup = Sound, meta = (BlueprintSpawnableComponent))
+UCLASS(ClassGroup = Sound, meta = (BlueprintSpawnableComponent))
 class ENGINE_API UAudioBus : public UObject
 {
 	GENERATED_UCLASS_BODY()
@@ -40,6 +43,9 @@ public:
 	virtual void BeginDestroy() override;
 	//~ End UObject
 
+	// Returns the number of channels of the audio bus in integer format
+	int32 GetNumChannels() const { return (int32)AudioBusChannels + 1; }
+
 	//~ Begin UObject Interface.
 #if WITH_EDITOR
 	virtual void PostEditChangeProperty(FPropertyChangedEvent& PropertyChangedEvent) override;
diff --git a/Engine/Source/Runtime/Engine/Private/AudioBus.cpp b/Engine/Source/Runtime/Engine/Private/AudioBus.cpp
--- a/Engine/Source/Runtime/Engine/Private/AudioBus.cpp
+++ b/Engine/Source/Runtime/Engine/Private/AudioBus.cpp
@@ -87,17 +87,7 @@ void UAudioBus::PostEditChangeProperty(FPropertyChangedEvent& PropertyChangedEve
 			if (Audio::FMixerDevice* MixerDevice = static_cast<Audio::FMixerDevice*>(InDevice))
 			{
 				MixerDevice->StopAudioBus(BusId);
-				switch(NumChannels)
-				{
-					case EAudioBusChannels::Stereo:
-						MixerDevice->StartAudioBus(BusId, 2, false /* bInIsAutomatic */);
-					break;
-
-					case EAudioBusChannels::Mono:
-					default:
-						MixerDevice->StartAudioBus(BusId, 1, false /* bInIsAutomatic */);
-					break;
-				}
+				MixerDevice->StartAudioBus(BusId, (int32)NumChannels + 1, false /* bInIsAutomatic */);
 			}
 		});
 	}
diff --git a/Engine/Source/Runtime/SignalProcessing/Private/MultithreadedPatching.cpp b/Engine/Source/Runtime/SignalProcessing/Private/MultithreadedPatching.cpp
--- a/Engine/Source/Runtime/SignalProcessing/Private/MultithreadedPatching.cpp
+++ b/Engine/Source/Runtime/SignalProcessing/Private/MultithreadedPatching.cpp
@@ -191,25 +191,36 @@ namespace Audio
 	{
 	}
 
-
-	FPatchInput FPatchMixer::AddNewInput(int32 MaxLatencyInSamples, float InGain)
+	FPatchInput FPatchMixer::AddNewInput(int32 InMaxLatencyInSamples, float InGain)
 	{
 		FScopeLock ScopeLock(&PendingNewInputsCriticalSection);
 
-		int32 NewPatchIndex = PendingNewInputs.Emplace(new FPatchOutput(MaxLatencyInSamples, InGain));
+		int32 NewPatchIndex = PendingNewInputs.Emplace(new FPatchOutput(InMaxLatencyInSamples, InGain));
 		return FPatchInput(PendingNewInputs[NewPatchIndex]);
 	}
 
-	void FPatchMixer::RemovePatch(const FPatchInput& PatchInput)
+	void FPatchMixer::AddNewInput(FPatchInput& InPatchInput)
+	{
+		if (!InPatchInput.OutputHandle.IsValid())
+		{
+			return;
+		}
+
+		FScopeLock ScopeLock(&PendingNewInputsCriticalSection);
+
+		PendingNewInputs.Add(InPatchInput.OutputHandle);
+	}
+
+	void FPatchMixer::RemovePatch(const FPatchInput& InPatchInput)
 	{
 		// If the output is already disconnected, early exit.
-		if (!PatchInput.OutputHandle.IsValid())
+		if (!InPatchInput.OutputHandle.IsValid())
 		{
 			return;
 		}
 
 		FScopeLock ScopeLock(&InputDeletionCriticalSection);
-		DisconnectedInputs.Add(PatchInput.OutputHandle->PatchID);
+		DisconnectedInputs.Add(InPatchInput.OutputHandle->PatchID);
 	}
 
 	int32 FPatchMixer::PopAudio(float* OutBuffer, int32 OutNumSamples, bool bUseLatestAudio)
@@ -353,7 +364,7 @@ namespace Audio
 	{
 	}
 
-	Audio::FPatchOutputStrongPtr FPatchSplitter::AddNewPatch(int32 MaxLatencyInSamples, float InGain)
+	FPatchOutputStrongPtr FPatchSplitter::AddNewPatch(int32 MaxLatencyInSamples, float InGain)
 	{
 		// Allocate a new FPatchOutput, then store a weak pointer to it in our PendingOutputs array to be added in our next call to PushAudio.
 		FPatchOutputStrongPtr StrongOutputPtr = MakeShareable(new FPatchOutput(MaxLatencyInSamples * 2, InGain));
@@ -366,6 +377,12 @@ namespace Audio
 		return StrongOutputPtr;
 	}
 
+	void FPatchSplitter::AddNewPatch(FPatchOutputStrongPtr& InPatchOutputStrongPtr)
+	{
+		FScopeLock ScopeLock(&PendingOutputsCriticalSection);
+		PendingOutputs.Emplace(InPatchOutputStrongPtr);
+	}
+
 	int32 FPatchSplitter::Num()
 	{
 		FScopeLock ScopeLock(&ConnectedOutputsCriticalSection);
@@ -448,16 +465,26 @@ namespace Audio
 	{
 	}
 
-	Audio::FPatchOutputStrongPtr FPatchMixerSplitter::AddNewOutput(int32 MaxLatencyInSamples, float InGain)
+	FPatchOutputStrongPtr FPatchMixerSplitter::AddNewOutput(int32 MaxLatencyInSamples, float InGain)
 	{
 		return Splitter.AddNewPatch(MaxLatencyInSamples, InGain);
 	}
 
-	Audio::FPatchInput FPatchMixerSplitter::AddNewInput(int32 MaxLatencyInSamples, float InGain)
+	void FPatchMixerSplitter::AddNewOutput(FPatchOutputStrongPtr& InPatchOutputStrongPtr)
+	{
+		Splitter.AddNewPatch(InPatchOutputStrongPtr);
+	}
+
+	FPatchInput FPatchMixerSplitter::AddNewInput(int32 MaxLatencyInSamples, float InGain)
 	{
 		return Mixer.AddNewInput(MaxLatencyInSamples, InGain);
 	}
 
+	void FPatchMixerSplitter::AddNewInput(FPatchInput& InInput)
+	{
+		Mixer.AddNewInput(InInput);
+	}
+
 	void FPatchMixerSplitter::RemovePatch(const FPatchInput& TapInput)
 	{
 		Mixer.RemovePatch(TapInput);
diff --git a/Engine/Source/Runtime/SignalProcessing/Public/DSP/MultithreadedPatching.h b/Engine/Source/Runtime/SignalProcessing/Public/DSP/MultithreadedPatching.h
--- a/Engine/Source/Runtime/SignalProcessing/Public/DSP/MultithreadedPatching.h
+++ b/Engine/Source/Runtime/SignalProcessing/Public/DSP/MultithreadedPatching.h
@@ -104,7 +104,7 @@ namespace Audio
 		friend class FPatchSplitter;
 
 	private:
-		/** Weak pointer to our destination buffer. */
+		/** Strong pointer to our destination buffer. */
 		FPatchOutputStrongPtr OutputHandle;
 
 		/** Counter of the number of push calls until we check OutputHandle.IsUnique() and possibly clean up the output. */
@@ -125,6 +125,9 @@ namespace Audio
 		/** Adds a new input to the tap collector. Calling this is thread safe, but individual instances of FPatchInput are only safe to be used from one thread. */
 		FPatchInput AddNewInput(int32 MaxLatencyInSamples, float InGain);
 
+		/** Adds an existing patch input to the patch mixer. */
+		void AddNewInput(FPatchInput& InPatchInput);
+
 		/** Removes a tap from the tap collector. Calling this is thread safe, though FPatchOutput will likely not be deleted until the next call of PopAudio. */
 		void RemovePatch(const FPatchInput& TapInput);
 
@@ -183,6 +186,9 @@ namespace Audio
 		 */
 		FPatchOutputStrongPtr AddNewPatch(int32 MaxLatencyInSamples, float InGain);
 
+		/** Adds a new a patch from an existing patch output. */
+		void AddNewPatch(FPatchOutputStrongPtr& InPatchOutputStrongPtr);
+
 		/** This call pushes audio to all outputs connected to this splitter. Only should be called from one thread. */
 		int32 PushAudio(const float* InBuffer, int32 InNumSamples);
 
@@ -205,7 +211,7 @@ namespace Audio
 
 	/**
 	 * This class is used to mix multiple inputs from disparate threads to a single mixdown and deliver that mixdown to multiple outputs.
-	 * This class is MPCMC, but only one thread can and should call ProcessAudio().
+	 * This class is MPMC, but only one thread can and should call ProcessAudio().
 	 */
 	class SIGNALPROCESSING_API FPatchMixerSplitter
 	{
@@ -222,9 +228,15 @@ namespace Audio
 		 */
 		FPatchOutputStrongPtr AddNewOutput(int32 MaxLatencyInSamples, float InGain);
 
+		/** Adds a new a patch from an existing patch output. */
+		void AddNewOutput(FPatchOutputStrongPtr& InPatchOutputStrongPtr);
+
 		/** Adds a new input to the tap collector. Calling this is thread safe, but individual instances of FPatchInput are only safe to be used from one thread. */
 		FPatchInput AddNewInput(int32 MaxLatencyInSamples, float InGain);
 
+		/** Adds a new a patch input from an existing patch input object. */
+		void AddNewInput(FPatchInput& InInput);
+
 		/** Removes a tap from the tap collector. Calling this is thread safe, though FPatchOutput will likely not be deleted until the next call of PopAudio. */
 		void RemovePatch(const FPatchInput& TapInput);
 
