Fix stuck generation sounds that end generating internally (had this in there in earlier review, but accidentally removed)

diff --git a/Engine/Source/Runtime/AudioMixer/Private/AudioMixerSourceManager.cpp b/Engine/Source/Runtime/AudioMixer/Private/AudioMixerSourceManager.cpp
--- a/Engine/Source/Runtime/AudioMixer/Private/AudioMixerSourceManager.cpp
+++ b/Engine/Source/Runtime/AudioMixer/Private/AudioMixerSourceManager.cpp
@@ -1667,6 +1667,7 @@ namespace Audio
 						SourceInfo.bIsPaused = false;
 						SourceInfo.bIsActive = false;
 						SourceInfo.bIsStopping = false;
+						SourceInfo.bIsLastBuffer = true;
 						break;
 					}
 				}
