Static analysis fix (potential nullptr deref)

diff --git a/Engine/Source/Runtime/AudioMixer/Private/AudioMixerBuffer.cpp b/Engine/Source/Runtime/AudioMixer/Private/AudioMixerBuffer.cpp
--- a/Engine/Source/Runtime/AudioMixer/Private/AudioMixerBuffer.cpp
+++ b/Engine/Source/Runtime/AudioMixer/Private/AudioMixerBuffer.cpp
@@ -309,6 +309,12 @@ namespace Audio
 
 	FMixerBuffer* FMixerBuffer::CreateStreamingBuffer(FAudioDevice* AudioDevice, USoundWave* InWave)
 	{
+		if (!InWave)
+		{
+			return nullptr;
+		}
+
+
 		FMixerBuffer* Buffer = new FMixerBuffer(AudioDevice, InWave, EBufferType::Streaming);
 
 		FSoundQualityInfo QualityInfo = { 0 };
