Fixes for source buses and audio buses

diff --git a/Engine/Source/Runtime/AudioMixer/Private/AudioMixerSource.cpp b/Engine/Source/Runtime/AudioMixer/Private/AudioMixerSource.cpp
--- a/Engine/Source/Runtime/AudioMixer/Private/AudioMixerSource.cpp
+++ b/Engine/Source/Runtime/AudioMixer/Private/AudioMixerSource.cpp
@@ -1403,9 +1403,17 @@ namespace Audio
 		if (ActiveSound->HasNewBusSends())
 		{
 			TArray<TTuple<EBusSendType, FSoundSourceBusSendInfo>> NewBusSends = ActiveSound->GetNewBusSends();
-			for (TTuple<EBusSendType, FSoundSourceBusSendInfo>& newSend : NewBusSends)
+			for (TTuple<EBusSendType, FSoundSourceBusSendInfo>& NewSend : NewBusSends)
 			{
-				MixerSourceVoice->SetAudioBusSendInfo(newSend.Key, newSend.Value.SoundSourceBus->GetUniqueID(), newSend.Value.SendLevel);
+				if (NewSend.Value.SoundSourceBus)
+				{
+					MixerSourceVoice->SetAudioBusSendInfo(NewSend.Key, NewSend.Value.SoundSourceBus->GetUniqueID(), NewSend.Value.SendLevel);
+				}
+
+				if (NewSend.Value.AudioBus)
+				{
+					MixerSourceVoice->SetAudioBusSendInfo(NewSend.Key, NewSend.Value.AudioBus->GetUniqueID(), NewSend.Value.SendLevel);
+				}
 			}
 
 			ActiveSound->ResetNewBusSends();
