Fix for out of bounds access in audio renderer for metasounds

diff --git a/Engine/Source/Runtime/AudioMixer/Private/AudioMixerSourceManager.cpp b/Engine/Source/Runtime/AudioMixer/Private/AudioMixerSourceManager.cpp
--- a/Engine/Source/Runtime/AudioMixer/Private/AudioMixerSourceManager.cpp
+++ b/Engine/Source/Runtime/AudioMixer/Private/AudioMixerSourceManager.cpp
@@ -1718,7 +1718,7 @@ namespace Audio
 					SourceInfo.NextFrameValues[Channel] = AudioData[NextSampleIndex + Channel];
 				}
 			}
-			else
+			else if (NextSampleIndex != SourceInfo.CurrentPCMBuffer->AudioData.Num())
 			{
 				for (int32 Channel = 0; Channel < NumChannels; ++Channel)
 				{
