Quartz system improvments:
* Exposed quantization boundary parameter on Transport Reset call
* addition of None as a quantiazition boundary
* few new helper BP functions

diff --git a/Engine/Source/Runtime/AudioMixer/Private/Quartz/AudioMixerClockHandle.cpp b/Engine/Source/Runtime/AudioMixer/Private/Quartz/AudioMixerClockHandle.cpp
--- a/Engine/Source/Runtime/AudioMixer/Private/Quartz/AudioMixerClockHandle.cpp
+++ b/Engine/Source/Runtime/AudioMixer/Private/Quartz/AudioMixerClockHandle.cpp
@@ -60,6 +60,11 @@ UQuartzClockHandle* UQuartzClockHandle::Init(UWorld* InWorldPtr)
 	return this;
 }
 
+void UQuartzClockHandle::StartClock(const UObject* WorldContextObject)
+{
+	ResumeClock(WorldContextObject);
+}
+
 void UQuartzClockHandle::PauseClock(const UObject* WorldContextObject)
 {
 	if (QuartzSubsystem)
@@ -86,15 +91,25 @@ void UQuartzClockHandle::ResumeClock(const UObject* WorldContextObject)
 	}
 }
 
-void UQuartzClockHandle::ResetTransport(const UObject* WorldContextObject, const FOnQuartzCommandEventBP& InDelegate)
+void UQuartzClockHandle::ResetTransport(const UObject* WorldContextObject, FQuartzQuantizationBoundary InQuantizationBoundary, const FOnQuartzCommandEventBP& InDelegate)
 {
 	if (QuartzSubsystem != nullptr)
 	{
-		Audio::FQuartzQuantizedCommandInitInfo Data(QuartzSubsystem->CreateDataForTransportReset(this, InDelegate));
+		Audio::FQuartzQuantizedCommandInitInfo Data(QuartzSubsystem->CreateDataForTransportReset(this, InQuantizationBoundary, InDelegate));
 		QuartzSubsystem->AddCommandToClock(WorldContextObject, Data);
 	}
 }
 
+bool UQuartzClockHandle::IsClockRunning(const UObject* WorldContextObject)
+{
+	if (QuartzSubsystem != nullptr)
+	{
+		return QuartzSubsystem->IsClockRunning(WorldContextObject, CurrentClockId);
+	}
+
+	return false;
+}
+
 void UQuartzClockHandle::SubscribeToQuantizationEvent(const UObject* WorldContextObject, EQuartzCommandQuantization InQuantizationBoundary, const FOnQuartzMetronomeEventBP& OnQuantizationEvent)
 {
 	if (QuartzSubsystem)
@@ -150,8 +165,9 @@ void UQuartzClockHandle::UnsubscribeFromAllTimeDivisions(const UObject* WorldCon
 }
 
 // Metronome Alteration (setters)
-void UQuartzClockHandle::SetMillisecondsPerTick(const UObject* WorldContextObject, const FQuartzQuantizationBoundary& InQuantizationBoundary, const FOnQuartzCommandEventBP& InDelegate, float MillisecondsPerTick)
+void UQuartzClockHandle::SetMillisecondsPerTick(const UObject* WorldContextObject, const FQuartzQuantizationBoundary& InQuantizationBoundary, const FOnQuartzCommandEventBP& InDelegate, UQuartzClockHandle*& ClockHandle, float MillisecondsPerTick)
 {
+	ClockHandle = this;
 	if (MillisecondsPerTick < 0 || FMath::IsNearlyZero(MillisecondsPerTick))
 	{
 		UE_LOG(LogAudioQuartz, Warning, TEXT("Ignoring invalid request on Clock: %s: MillisecondsPerTick was %f"), *this->CurrentClockId.ToString(), MillisecondsPerTick);
@@ -168,8 +184,9 @@ void UQuartzClockHandle::SetMillisecondsPerTick(const UObject* WorldContextObjec
 	}
 }
 
-void UQuartzClockHandle::SetTicksPerSecond(const UObject* WorldContextObject, const FQuartzQuantizationBoundary& InQuantizationBoundary, const FOnQuartzCommandEventBP& InDelegate, float TicksPerSecond)
+void UQuartzClockHandle::SetTicksPerSecond(const UObject* WorldContextObject, const FQuartzQuantizationBoundary& InQuantizationBoundary, const FOnQuartzCommandEventBP& InDelegate, UQuartzClockHandle*& ClockHandle, float TicksPerSecond)
 {
+	ClockHandle = this;
 	if (TicksPerSecond < 0 || FMath::IsNearlyZero(TicksPerSecond))
 	{
 		UE_LOG(LogAudioQuartz, Warning, TEXT("Ignoring invalid request on Clock: %s: TicksPerSecond was %f"), *this->CurrentClockId.ToString(), TicksPerSecond);
@@ -186,8 +203,9 @@ void UQuartzClockHandle::SetTicksPerSecond(const UObject* WorldContextObject, co
 	}
 }
 
-void UQuartzClockHandle::SetSecondsPerTick(const UObject* WorldContextObject, const FQuartzQuantizationBoundary& InQuantizationBoundary, const FOnQuartzCommandEventBP& InDelegate, float SecondsPerTick)
+void UQuartzClockHandle::SetSecondsPerTick(const UObject* WorldContextObject, const FQuartzQuantizationBoundary& InQuantizationBoundary, const FOnQuartzCommandEventBP& InDelegate, UQuartzClockHandle*& ClockHandle, float SecondsPerTick)
 {
+	ClockHandle = this;
 	if (SecondsPerTick < 0 || FMath::IsNearlyZero(SecondsPerTick))
 	{
 		UE_LOG(LogAudioQuartz, Warning, TEXT("Ignoring invalid request on Clock: %s: SecondsPerTick was %f"), *this->CurrentClockId.ToString(), SecondsPerTick);
@@ -204,8 +222,9 @@ void UQuartzClockHandle::SetSecondsPerTick(const UObject* WorldContextObject, co
 	}
 }
 
-void UQuartzClockHandle::SetThirtySecondNotesPerMinute(const UObject* WorldContextObject, const FQuartzQuantizationBoundary& InQuantizationBoundary, const FOnQuartzCommandEventBP& InDelegate, float ThirtySecondsNotesPerMinute)
+void UQuartzClockHandle::SetThirtySecondNotesPerMinute(const UObject* WorldContextObject, const FQuartzQuantizationBoundary& InQuantizationBoundary, const FOnQuartzCommandEventBP& InDelegate, UQuartzClockHandle*& ClockHandle, float ThirtySecondsNotesPerMinute)
 {
+	ClockHandle = this;
 	if (ThirtySecondsNotesPerMinute < 0 || FMath::IsNearlyZero(ThirtySecondsNotesPerMinute))
 	{
 		UE_LOG(LogAudioQuartz, Warning, TEXT("Ignoring invalid request on Clock: %s: ThirtySecondsNotesPerMinute was %f"), *this->CurrentClockId.ToString(), ThirtySecondsNotesPerMinute);
@@ -222,8 +241,9 @@ void UQuartzClockHandle::SetThirtySecondNotesPerMinute(const UObject* WorldConte
 	}
 }
 
-void UQuartzClockHandle::SetBeatsPerMinute(const UObject* WorldContextObject, const FQuartzQuantizationBoundary& InQuantizationBoundary, const FOnQuartzCommandEventBP& InDelegate, float BeatsPerMinute)
+void UQuartzClockHandle::SetBeatsPerMinute(const UObject* WorldContextObject, const FQuartzQuantizationBoundary& InQuantizationBoundary, const FOnQuartzCommandEventBP& InDelegate, UQuartzClockHandle*& ClockHandle, float BeatsPerMinute)
 {
+	ClockHandle = this;
 	if (BeatsPerMinute < 0 || FMath::IsNearlyZero(BeatsPerMinute))
 	{
 		UE_LOG(LogAudioQuartz, Warning, TEXT("Ignoring invalid request on Clock: %s: BeatsPerMinute was %f"), *this->CurrentClockId.ToString(), BeatsPerMinute);
diff --git a/Engine/Source/Runtime/AudioMixer/Private/Quartz/QuartzMetronome.cpp b/Engine/Source/Runtime/AudioMixer/Private/Quartz/QuartzMetronome.cpp
--- a/Engine/Source/Runtime/AudioMixer/Private/Quartz/QuartzMetronome.cpp
+++ b/Engine/Source/Runtime/AudioMixer/Private/Quartz/QuartzMetronome.cpp
@@ -65,7 +65,6 @@ namespace Audio
 						}
 
 						FramesLeftInMusicalDuration[DurationType] += PulseDurations[PulseDurationIndex];
-
 					}
 					while (FramesLeftInMusicalDuration[DurationType] <= 0);
 				}
@@ -126,6 +125,11 @@ namespace Audio
 
 	int32 FQuartzMetronome::GetFramesUntilBoundary(FQuartzQuantizationBoundary InQuantizationBoundary) const
 	{
+		if(!ensure(InQuantizationBoundary.Quantization != EQuartzCommandQuantization::None))
+		{
+			return 0; // Metronome's should not have to deal w/ Quartization == None
+		}
+
 		if (InQuantizationBoundary.Multiplier < 1.0f)
 		{
 			UE_LOG(LogAudioQuartz, Warning, TEXT("Quantizatoin Boundary being clamped to 1.0 (from %f)"), InQuantizationBoundary.Multiplier);
@@ -309,6 +313,12 @@ namespace Audio
 	void FQuartzMetronome::ResetTransport()
 	{
 		CurrentTimeStamp.Reset();
+
+		for (int32& FrameCount : FramesLeftInMusicalDuration.FramesInTimeValueInternal)
+		{
+			FrameCount = 0;
+		}
+
 	}
 
 	void FQuartzMetronome::RecalculateDurations()
diff --git a/Engine/Source/Runtime/AudioMixer/Private/Quartz/QuartzSubsystem.cpp b/Engine/Source/Runtime/AudioMixer/Private/Quartz/QuartzSubsystem.cpp
--- a/Engine/Source/Runtime/AudioMixer/Private/Quartz/QuartzSubsystem.cpp
+++ b/Engine/Source/Runtime/AudioMixer/Private/Quartz/QuartzSubsystem.cpp
@@ -229,7 +229,7 @@ Audio::FQuartzQuantizedRequestData UQuartzSubsystem::CreateDataForTickRateChange
 	return CommandInitInfo;
 }
 
-Audio::FQuartzQuantizedRequestData UQuartzSubsystem::CreateDataForTransportReset(UQuartzClockHandle* InClockHandle, const FOnQuartzCommandEventBP& InDelegate)
+Audio::FQuartzQuantizedRequestData UQuartzSubsystem::CreateDataForTransportReset(UQuartzClockHandle* InClockHandle, const FQuartzQuantizationBoundary& InQuantizationBoundary, const FOnQuartzCommandEventBP& InDelegate)
 {
 	TSharedPtr<Audio::FQuantizedTransportReset> TransportResetCommandPtr = MakeShared<Audio::FQuantizedTransportReset>();
 
@@ -237,7 +237,7 @@ Audio::FQuartzQuantizedRequestData UQuartzSubsystem::CreateDataForTransportReset
 
 	CommandInitInfo.ClockName = InClockHandle->GetClockName();
 	CommandInitInfo.ClockHandleName = InClockHandle->GetHandleName();
-	CommandInitInfo.QuantizationBoundary = EQuartzCommandQuantization::Bar;
+	CommandInitInfo.QuantizationBoundary = InQuantizationBoundary;
 	CommandInitInfo.QuantizedCommandPtr = TransportResetCommandPtr;
 
 	if (InDelegate.IsBound())
@@ -341,6 +341,22 @@ bool UQuartzSubsystem::DoesClockExist(const UObject* WorldContextObject, FName C
 	return ClockManager->DoesClockExist(ClockName);
 }
 
+bool UQuartzSubsystem::IsClockRunning(const UObject* WorldContextObject, FName ClockName)
+{
+	if (DisableQuartzCvar)
+	{
+		return false;
+	}
+
+	Audio::FQuartzClockManager* ClockManager = GetClockManager(WorldContextObject);
+	if (!ClockManager)
+	{
+		return false;
+	}
+
+	return ClockManager->IsClockRunning(ClockName);
+}
+
 
 float UQuartzSubsystem::GetGameThreadToAudioRenderThreadAverageLatency(const UObject* WorldContextObject)
 {
diff --git a/Engine/Source/Runtime/AudioMixer/Public/Quartz/AudioMixerClockHandle.h b/Engine/Source/Runtime/AudioMixer/Public/Quartz/AudioMixerClockHandle.h
--- a/Engine/Source/Runtime/AudioMixer/Public/Quartz/AudioMixerClockHandle.h
+++ b/Engine/Source/Runtime/AudioMixer/Public/Quartz/AudioMixerClockHandle.h
@@ -37,6 +37,9 @@ public:
 // Begin Blueprint Interface
 
 	// Clock manipulation
+	UFUNCTION(BlueprintCallable, Category = "Quartz Clock", meta = (WorldContext = "WorldContextObject"))
+	void StartClock(const UObject* WorldContextObject);
+
 	UFUNCTION(BlueprintCallable, Category = "Quartz Subsystem", meta = (WorldContext = "WorldContextObject"))
 	void PauseClock(const UObject* WorldContextObject);
 
@@ -44,7 +47,10 @@ public:
  	void ResumeClock(const UObject* WorldContextObject);
 
 	UFUNCTION(BlueprintCallable, Category = "Quartz Clock", meta = (WorldContext = "WorldContextObject", AutoCreateRefTerm = "InDelegate", Keywords = "Transport, Counter"))
-	void ResetTransport(const UObject* WorldContextObject, const FOnQuartzCommandEventBP& InDelegate);
+	void ResetTransport(const UObject* WorldContextObject, FQuartzQuantizationBoundary InQuantizationBoundary, const FOnQuartzCommandEventBP& InDelegate);
+
+	UFUNCTION(BlueprintCallable, Category = "Quartz Clock", meta = (WorldContext = "WorldContextObject", Keywords = "Transport, Counter"))
+	bool IsClockRunning(const UObject* WorldContextObject);
 
 
 	// Metronome subscription
@@ -63,19 +69,19 @@ public:
 
 	// Metronome Alteration (setters)
 	UFUNCTION(BlueprintCallable, Category = "Quantization", meta = (WorldContext = "WorldContextObject", AdvancedDisplay = "QuantizationBoundary, Delegate", AutoCreateRefTerm = "QuantizationBoundary, Delegate", Keywords = "BPM, Tempo"))
-	void SetMillisecondsPerTick(const UObject* WorldContextObject, UPARAM(ref) const FQuartzQuantizationBoundary& QuantizationBoundary, const FOnQuartzCommandEventBP& Delegate, float MillisecondsPerTick = 100.f);
+	void SetMillisecondsPerTick(const UObject* WorldContextObject, UPARAM(ref) const FQuartzQuantizationBoundary& QuantizationBoundary, const FOnQuartzCommandEventBP& Delegate, UQuartzClockHandle*& ClockHandle, float MillisecondsPerTick = 100.f);
 
 	UFUNCTION(BlueprintCallable, Category = "Quantization", meta = (WorldContext = "WorldContextObject", AdvancedDisplay = "QuantizationBoundary, Delegate", AutoCreateRefTerm = "QuantizationBoundary, Delegate", Keywords = "BPM, Tempo"))
-	void SetTicksPerSecond(const UObject* WorldContextObject, UPARAM(ref) const FQuartzQuantizationBoundary& QuantizationBoundary, const FOnQuartzCommandEventBP& Delegate, float TicksPerSecond = 10.f);
+	void SetTicksPerSecond(const UObject* WorldContextObject, UPARAM(ref) const FQuartzQuantizationBoundary& QuantizationBoundary, const FOnQuartzCommandEventBP& Delegate, UQuartzClockHandle*& ClockHandle, float TicksPerSecond = 10.f);
 
 	UFUNCTION(BlueprintCallable, Category = "Quantization", meta = (WorldContext = "WorldContextObject", AdvancedDisplay = "QuantizationBoundary, Delegate", AutoCreateRefTerm = "QuantizationBoundary, Delegate", Keywords = "BPM, Tempo"))
-	void SetSecondsPerTick(const UObject* WorldContextObject, UPARAM(ref) const FQuartzQuantizationBoundary& QuantizationBoundary, const FOnQuartzCommandEventBP& Delegate, float SecondsPerTick = 0.25f);
+	void SetSecondsPerTick(const UObject* WorldContextObject, UPARAM(ref) const FQuartzQuantizationBoundary& QuantizationBoundary, const FOnQuartzCommandEventBP& Delegate, UQuartzClockHandle*& ClockHandle, float SecondsPerTick = 0.25f);
 
 	UFUNCTION(BlueprintCallable, Category = "Quantization", meta = (WorldContext = "WorldContextObject", AdvancedDisplay = "QuantizationBoundary, Delegate", AutoCreateRefTerm = "QuantizationBoundary, Delegate", Keywords = "BPM, Tempo"))
-	void SetThirtySecondNotesPerMinute(const UObject* WorldContextObject, UPARAM(ref) const FQuartzQuantizationBoundary& QuantizationBoundary, const FOnQuartzCommandEventBP& Delegate, float ThirtySecondsNotesPerMinute = 960.f);
+	void SetThirtySecondNotesPerMinute(const UObject* WorldContextObject, UPARAM(ref) const FQuartzQuantizationBoundary& QuantizationBoundary, const FOnQuartzCommandEventBP& Delegate, UQuartzClockHandle*& ClockHandle, float ThirtySecondsNotesPerMinute = 960.f);
 
 	UFUNCTION(BlueprintCallable, Category = "Quantization", meta = (WorldContext = "WorldContextObject", AdvancedDisplay = "QuantizationBoundary, Delegate", AutoCreateRefTerm = "QuantizationBoundary, Delegate", Keywords = "BPM, Tempo"))
-	void SetBeatsPerMinute(const UObject* WorldContextObject, UPARAM(ref) const FQuartzQuantizationBoundary& QuantizationBoundary, const FOnQuartzCommandEventBP& Delegate, float BeatsPerMinute = 60.f);
+	void SetBeatsPerMinute(const UObject* WorldContextObject, UPARAM(ref) const FQuartzQuantizationBoundary& QuantizationBoundary, const FOnQuartzCommandEventBP& Delegate, UQuartzClockHandle*& ClockHandle, float BeatsPerMinute = 60.f);
 
 
 	// Metronome getters
diff --git a/Engine/Source/Runtime/AudioMixer/Public/Quartz/AudioMixerClockManager.h b/Engine/Source/Runtime/AudioMixer/Public/Quartz/AudioMixerClockManager.h
--- a/Engine/Source/Runtime/AudioMixer/Public/Quartz/AudioMixerClockManager.h
+++ b/Engine/Source/Runtime/AudioMixer/Public/Quartz/AudioMixerClockManager.h
@@ -26,13 +26,15 @@ namespace Audio
 
 		// add (and take ownership of) a new clock
 		// safe to call from AudioThread (uses critical section)
-		FQuartzClock* GetOrCreateClock(const FName& InClockName, const FQuartzClockSettings& InClockSettings, bool bOverrideTickRateIfClockExists = false);
+		TSharedPtr<FQuartzClock> GetOrCreateClock(const FName& InClockName, const FQuartzClockSettings& InClockSettings, bool bOverrideTickRateIfClockExists = false);
 
 
 		// returns true if a clock with the given name already exists.
-		// OutKey is overwritten with a default object if the clock didn't exist
 		bool DoesClockExist(const FName& InClockName);
 
+		// returns true if the name is running
+		bool IsClockRunning(const FName& InClockName);
+
 		// remove existing clock
 		// safe to call from AudioThread (uses Audio Render Thread Command)
 		void RemoveClock(const FName& InName);
@@ -84,7 +86,7 @@ namespace Audio
 		void TickClocks(int32 NumFramesToTick);
 
 		// find clock with a given key
-		FQuartzClock* FindClock(const FName& InName);
+		TSharedPtr<FQuartzClock> FindClock(const FName& InName);
 
 		// pointer to owning FMixerDevice
 		FMixerDevice* MixerDevice;
@@ -93,6 +95,6 @@ namespace Audio
 		FCriticalSection ActiveClockCritSec;
 
 		// Our array of active clocks (mutation/access acquires clock)
-		TArray<FQuartzClock> ActiveClocks;
+		TArray<TSharedPtr<FQuartzClock>> ActiveClocks;
 	};
 } // namespace Audio
diff --git a/Engine/Source/Runtime/AudioMixer/Public/Quartz/QuartzSubsystem.h b/Engine/Source/Runtime/AudioMixer/Public/Quartz/QuartzSubsystem.h
--- a/Engine/Source/Runtime/AudioMixer/Public/Quartz/QuartzSubsystem.h
+++ b/Engine/Source/Runtime/AudioMixer/Public/Quartz/QuartzSubsystem.h
@@ -52,7 +52,7 @@ public:
 
 	// Helper functions for initializing quantized command initialization struct (to consolidate eyesore)
 	static Audio::FQuartzQuantizedRequestData CreateDataForTickRateChange(UQuartzClockHandle* InClockHandle, const FOnQuartzCommandEventBP& InDelegate, const Audio::FQuartzClockTickRate& InNewTickRate, const FQuartzQuantizationBoundary& InQuantizationBoundary);
-	static Audio::FQuartzQuantizedRequestData CreateDataForTransportReset(UQuartzClockHandle* InClockHandle, const FOnQuartzCommandEventBP& InDelegate);
+	static Audio::FQuartzQuantizedRequestData CreateDataForTransportReset(UQuartzClockHandle* InClockHandle, const FQuartzQuantizationBoundary& InQuantizationBoundary, const FOnQuartzCommandEventBP& InDelegate);
 	static Audio::FQuartzQuantizedRequestData CreateDataDataForSchedulePlaySound(UQuartzClockHandle* InClockHandle, const FOnQuartzCommandEventBP& InDelegate, const FQuartzQuantizationBoundary& InQuantizationBoundary);
 
 	UFUNCTION(BlueprintCallable, Category = "Quartz Clock Handle")
@@ -78,6 +78,10 @@ public:
 	UFUNCTION(BlueprintCallable, Category = "Quartz Clock Handle", meta = (WorldContext = "WorldContextObject"))
 	bool DoesClockExist(const UObject* WorldContextObject, FName ClockName);
 
+	// returns true if the clock is running
+	UFUNCTION(BlueprintCallable, Category = "Quartz Clock Handle", meta = (WorldContext = "WorldContextObject"))
+	bool IsClockRunning(const UObject* WorldContextObject, FName ClockName);
+
 	// latency data (Game thread -> Audio Render Thread)
 	UFUNCTION(BlueprintCallable, Category = "Quartz Subsystem", meta = (WorldContext = "WorldContextObject"))
 	float GetGameThreadToAudioRenderThreadAverageLatency(const UObject* WorldContextObject);
diff --git a/Engine/Source/Runtime/Engine/Classes/Sound/QuartzQuantizationUtilities.cpp b/Engine/Source/Runtime/Engine/Classes/Sound/QuartzQuantizationUtilities.cpp
--- a/Engine/Source/Runtime/Engine/Classes/Sound/QuartzQuantizationUtilities.cpp
+++ b/Engine/Source/Runtime/Engine/Classes/Sound/QuartzQuantizationUtilities.cpp
@@ -190,6 +190,8 @@ namespace Audio
 
 		switch (InDuration)
 		{
+		case EQuartzCommandQuantization::None:
+			return 0;
 
 			// NORMAL
 		case EQuartzCommandQuantization::Tick:
@@ -603,3 +605,15 @@ namespace Audio
 
 
 } // namespace Audio
+
+bool FQuartzTransportTimeStamp::IsZero() const
+{
+	return (!Bars) && (!Beat) && FMath::IsNearlyZero(BeatFraction);
+}
+
+void FQuartzTransportTimeStamp::Reset()
+{
+	Bars = 0;
+	Beat = 0;
+	BeatFraction = 0.f;
+}
diff --git a/Engine/Source/Runtime/Engine/Classes/Sound/QuartzQuantizationUtilities.h b/Engine/Source/Runtime/Engine/Classes/Sound/QuartzQuantizationUtilities.h
--- a/Engine/Source/Runtime/Engine/Classes/Sound/QuartzQuantizationUtilities.h
+++ b/Engine/Source/Runtime/Engine/Classes/Sound/QuartzQuantizationUtilities.h
@@ -67,7 +67,10 @@ enum class EQuartzCommandQuantization : uint8
 
 	Tick					UMETA(DisplayName = "On Tick (Smallest Value, same as 1/32)", ToolTip = "(same as 1/32)"),
 
-	Count				UMETA(Hidden),
+	Count					UMETA(Hidden),
+
+	None					UMETA(DisplayName = "None", ToolTip = "(Execute as son as possible)"),
+	// (when using "Count" in various logic, we don't want to account for "None")
 };
 
 // An enumeration for specifying the denominator of time signatures
@@ -137,21 +140,15 @@ struct ENGINE_API FQuartzTransportTimeStamp
 {
 	GENERATED_BODY()
 
-		int32 Bars { 0 };
+	int32 Bars { 0 };
 
 	int32 Beat{ 0 };
 
-	float BeatFration{ 0.f };
+	float BeatFraction{ 0.f };
 
-	bool IsZero() const
-	{
-		return (!Bars) && (!Beat) && FMath::IsNearlyZero(BeatFration);
-	}
+	bool IsZero() const;
 
-	void Reset()
-	{
-		Bars = 0;
-	}
+	void Reset();
 };
 
 
@@ -429,7 +426,7 @@ namespace Audio
 			, int32 InSourceID = -1
 		);
 
-		void SetOwningClockPtr(Audio::FQuartzClock* InClockPointer)
+		void SetOwningClockPtr(TSharedPtr<Audio::FQuartzClock> InClockPointer)
 		{
 			OwningClockPointer = InClockPointer;
 		}
@@ -443,7 +440,7 @@ namespace Audio
 		int32 GameThreadDelegateID{ -1 };
 
 		// Audio Render thread-specific data:
-		Audio::FQuartzClock* OwningClockPointer{ nullptr };
+		TSharedPtr<Audio::FQuartzClock> OwningClockPointer{ nullptr };
 		int32 SourceID{ -1 };
 	};
 
