Fixes for low-rate quartz crashes (pointer invalidation by Tarray reallocation)

diff --git a/Engine/Source/Runtime/AudioMixer/Private/Quartz/AudioMixerClock.cpp b/Engine/Source/Runtime/AudioMixer/Private/Quartz/AudioMixerClock.cpp
--- a/Engine/Source/Runtime/AudioMixer/Private/Quartz/AudioMixerClock.cpp
+++ b/Engine/Source/Runtime/AudioMixer/Private/Quartz/AudioMixerClock.cpp
@@ -220,6 +220,14 @@ namespace Audio
 
 	void FQuartzClock::AddQuantizedCommand(FQuartzQuantizationBoundary InQuantizationBondary, TSharedPtr<IQuartzQuantizedCommand> InNewEvent)
 	{
+		// if this is unquantized, execute immediately (even if the clock is paused)
+		if (InQuantizationBondary.Quantization == EQuartzCommandQuantization::None)
+		{
+			InNewEvent->AboutToStart();
+			InNewEvent->OnFinalCallback(0);
+			return;
+		}
+
 		// get number of frames until event (assuming we are at frame 0)
 		int32 FramesUntilExec = Metronome.GetFramesUntilBoundary(InQuantizationBondary);
 
@@ -250,6 +258,11 @@ namespace Audio
 		return ((PendingCommands.Num() + ClockAlteringPendingCommands.Num() ) > 0);
 	}
 
+	bool FQuartzClock::IsRunning()
+	{
+		return bIsRunning;
+	}
+
 	FMixerDevice* FQuartzClock::GetMixerDevice()
 	{
 		checkSlow(OwningClockManagerPtr);
diff --git a/Engine/Source/Runtime/AudioMixer/Private/Quartz/AudioMixerClockManager.cpp b/Engine/Source/Runtime/AudioMixer/Private/Quartz/AudioMixerClockManager.cpp
--- a/Engine/Source/Runtime/AudioMixer/Private/Quartz/AudioMixerClockManager.cpp
+++ b/Engine/Source/Runtime/AudioMixer/Private/Quartz/AudioMixerClockManager.cpp
@@ -13,6 +13,7 @@ namespace Audio
 
 	FQuartzClockManager::~FQuartzClockManager()
 	{
+		check(ActiveClocks.Num() == 0);
 	}
 
 	void FQuartzClockManager::Update(int32 NumFramesUntilNextUpdate)
@@ -24,7 +25,7 @@ namespace Audio
 		TickClocks(NumFramesUntilNextUpdate);
 	}
 
-	FQuartzClock* FQuartzClockManager::GetOrCreateClock(const FName& InClockName, const FQuartzClockSettings& InClockSettings, bool bOverrideTickRateIfClockExists)
+	TSharedPtr<FQuartzClock> FQuartzClockManager::GetOrCreateClock(const FName& InClockName, const FQuartzClockSettings& InClockSettings, bool bOverrideTickRateIfClockExists)
 	{
 		FScopeLock Lock(&ActiveClockCritSec);
 
@@ -32,7 +33,7 @@ namespace Audio
 		FQuartzClockSettings NewSettings = InClockSettings;
 
 		// See if this clock already exists
-		FQuartzClock* Clock = FindClock(InClockName);
+		TSharedPtr<FQuartzClock> Clock = FindClock(InClockName);
 
 		if (Clock)
 		{
@@ -46,7 +47,7 @@ namespace Audio
 		}
 
 		// doesn't exist, create new clock
-		return &ActiveClocks.Emplace_GetRef(FQuartzClock(InClockName, NewSettings, this));
+		return ActiveClocks.Emplace_GetRef(MakeShared<FQuartzClock>(InClockName, NewSettings, this));
 	}
 
 	bool FQuartzClockManager::DoesClockExist(const FName& InClockName)
@@ -55,6 +56,22 @@ namespace Audio
 		return !!FindClock(InClockName);
 	}
 
+	bool FQuartzClockManager::IsClockRunning(const FName& InClockName)
+	{
+		FScopeLock Lock(&ActiveClockCritSec);
+
+		// See if this clock already exists
+		TSharedPtr<FQuartzClock> Clock = FindClock(InClockName);
+
+		if (Clock)
+		{
+			return Clock->IsRunning();
+		}
+
+		// doesn't exist, create new clock
+		return false;
+	}
+
 	void FQuartzClockManager::RemoveClock(const FName& InName)
 	{
 		if (!MixerDevice->IsAudioRenderingThread())
@@ -72,7 +89,7 @@ namespace Audio
 		int32 NumClocks = ActiveClocks.Num();
 		for (int32 i = NumClocks - 1; i >= 0; --i)
 		{
-			if (ActiveClocks[i].GetName() == InName)
+			if (ActiveClocks[i]->GetName() == InName)
 			{
 				UE_LOG(LogAudioQuartz, Display, TEXT("Removing Clock: %s"), *InName.ToString());
 				ActiveClocks.RemoveAtSwap(i);
@@ -85,7 +102,7 @@ namespace Audio
 	{
 		FScopeLock Lock(&ActiveClockCritSec);
 
-		FQuartzClock* Clock = FindClock(InName);
+		TSharedPtr<FQuartzClock> Clock = FindClock(InName);
 
 		if (Clock)
 		{
@@ -109,7 +126,7 @@ namespace Audio
 
 		// Anything below is being executed on the Audio Render Thread
 		FScopeLock Lock(&ActiveClockCritSec);
-		FQuartzClock* Clock = FindClock(InName);
+		TSharedPtr<FQuartzClock> Clock = FindClock(InName);
 		if (Clock)
 		{
 			Clock->ChangeTickRate(InNewTickRate);
@@ -130,7 +147,7 @@ namespace Audio
 
 		// Anything below is being executed on the Audio Render Thread
 		FScopeLock Lock(&ActiveClockCritSec);
-		FQuartzClock* Clock = FindClock(InName);
+		TSharedPtr<FQuartzClock> Clock = FindClock(InName);
 		if (Clock)
 		{
 			Clock->Resume();
@@ -151,7 +168,7 @@ namespace Audio
 
 		// Anything below is being executed on the Audio Render Thread
 		FScopeLock Lock(&ActiveClockCritSec);
-		FQuartzClock* Clock = FindClock(InName);
+		TSharedPtr<FQuartzClock> Clock = FindClock(InName);
 		if (Clock)
 		{
 			Clock->Pause();
@@ -166,7 +183,7 @@ namespace Audio
 
 		for (int32 i = NumClocks - 1; i >= 0; --i)
 		{
-			if (!ActiveClocks[i].IgnoresFlush())
+			if (!ActiveClocks[i]->IgnoresFlush())
 			{
 				ActiveClocks.RemoveAtSwap(i);
 			}
@@ -180,7 +197,7 @@ namespace Audio
 		FScopeLock Lock(&ActiveClockCritSec);
 		for (auto& Clock : ActiveClocks)
 		{
-			Clock.Shutdown();
+			Clock->Shutdown();
 		}
 	}
 
@@ -188,7 +205,7 @@ namespace Audio
 	{
 		FScopeLock Lock(&ActiveClockCritSec);
 
-		if (FQuartzClock* Clock = FindClock(InQuantizationCommandInitInfo.ClockName))
+		if (TSharedPtr<FQuartzClock> Clock = FindClock(InQuantizationCommandInitInfo.ClockName))
 		{
 			// pass the quantized command to it's clock
 			InQuantizationCommandInitInfo.SetOwningClockPtr(Clock);
@@ -221,7 +238,7 @@ namespace Audio
 
 		// Anything below is being executed on the Audio Render Thread
 		FScopeLock Lock(&ActiveClockCritSec);
-		FQuartzClock* Clock = FindClock(InClockName);
+		TSharedPtr<FQuartzClock> Clock = FindClock(InClockName);
 		if (Clock)
 		{
 			Clock->SubscribeToTimeDivision(InListenerQueue, InQuantizationBoundary);
@@ -242,7 +259,7 @@ namespace Audio
 
 		// Anything below is being executed on the Audio Render Thread
 		FScopeLock Lock(&ActiveClockCritSec);
-		FQuartzClock* Clock = FindClock(InClockName);
+		TSharedPtr<FQuartzClock> Clock = FindClock(InClockName);
 		if (Clock)
 		{
 			Clock->SubscribeToAllTimeDivisions(InListenerQueue);
@@ -263,7 +280,7 @@ namespace Audio
 
 		// Anything below is being executed on the Audio Render Thread
 		FScopeLock Lock(&ActiveClockCritSec);
-		FQuartzClock* Clock = FindClock(InClockName);
+		TSharedPtr<FQuartzClock> Clock = FindClock(InClockName);
 		if (Clock)
 		{
 			Clock->UnsubscribeFromTimeDivision(InListenerQueue, InQuantizationBoundary);
@@ -284,7 +301,7 @@ namespace Audio
 
 		// Anything below is being executed on the Audio Render Thread
 		FScopeLock Lock(&ActiveClockCritSec);
-		FQuartzClock* Clock = FindClock(InClockName);
+		TSharedPtr<FQuartzClock> Clock = FindClock(InClockName);
 		if (Clock)
 		{
 			Clock->UnsubscribeFromAllTimeDivisions(InListenerQueue);
@@ -297,7 +314,7 @@ namespace Audio
 		check(MixerDevice->IsAudioRenderingThread());
 
 		FScopeLock Lock(&ActiveClockCritSec);
-		FQuartzClock* Clock = FindClock(InOwningClockName);
+		TSharedPtr<FQuartzClock> Clock = FindClock(InOwningClockName);
 
 		if (Clock && InCommandPtr)
 		{
@@ -329,20 +346,20 @@ namespace Audio
 		FScopeLock Lock(&ActiveClockCritSec);
 		for (auto& Clock : ActiveClocks)
 		{
-			Clock.Tick(NumFramesToTick);
+			Clock->Tick(NumFramesToTick);
 		}
 	}
 
-	FQuartzClock* FQuartzClockManager::FindClock(const FName& InName)
+	TSharedPtr<FQuartzClock> FQuartzClockManager::FindClock(const FName& InName)
 	{
 		// the function calling this should be be on the audio rendering thread or have acquired the lock
 		checkSlow(MixerDevice->IsAudioRenderingThread() || ActiveClockCritSec.TryLock());
 
 		for (auto& Clock : ActiveClocks)
 		{
-			if (Clock.GetName() == InName)
+			if (Clock->GetName() == InName)
 			{
-				return &Clock;
+				return Clock;
 			}
 		}
 
diff --git a/Engine/Source/Runtime/AudioMixer/Public/Quartz/AudioMixerClock.h b/Engine/Source/Runtime/AudioMixer/Public/Quartz/AudioMixerClock.h
--- a/Engine/Source/Runtime/AudioMixer/Public/Quartz/AudioMixerClock.h
+++ b/Engine/Source/Runtime/AudioMixer/Public/Quartz/AudioMixerClock.h
@@ -98,6 +98,9 @@ namespace Audio
 		// does the clock have any pending events
 		bool HasPendingEvents();
 
+		// is the clock currently ticking?
+		bool IsRunning();
+
 		FMixerDevice* GetMixerDevice();
 
 		FMixerSourceManager* GetSourceManager();
diff --git a/Engine/Source/Runtime/AudioMixer/Public/Quartz/AudioMixerQuantizedCommands.h b/Engine/Source/Runtime/AudioMixer/Public/Quartz/AudioMixerQuantizedCommands.h
--- a/Engine/Source/Runtime/AudioMixer/Public/Quartz/AudioMixerQuantizedCommands.h
+++ b/Engine/Source/Runtime/AudioMixer/Public/Quartz/AudioMixerQuantizedCommands.h
@@ -26,7 +26,7 @@ namespace Audio
 		virtual FName GetCommandName() const override;
 
 	private:
-		FQuartzClock* OwningClockPtr{ nullptr };
+		TSharedPtr<FQuartzClock> OwningClockPtr{ nullptr };
 
 		int32 SourceID{ -1 };
 
@@ -54,7 +54,7 @@ namespace Audio
 
 	private:
 		FQuartzClockTickRate TickRate;
-		FQuartzClock* OwningClockPtr{ nullptr };
+		TSharedPtr<FQuartzClock> OwningClockPtr{ nullptr };
 
 	}; // class FQuantizedTickRateChange 
 
@@ -74,7 +74,7 @@ namespace Audio
 		virtual FName GetCommandName() const override;
 
 	private:
-		FQuartzClock* OwningClockPtr{ nullptr };
+		TSharedPtr<FQuartzClock> OwningClockPtr{ nullptr };
 
 	}; // class FQuantizedTransportReset 
 
