Quartz: Adding quantized ability to start another clock

diff --git a/Engine/Source/Runtime/AudioMixer/Private/Quartz/AudioMixerClock.cpp b/Engine/Source/Runtime/AudioMixer/Private/Quartz/AudioMixerClock.cpp
--- a/Engine/Source/Runtime/AudioMixer/Private/Quartz/AudioMixerClock.cpp
+++ b/Engine/Source/Runtime/AudioMixer/Private/Quartz/AudioMixerClock.cpp
@@ -80,6 +80,26 @@ namespace Audio
 		bIsRunning = true;
 	}
 
+	void FQuartzClock::Stop(bool CancelPendingEvents)
+	{
+		bIsRunning = false;
+		Metronome.ResetTransport();
+		TickDelayLengthInFrames = 0;
+
+		if (CancelPendingEvents)
+		{
+			for (auto& Command : PendingCommands)
+			{
+				Command.Command->Cancel();
+			}
+
+			for (auto& Command : ClockAlteringPendingCommands)
+			{
+				Command.Command->Cancel();
+			}
+		}
+	}
+
 	void FQuartzClock::Pause()
 	{
 		if (bIsRunning)
@@ -103,6 +123,7 @@ namespace Audio
 	void FQuartzClock::Restart(bool bPause)
 	{
 		bIsRunning = !bPause;
+		TickDelayLengthInFrames = 0;
 	}
 
 	void FQuartzClock::Shutdown()
@@ -123,20 +144,38 @@ namespace Audio
 
 	void FQuartzClock::Tick(int32 InNumFramesUntilNextTick)
 	{
-		const int32 FramesOfLatency = (ThreadLatencyInMilliseconds / 1000) * Metronome.GetTickRate().GetSampleRate();
-
 		if (!bIsRunning)
 		{
 			return;
 		}
 
-		TickInternal(InNumFramesUntilNextTick, ClockAlteringPendingCommands, FramesOfLatency); // (process things like BPM changes first)
-		TickInternal(InNumFramesUntilNextTick, PendingCommands, FramesOfLatency);
+		if (TickDelayLengthInFrames >= InNumFramesUntilNextTick)
+		{
+			TickDelayLengthInFrames -= InNumFramesUntilNextTick;
+			return;
+		}
+
+		const int32 FramesOfLatency = (ThreadLatencyInMilliseconds / 1000) * Metronome.GetTickRate().GetSampleRate();
+
+		if (TickDelayLengthInFrames == 0)
+		{
+			TickInternal(InNumFramesUntilNextTick, ClockAlteringPendingCommands, FramesOfLatency); // (process things like BPM changes first)
+			TickInternal(InNumFramesUntilNextTick, PendingCommands, FramesOfLatency);
+		}
+		else
+		{
+			TickInternal(TickDelayLengthInFrames, ClockAlteringPendingCommands, FramesOfLatency);
+			TickInternal(TickDelayLengthInFrames, PendingCommands, FramesOfLatency);
+
+			TickInternal(InNumFramesUntilNextTick - TickDelayLengthInFrames, ClockAlteringPendingCommands, FramesOfLatency, TickDelayLengthInFrames);
+			TickInternal(InNumFramesUntilNextTick - TickDelayLengthInFrames, PendingCommands, FramesOfLatency, TickDelayLengthInFrames);
+		}
+
 
 		Metronome.Tick(InNumFramesUntilNextTick, FramesOfLatency);
 	}
 
-	void FQuartzClock::TickInternal(int32 InNumFramesUntilNextTick, TArray<PendingCommand>& CommandsToTick, int32 FramesOfLatency)
+	void FQuartzClock::TickInternal(int32 InNumFramesUntilNextTick, TArray<PendingCommand>& CommandsToTick, int32 FramesOfLatency, int32 FramesOfDelay)
 	{
 		bool bHaveCommandsToRemove = false;
 
@@ -152,7 +191,7 @@ namespace Audio
 			// Time To execute?
 			if (PendingCommand.NumFramesUntilExec < InNumFramesUntilNextTick)
 			{
-				PendingCommand.Command->OnFinalCallback(PendingCommand.NumFramesUntilExec);
+				PendingCommand.Command->OnFinalCallback(PendingCommand.NumFramesUntilExec + FramesOfDelay);
 				PendingCommand.Command.Reset();
 				bHaveCommandsToRemove = true;
 			}
@@ -287,6 +326,16 @@ namespace Audio
 		return nullptr;
 	}
 
+	FQuartzClockManager* FQuartzClock::GetClockManager()
+	{
+		checkSlow(OwningClockManagerPtr);
+		if (OwningClockManagerPtr)
+		{
+			return OwningClockManagerPtr;
+		}
+		return nullptr;
+	}
+
 	void FQuartzClock::ResetTransport()
 	{
 		Metronome.ResetTransport();
diff --git a/Engine/Source/Runtime/AudioMixer/Private/Quartz/AudioMixerClockHandle.cpp b/Engine/Source/Runtime/AudioMixer/Private/Quartz/AudioMixerClockHandle.cpp
--- a/Engine/Source/Runtime/AudioMixer/Private/Quartz/AudioMixerClockHandle.cpp
+++ b/Engine/Source/Runtime/AudioMixer/Private/Quartz/AudioMixerClockHandle.cpp
@@ -65,6 +65,19 @@ void UQuartzClockHandle::StartClock(const UObject* WorldContextObject)
 	ResumeClock(WorldContextObject);
 }
 
+void UQuartzClockHandle::StopClock(const UObject* WorldContextObject, bool CancelPendingEvents)
+{
+	if (QuartzSubsystem)
+	{
+		Audio::FQuartzClockManager* ClockManager = QuartzSubsystem->GetClockManager(WorldContextObject);
+
+		if (ClockManager)
+		{
+			ClockManager->StopClock(CurrentClockId, CancelPendingEvents);
+		}
+	}
+}
+
 void UQuartzClockHandle::PauseClock(const UObject* WorldContextObject)
 {
 	if (QuartzSubsystem)
@@ -110,6 +123,15 @@ bool UQuartzClockHandle::IsClockRunning(const UObject* WorldContextObject)
 	return false;
 }
 
+void UQuartzClockHandle::StartOtherClock(const UObject* WorldContextObject, FName OtherClockName, FQuartzQuantizationBoundary InQuantizationBoundary, const FOnQuartzCommandEventBP& InDelegate)
+{
+	if (QuartzSubsystem)
+	{
+		Audio::FQuartzQuantizedCommandInitInfo Data(QuartzSubsystem->CreateDataForStartOtherClock(this, OtherClockName, InQuantizationBoundary, InDelegate));
+		QuartzSubsystem->AddCommandToClock(WorldContextObject, Data);
+	}
+}
+
 void UQuartzClockHandle::SubscribeToQuantizationEvent(const UObject* WorldContextObject, EQuartzCommandQuantization InQuantizationBoundary, const FOnQuartzMetronomeEventBP& OnQuantizationEvent)
 {
 	if (QuartzSubsystem)
diff --git a/Engine/Source/Runtime/AudioMixer/Private/Quartz/AudioMixerClockManager.cpp b/Engine/Source/Runtime/AudioMixer/Private/Quartz/AudioMixerClockManager.cpp
--- a/Engine/Source/Runtime/AudioMixer/Private/Quartz/AudioMixerClockManager.cpp
+++ b/Engine/Source/Runtime/AudioMixer/Private/Quartz/AudioMixerClockManager.cpp
@@ -21,10 +21,21 @@ namespace Audio
 		// this function should only be called on the Audio Render Thread
 		// (original intent was for this to be called only by the owning FMixerDevice)
 		check(MixerDevice->IsAudioRenderingThread());
-
+		LastUpdateSizeInFrames = NumFramesUntilNextUpdate;
 		TickClocks(NumFramesUntilNextUpdate);
 	}
 
+	void FQuartzClockManager::UpdateClock(FName InClockToAdvance, int32 NumFramesToAdvance)
+	{
+		FScopeLock Lock(&ActiveClockCritSec);
+		TSharedPtr<FQuartzClock> ClockPtr = FindClock(InClockToAdvance);
+
+		if (ClockPtr.IsValid())
+		{
+			ClockPtr->Tick(NumFramesToAdvance);
+		}
+	}
+
 	TSharedPtr<FQuartzClock> FQuartzClockManager::GetOrCreateClock(const FName& InClockName, const FQuartzClockSettings& InClockSettings, bool bOverrideTickRateIfClockExists)
 	{
 		FScopeLock Lock(&ActiveClockCritSec);
@@ -133,7 +144,7 @@ namespace Audio
 		}
 	}
 
-	void FQuartzClockManager::ResumeClock(const FName& InName)
+	void FQuartzClockManager::ResumeClock(const FName& InName, int32 NumFramesToDelayStart)
 	{
 		if (!MixerDevice->IsAudioRenderingThread())
 		{
@@ -150,10 +161,32 @@ namespace Audio
 		TSharedPtr<FQuartzClock> Clock = FindClock(InName);
 		if (Clock)
 		{
+			Clock->AddToTickDelay(NumFramesToDelayStart);
 			Clock->Resume();
 		}
 	}
 
+	void FQuartzClockManager::StopClock(const FName& InName, bool CancelPendingEvents)
+	{
+		if (!MixerDevice->IsAudioRenderingThread())
+		{
+			MixerDevice->AudioRenderThreadCommand([this, InName, CancelPendingEvents]()
+			{
+				StopClock(InName, CancelPendingEvents);
+			});
+
+			return;
+		}
+
+		// Anything below is being executed on the Audio Render Thread
+		FScopeLock Lock(&ActiveClockCritSec);
+		TSharedPtr<FQuartzClock> Clock = FindClock(InName);
+		if (Clock)
+		{
+			Clock->Stop(CancelPendingEvents);
+		}
+	}
+
 	void FQuartzClockManager::PauseClock(const FName& InName)
 	{
 		if (!MixerDevice->IsAudioRenderingThread())
@@ -324,6 +357,32 @@ namespace Audio
 		return false;
 	}
 
+	bool FQuartzClockManager::HasClockBeenTickedThisUpdate(FName InClockName)
+	{
+		FScopeLock Lock(&ActiveClockCritSec);
+		int32 NumClocks = ActiveClocks.Num();
+
+		for (int32 i = 0; i < NumClocks; ++i)
+		{
+			TSharedPtr<FQuartzClock> ClockPtr = ActiveClocks[i];
+			check(ClockPtr.IsValid());
+
+			if (ClockPtr->GetName() == InClockName)
+			{
+				// if this clock is earlier in the array than the last clock we ticked,
+				// then it has already been ticked this update
+				if (&ClockPtr < &ActiveClocks[LastClockTickedIndex.GetValue()])
+				{
+					return true;
+				}
+
+				return false;
+			}
+		}
+
+		return false;
+	}
+
 	FMixerDevice* FQuartzClockManager::GetMixerDevice() const
 	{
 		if (MixerDevice)
@@ -347,7 +406,10 @@ namespace Audio
 		for (auto& Clock : ActiveClocks)
 		{
 			Clock->Tick(NumFramesToTick);
+			LastClockTickedIndex.Increment();
 		}
+
+		LastClockTickedIndex.Reset();
 	}
 
 	TSharedPtr<FQuartzClock> FQuartzClockManager::FindClock(const FName& InName)
diff --git a/Engine/Source/Runtime/AudioMixer/Private/Quartz/AudioMixerQuantizedCommands.cpp b/Engine/Source/Runtime/AudioMixer/Private/Quartz/AudioMixerQuantizedCommands.cpp
--- a/Engine/Source/Runtime/AudioMixer/Private/Quartz/AudioMixerQuantizedCommands.cpp
+++ b/Engine/Source/Runtime/AudioMixer/Private/Quartz/AudioMixerQuantizedCommands.cpp
@@ -130,4 +130,59 @@ namespace Audio
 		return TransportResetCommandName;
 	}
 
+
+	TSharedPtr<IQuartzQuantizedCommand> FQuantizedOtherClockStart::GetDeepCopyOfDerivedObject() const
+	{
+		TSharedPtr<FQuantizedOtherClockStart> NewCopy = MakeShared<FQuantizedOtherClockStart>();
+
+		NewCopy->OwningClockPtr = OwningClockPtr;
+		NewCopy->NameOfClockToStart = NameOfClockToStart;
+
+		return NewCopy;
+	}
+
+	void FQuantizedOtherClockStart::OnQueuedCustom(const FQuartzQuantizedCommandInitInfo& InCommandInitInfo)
+	{
+		OwningClockPtr = InCommandInitInfo.OwningClockPointer;
+		check(OwningClockPtr.IsValid());
+
+		NameOfClockToStart = InCommandInitInfo.OtherClockName;
+
+		// shouldn't be trying to start the same clock that owns this command
+		if (OwningClockPtr->GetName() == NameOfClockToStart)
+		{
+			UE_LOG(LogAudioQuartz, Warning, TEXT("Clock: (%s) is attempting to start itself on a quantization boundary.  Ignoring command"), *NameOfClockToStart.ToString());
+			OwningClockPtr->CancelQuantizedCommand(TSharedPtr<IQuartzQuantizedCommand>(this));
+		}
+	}
+
+	void FQuantizedOtherClockStart::OnFinalCallbackCustom(int32 InNumFramesLeft)
+	{
+		// get access to the source & clock manager
+		// access source manager through owning clock (via clock manager)
+		FMixerSourceManager* SourceManager = OwningClockPtr->GetSourceManager();
+		FQuartzClockManager* ClockManager = OwningClockPtr->GetClockManager();
+
+		check(SourceManager && ClockManager && OwningClockPtr.IsValid());
+		bool bShouldStart = SourceManager && ClockManager && OwningClockPtr.IsValid();
+		bShouldStart |= !ClockManager->IsClockRunning(NameOfClockToStart);
+
+		if (bShouldStart)
+		{
+			// ...start the clock
+			ClockManager->ResumeClock(NameOfClockToStart, InNumFramesLeft);
+
+			if (ClockManager->HasClockBeenTickedThisUpdate(NameOfClockToStart))
+			{
+				ClockManager->UpdateClock(NameOfClockToStart, ClockManager->GetLastUpdateSizeInFrames());
+			}
+		}
+	}
+
+	static const FName StartOtherClockName("Start Other Clock Command");
+	FName FQuantizedOtherClockStart::GetCommandName() const
+	{
+		return StartOtherClockName;
+	}
+
 } // namespace Audio
diff --git a/Engine/Source/Runtime/AudioMixer/Private/Quartz/QuartzMetronome.cpp b/Engine/Source/Runtime/AudioMixer/Private/Quartz/QuartzMetronome.cpp
--- a/Engine/Source/Runtime/AudioMixer/Private/Quartz/QuartzMetronome.cpp
+++ b/Engine/Source/Runtime/AudioMixer/Private/Quartz/QuartzMetronome.cpp
@@ -319,6 +319,7 @@ namespace Audio
 			FrameCount = 0;
 		}
 
+		PulseDurationIndex = -1;
 	}
 
 	void FQuartzMetronome::RecalculateDurations()
diff --git a/Engine/Source/Runtime/AudioMixer/Private/Quartz/QuartzSubsystem.cpp b/Engine/Source/Runtime/AudioMixer/Private/Quartz/QuartzSubsystem.cpp
--- a/Engine/Source/Runtime/AudioMixer/Private/Quartz/QuartzSubsystem.cpp
+++ b/Engine/Source/Runtime/AudioMixer/Private/Quartz/QuartzSubsystem.cpp
@@ -248,6 +248,26 @@ Audio::FQuartzQuantizedRequestData UQuartzSubsystem::CreateDataForTransportReset
 	return CommandInitInfo;
 }
 
+Audio::FQuartzQuantizedRequestData UQuartzSubsystem::CreateDataForStartOtherClock(UQuartzClockHandle* InClockHandle, FName InClockToStart, const FQuartzQuantizationBoundary& InQuantizationBoundary, const FOnQuartzCommandEventBP& InDelegate)
+{
+	TSharedPtr<Audio::FQuantizedOtherClockStart> TransportResetCommandPtr = MakeShared<Audio::FQuantizedOtherClockStart>();
+
+	Audio::FQuartzQuantizedRequestData CommandInitInfo;
+
+	CommandInitInfo.ClockName = InClockHandle->GetClockName();
+	CommandInitInfo.ClockHandleName = InClockHandle->GetHandleName();
+	CommandInitInfo.OtherClockName = InClockToStart;
+	CommandInitInfo.QuantizationBoundary = InQuantizationBoundary;
+	CommandInitInfo.QuantizedCommandPtr = TransportResetCommandPtr;
+
+	if (InDelegate.IsBound())
+	{
+		CommandInitInfo.GameThreadDelegateID = InClockHandle->AddCommandDelegate(InDelegate, CommandInitInfo.GameThreadCommandQueue);
+	}
+
+	return CommandInitInfo;
+}
+
 UQuartzClockHandle* UQuartzSubsystem::CreateNewClock(const UObject* WorldContextObject, FName ClockName, FQuartzClockSettings InSettings, bool bOverrideSettingsIfClockExists)
 {
 	if (DisableQuartzCvar)
diff --git a/Engine/Source/Runtime/AudioMixer/Public/Quartz/AudioMixerClock.h b/Engine/Source/Runtime/AudioMixer/Public/Quartz/AudioMixerClock.h
--- a/Engine/Source/Runtime/AudioMixer/Public/Quartz/AudioMixerClock.h
+++ b/Engine/Source/Runtime/AudioMixer/Public/Quartz/AudioMixerClock.h
@@ -53,6 +53,9 @@ namespace Audio
 		// start ticking the clock
 		void Resume();
 
+		// stop ticking and reset the clock
+		void Stop(bool CancelPendingEvents);
+
 		// stop ticking the clock
 		void Pause();
 
@@ -105,10 +108,22 @@ namespace Audio
 
 		FMixerSourceManager* GetSourceManager();
 
+		FQuartzClockManager* GetClockManager();
+
 		void ResetTransport();
 
+		void AddToTickDelay(int32 NumFramesOfDelayToAdd)
+		{
+			TickDelayLengthInFrames += NumFramesOfDelayToAdd;
+		}
+
+		void SetTickDelay(int32 NumFramesOfDelay)
+		{
+			TickDelayLengthInFrames = NumFramesOfDelay;
+		}
+
 	private:
-		void TickInternal(int32 InNumFramesUntilNextTick, TArray<PendingCommand>& CommandsToTick, int32 FramesOfLatency = 0);
+		void TickInternal(int32 InNumFramesUntilNextTick, TArray<PendingCommand>& CommandsToTick, int32 FramesOfLatency = 0, int32 FramesOfDelay = 0);
 
 		bool CancelQuantizedCommandInternal(TSharedPtr<IQuartzQuantizedCommand> InCommandPtr, TArray<PendingCommand>& CommandsToTick);
 
@@ -133,6 +148,7 @@ namespace Audio
 
 		bool bIgnoresFlush{ false };
 
+		int32 TickDelayLengthInFrames{ 0 };
 
 	}; // class FQuartzClock
 
diff --git a/Engine/Source/Runtime/AudioMixer/Public/Quartz/AudioMixerClockHandle.h b/Engine/Source/Runtime/AudioMixer/Public/Quartz/AudioMixerClockHandle.h
--- a/Engine/Source/Runtime/AudioMixer/Public/Quartz/AudioMixerClockHandle.h
+++ b/Engine/Source/Runtime/AudioMixer/Public/Quartz/AudioMixerClockHandle.h
@@ -40,6 +40,9 @@ public:
 	UFUNCTION(BlueprintCallable, Category = "Quartz Clock", meta = (WorldContext = "WorldContextObject"))
 	void StartClock(const UObject* WorldContextObject);
 
+	UFUNCTION(BlueprintCallable, Category = "Quartz Subsystem", meta = (WorldContext = "WorldContextObject"))
+	void StopClock(const UObject* WorldContextObject, bool CancelPendingEvents);
+
 	UFUNCTION(BlueprintCallable, Category = "Quartz Subsystem", meta = (WorldContext = "WorldContextObject"))
 	void PauseClock(const UObject* WorldContextObject);
 
@@ -52,6 +55,9 @@ public:
 	UFUNCTION(BlueprintCallable, Category = "Quartz Clock", meta = (WorldContext = "WorldContextObject", Keywords = "Transport, Counter"))
 	bool IsClockRunning(const UObject* WorldContextObject);
 
+	// "other" clock manipulation
+	UFUNCTION(BlueprintCallable, Category = "Quartz Clock", meta = (WorldContext = "WorldContextObject", AutoCreateRefTerm = "InDelegate", Keywords = "Transport, Counter"))
+	void StartOtherClock(const UObject* WorldContextObject, FName OtherClockName, FQuartzQuantizationBoundary InQuantizationBoundary, const FOnQuartzCommandEventBP& InDelegate);
 
 	// Metronome subscription
 	UFUNCTION(BlueprintCallable, Category = "Quartz Clock", meta = (WorldContext = "WorldContextObject"))
diff --git a/Engine/Source/Runtime/AudioMixer/Public/Quartz/AudioMixerClockManager.h b/Engine/Source/Runtime/AudioMixer/Public/Quartz/AudioMixerClockManager.h
--- a/Engine/Source/Runtime/AudioMixer/Public/Quartz/AudioMixerClockManager.h
+++ b/Engine/Source/Runtime/AudioMixer/Public/Quartz/AudioMixerClockManager.h
@@ -23,6 +23,7 @@ namespace Audio
 
 		// Called on AudioRenderThread
 		void Update(int32 NumFramesUntilNextUpdate);
+		void UpdateClock(FName InClockToAdvance, int32 NumFramesToAdvance);
 
 		// add (and take ownership of) a new clock
 		// safe to call from AudioThread (uses critical section)
@@ -47,7 +48,11 @@ namespace Audio
 
 		// start the given clock
 		// safe to call from AudioThread (uses Audio Render Thread command)
-		void ResumeClock(const FName& InName);
+		void ResumeClock(const FName& InName, int32 NumFramesToDelayStart = 0);
+
+		// stop the given clock
+		// safe to call from AudioThread (uses Audio Render Thread command)
+		void StopClock(const FName& InName, bool CancelPendingEvents);
 
 		// stop the given clock
 		// safe to call from AudioThread (uses Audio Render Thread command)
@@ -78,6 +83,10 @@ namespace Audio
 		// cancel a queued command on a clock (i.e. cancel a PlayQuantized command if the sound is stopped before it is played)
 		bool CancelCommandOnClock(FName InOwningClockName, TSharedPtr<IQuartzQuantizedCommand> InCommandPtr);
 
+		bool HasClockBeenTickedThisUpdate(FName InClockName);
+
+		int32 GetLastUpdateSizeInFrames() { return LastUpdateSizeInFrames; }
+
 		// get access to the owning FMixerDevice
 		FMixerDevice* GetMixerDevice() const;
 
@@ -96,5 +105,8 @@ namespace Audio
 
 		// Our array of active clocks (mutation/access acquires clock)
 		TArray<TSharedPtr<FQuartzClock>> ActiveClocks;
+
+		FThreadSafeCounter LastClockTickedIndex{ 0 };
+		int32 LastUpdateSizeInFrames{ 0 };
 	};
 } // namespace Audio
diff --git a/Engine/Source/Runtime/AudioMixer/Public/Quartz/AudioMixerQuantizedCommands.h b/Engine/Source/Runtime/AudioMixer/Public/Quartz/AudioMixerQuantizedCommands.h
--- a/Engine/Source/Runtime/AudioMixer/Public/Quartz/AudioMixerQuantizedCommands.h
+++ b/Engine/Source/Runtime/AudioMixer/Public/Quartz/AudioMixerQuantizedCommands.h
@@ -78,4 +78,25 @@ namespace Audio
 
 	}; // class FQuantizedTransportReset 
 
+
+		// QuartzQuantizedCommand that starts a second clock on a sample-accurate boundary
+	class AUDIOMIXER_API FQuantizedOtherClockStart : public IQuartzQuantizedCommand
+	{
+	public:
+		virtual TSharedPtr<IQuartzQuantizedCommand> GetDeepCopyOfDerivedObject() const override;
+
+		virtual void OnQueuedCustom(const FQuartzQuantizedCommandInitInfo& InCommandInitInfo) override;
+
+		virtual void OnFinalCallbackCustom(int32 InNumFramesLeft) override;
+
+		virtual bool IsClockAltering() override { return true; }
+
+		virtual FName GetCommandName() const override;
+
+	private:
+		TSharedPtr<FQuartzClock> OwningClockPtr{ nullptr };
+		FName NameOfClockToStart;
+
+	}; // class FQuantizedOtherClockStart 
+
 } // namespace Audio
diff --git a/Engine/Source/Runtime/AudioMixer/Public/Quartz/QuartzSubsystem.h b/Engine/Source/Runtime/AudioMixer/Public/Quartz/QuartzSubsystem.h
--- a/Engine/Source/Runtime/AudioMixer/Public/Quartz/QuartzSubsystem.h
+++ b/Engine/Source/Runtime/AudioMixer/Public/Quartz/QuartzSubsystem.h
@@ -53,6 +53,7 @@ public:
 	// Helper functions for initializing quantized command initialization struct (to consolidate eyesore)
 	static Audio::FQuartzQuantizedRequestData CreateDataForTickRateChange(UQuartzClockHandle* InClockHandle, const FOnQuartzCommandEventBP& InDelegate, const Audio::FQuartzClockTickRate& InNewTickRate, const FQuartzQuantizationBoundary& InQuantizationBoundary);
 	static Audio::FQuartzQuantizedRequestData CreateDataForTransportReset(UQuartzClockHandle* InClockHandle, const FQuartzQuantizationBoundary& InQuantizationBoundary, const FOnQuartzCommandEventBP& InDelegate);
+	static Audio::FQuartzQuantizedRequestData CreateDataForStartOtherClock(UQuartzClockHandle* InClockHandle, FName InClockToStart, const FQuartzQuantizationBoundary& InQuantizationBoundary, const FOnQuartzCommandEventBP& InDelegate);
 	static Audio::FQuartzQuantizedRequestData CreateDataDataForSchedulePlaySound(UQuartzClockHandle* InClockHandle, const FOnQuartzCommandEventBP& InDelegate, const FQuartzQuantizationBoundary& InQuantizationBoundary);
 
 	UFUNCTION(BlueprintCallable, Category = "Quartz Clock Handle")
diff --git a/Engine/Source/Runtime/Engine/Classes/Sound/QuartzQuantizationUtilities.cpp b/Engine/Source/Runtime/Engine/Classes/Sound/QuartzQuantizationUtilities.cpp
--- a/Engine/Source/Runtime/Engine/Classes/Sound/QuartzQuantizationUtilities.cpp
+++ b/Engine/Source/Runtime/Engine/Classes/Sound/QuartzQuantizationUtilities.cpp
@@ -355,6 +355,7 @@ namespace Audio
 	)
 		: ClockName(RHS.ClockName)
 		, ClockHandleName(RHS.ClockHandleName)
+		, OtherClockName(RHS.OtherClockName)
 		, QuantizedCommandPtr(RHS.QuantizedCommandPtr)
 		, QuantizationBoundary(RHS.QuantizationBoundary)
 		, GameThreadCommandQueue(RHS.GameThreadCommandQueue)
diff --git a/Engine/Source/Runtime/Engine/Classes/Sound/QuartzQuantizationUtilities.h b/Engine/Source/Runtime/Engine/Classes/Sound/QuartzQuantizationUtilities.h
--- a/Engine/Source/Runtime/Engine/Classes/Sound/QuartzQuantizationUtilities.h
+++ b/Engine/Source/Runtime/Engine/Classes/Sound/QuartzQuantizationUtilities.h
@@ -406,6 +406,7 @@ namespace Audio
 		// shared with FQuartzQuantizedCommandInitInfo:
 		FName ClockName;
 		FName ClockHandleName;
+		FName OtherClockName;
 		TSharedPtr<IQuartzQuantizedCommand> QuantizedCommandPtr;
 		FQuartzQuantizationBoundary QuantizationBoundary{ /* InQuantization */ EQuartzCommandQuantization::Tick, /* InMultiplier */ 1.f };
 		TSharedPtr<FShareableQuartzCommandQueue, ESPMode::ThreadSafe> GameThreadCommandQueue{ nullptr };
@@ -434,6 +435,7 @@ namespace Audio
 		// shared with FQuartzQuantizedRequestData
 		FName ClockName;
 		FName ClockHandleName;
+		FName OtherClockName;
 		TSharedPtr<IQuartzQuantizedCommand> QuantizedCommandPtr;
 		FQuartzQuantizationBoundary QuantizationBoundary;
 		TSharedPtr<FShareableQuartzCommandQueue, ESPMode::ThreadSafe> GameThreadCommandQueue;
