removing invalid check in ~FQuartzClockManager()

diff --git a/Engine/Source/Runtime/AudioMixer/Private/Quartz/AudioMixerClockManager.cpp b/Engine/Source/Runtime/AudioMixer/Private/Quartz/AudioMixerClockManager.cpp
--- a/Engine/Source/Runtime/AudioMixer/Private/Quartz/AudioMixerClockManager.cpp
+++ b/Engine/Source/Runtime/AudioMixer/Private/Quartz/AudioMixerClockManager.cpp
@@ -13,7 +13,6 @@ namespace Audio
 
 	FQuartzClockManager::~FQuartzClockManager()
 	{
-		check(ActiveClocks.Num() == 0);
 	}
 
 	void FQuartzClockManager::Update(int32 NumFramesUntilNextUpdate)
