Fix for debug-only hang in Quartz

diff --git a/Engine/Source/Runtime/AudioMixer/Private/Quartz/AudioMixerClockManager.cpp b/Engine/Source/Runtime/AudioMixer/Private/Quartz/AudioMixerClockManager.cpp
--- a/Engine/Source/Runtime/AudioMixer/Private/Quartz/AudioMixerClockManager.cpp
+++ b/Engine/Source/Runtime/AudioMixer/Private/Quartz/AudioMixerClockManager.cpp
@@ -413,9 +413,7 @@ namespace Audio
 
 	TSharedPtr<FQuartzClock> FQuartzClockManager::FindClock(const FName& InName)
 	{
-		// the function calling this should be be on the audio rendering thread or have acquired the lock
-		checkSlow(MixerDevice->IsAudioRenderingThread() || ActiveClockCritSec.TryLock());
-
+		FScopeLock Lock(&ActiveClockCritSec);
 		for (auto& Clock : ActiveClocks)
 		{
 			if (Clock->GetName() == InName)
