[Quartz] Adding UPROPERTY(Transient) to cached World and Subsystem pointers for safety during cleanup

diff --git a/Engine/Source/Runtime/AudioMixer/Public/Quartz/AudioMixerClockHandle.h b/Engine/Source/Runtime/AudioMixer/Public/Quartz/AudioMixerClockHandle.h
--- a/Engine/Source/Runtime/AudioMixer/Public/Quartz/AudioMixerClockHandle.h
+++ b/Engine/Source/Runtime/AudioMixer/Public/Quartz/AudioMixerClockHandle.h
@@ -159,6 +159,7 @@ private:
 
 	MetronomeDelegateGameThreadData MetronomeDelegates[static_cast<int32>(EQuartzCommandQuantization::Count)];
 
+	UPROPERTY(Transient)
 	UQuartzSubsystem* QuartzSubsystem;
 
 	FName ClockHandleId;
@@ -167,6 +168,7 @@ private:
 
 	bool bConnectedToClock{ false };
 
+	UPROPERTY(Transient)
 	UWorld* WorldPtr{ nullptr };
 
 }; // class UClockHandleBlueprintLibrary
