Edigrate Dynamics Processor fixes to UE5 Early Access

diff --git a/Engine/Source/Runtime/AudioMixer/Private/Effects/AudioMixerSubmixEffectDynamicsProcessor.cpp b/Engine/Source/Runtime/AudioMixer/Private/Effects/AudioMixerSubmixEffectDynamicsProcessor.cpp
--- a/Engine/Source/Runtime/AudioMixer/Private/Effects/AudioMixerSubmixEffectDynamicsProcessor.cpp
+++ b/Engine/Source/Runtime/AudioMixer/Private/Effects/AudioMixerSubmixEffectDynamicsProcessor.cpp
@@ -250,11 +250,23 @@ void FSubmixEffectDynamicsProcessor::OnProcessAudio(const FSoundEffectSubmixInpu
 		return;
 	}
 
-	const int32 NumKeyChannels = DynamicsProcessor.GetKeyNumChannels();
-	const int32 NumKeySamples = InData.NumFrames * NumKeyChannels;
+	int32 NumKeyChannels = DynamicsProcessor.GetKeyNumChannels();
+	int32 NumKeySamples = InData.NumFrames * NumKeyChannels;
 
 	AudioExternal.Reset();
-	if (KeySource.GetType() != ESubmixEffectDynamicsKeySource::Default)
+	
+	// If set to default, enforce num key channels to always be number of input channels.
+	// If either unset or KeySource was changed between frames back to 'Default', NumKeyChannels
+	// could be stale or left as initialized number of scratch channels.
+	if (KeySource.GetType() == ESubmixEffectDynamicsKeySource::Default)
+	{
+		if (InData.NumChannels != NumKeyChannels)
+			NumKeyChannels = InData.NumChannels;
+			NumKeySamples = InData.NumFrames * NumKeyChannels;
+			DynamicsProcessor.SetKeyNumChannels(NumKeyChannels);
+		}
+	}
+	else
 	{
 		AudioExternal.AddZeroed(NumKeySamples);
 	}
