Fix bad robomerge

diff --git a/Engine/Source/Runtime/AudioMixer/Private/Effects/AudioMixerSubmixEffectDynamicsProcessor.cpp b/Engine/Source/Runtime/AudioMixer/Private/Effects/AudioMixerSubmixEffectDynamicsProcessor.cpp
--- a/Engine/Source/Runtime/AudioMixer/Private/Effects/AudioMixerSubmixEffectDynamicsProcessor.cpp
+++ b/Engine/Source/Runtime/AudioMixer/Private/Effects/AudioMixerSubmixEffectDynamicsProcessor.cpp
@@ -261,6 +261,7 @@ void FSubmixEffectDynamicsProcessor::OnProcessAudio(const FSoundEffectSubmixInpu
 	if (KeySource.GetType() == ESubmixEffectDynamicsKeySource::Default)
 	{
 		if (InData.NumChannels != NumKeyChannels)
+		{
 			NumKeyChannels = InData.NumChannels;
 			NumKeySamples = InData.NumFrames * NumKeyChannels;
 			DynamicsProcessor.SetKeyNumChannels(NumKeyChannels);
