Fix Compile warnings in AudioMixerSource.cpp

diff --git a/Engine/Source/Runtime/AudioMixer/Private/AudioMixerSource.cpp b/Engine/Source/Runtime/AudioMixer/Private/AudioMixerSource.cpp
--- a/Engine/Source/Runtime/AudioMixer/Private/AudioMixerSource.cpp
+++ b/Engine/Source/Runtime/AudioMixer/Private/AudioMixerSource.cpp
@@ -303,7 +303,11 @@ namespace Audio
 
 		FSoundSource::InitCommon();
 
-		check(WaveInstance);
+
+		if (!ensure(InWaveInstance))
+		{
+			return false;
+		}
 
 		USoundWave* WaveData = WaveInstance->WaveData;
 		check(WaveData);
@@ -380,10 +384,11 @@ namespace Audio
 			{
 				UE_LOG(LogAudioMixer, Warning, TEXT("Sound wave %s was flagged as being ambisonics but had a channel count of %d. Currently the audio engine only supports FOA sources that have four channels."), *InWaveInstance->GetName(), WaveData->NumChannels);
 			}
-
-			InitParams.AudioComponentUserID = WaveInstance->ActiveSound->GetAudioComponentUserID();
-
-			InitParams.AudioComponentID = WaveInstance->ActiveSound->GetAudioComponentID();
+			if (ActiveSound)
+			{
+				InitParams.AudioComponentUserID = WaveInstance->ActiveSound->GetAudioComponentUserID();
+				InitParams.AudioComponentID = WaveInstance->ActiveSound->GetAudioComponentID();
+			}
 
 			InitParams.EnvelopeFollowerAttackTime = WaveInstance->EnvelopeFollowerAttackTime;
 			InitParams.EnvelopeFollowerReleaseTime = WaveInstance->EnvelopeFollowerReleaseTime;
@@ -744,6 +749,11 @@ namespace Audio
 	{
 		LLM_SCOPE(ELLMTag::AudioMixer);
 
+		if (!ensure(InWaveInstance))
+		{
+			return false;
+		}
+
 		// We are currently not supporting playing audio on a controller
 		if (InWaveInstance->OutputTarget == EAudioOutputTarget::Controller)
 		{
@@ -807,9 +817,11 @@ namespace Audio
 
 		// Active sound instance ID is the audio component ID of active sound.
 		uint64 InstanceID = 0;
+		bool bActiveSoundIsPreviewSound = false;
 		if (WaveInstance->ActiveSound)
 		{
 			InstanceID = WaveInstance->ActiveSound->GetAudioComponentID();
+			bActiveSoundIsPreviewSound = WaveInstance->ActiveSound->bIsPreviewSound;
 		}
 
 		FMixerSourceBufferInitArgs BufferInitArgs;
@@ -819,7 +831,7 @@ namespace Audio
 		BufferInitArgs.SoundWave = &SoundWave;
 		BufferInitArgs.LoopingMode = InWaveInstance->LoopingMode;
 		BufferInitArgs.bIsSeeking = bIsSeeking;
-		BufferInitArgs.bIsPreviewSound = WaveInstance->ActiveSound->bIsPreviewSound;
+		BufferInitArgs.bIsPreviewSound = bActiveSoundIsPreviewSound;
 
 		MixerSourceBuffer = FMixerSourceBuffer::Create(BufferInitArgs);
 		
