Default EQ Submix off for 5.0 to fix silent output issue in packaged builds

diff --git a/Engine/Source/Runtime/AudioMixer/Private/AudioMixerDevice.cpp b/Engine/Source/Runtime/AudioMixer/Private/AudioMixerDevice.cpp
--- a/Engine/Source/Runtime/AudioMixer/Private/AudioMixerDevice.cpp
+++ b/Engine/Source/Runtime/AudioMixer/Private/AudioMixerDevice.cpp
@@ -29,11 +29,11 @@
 #endif // WITH_EDITOR
 
 
-static int32 DisableSubmixEffectEQCvar = 0;
+static int32 DisableSubmixEffectEQCvar = 1;
 FAutoConsoleVariableRef CVarDisableSubmixEQ(
 	TEXT("au.DisableSubmixEffectEQ"),
 	DisableSubmixEffectEQCvar,
-	TEXT("Disables the eq submix.\n")
+	TEXT("Disables the eq submix (true by default as of 5.0).\n")
 	TEXT("0: Not Disabled, 1: Disabled"),
 	ECVF_Default);
 
@@ -1637,16 +1637,26 @@ namespace Audio
 
 	FMixerSubmixPtr FMixerDevice::FindSubmixInstanceByObjectId(uint32 InObjectId)
 	{
-		for (const USoundSubmix* MasterSubmix : MasterSubmixes)
+		for (int32 i = 0; i < MasterSubmixes.Num(); i++)
 		{
-			if (!ensure(MasterSubmix))
+			if (const USoundSubmix* MasterSubmix = MasterSubmixes[i])
 			{
-				continue;
+				if (MasterSubmix->GetUniqueID() == InObjectId)
+				{
+					return GetMasterSubmixInstance(MasterSubmix);
+				}
 			}
-
-			if (MasterSubmix->GetUniqueID() == InObjectId)
+			else
 			{
-				return GetMasterSubmixInstance(MasterSubmix);
+				const EMasterSubmixType::Type SubmixType = static_cast<EMasterSubmixType::Type>(i);
+				ensureAlwaysMsgf(EMasterSubmixType::Master != SubmixType,
+					TEXT("Top-level master has to be registered before anything else, and is required for the lifetime of the application.")
+				);
+
+				if (!DisableSubmixEffectEQCvar && EMasterSubmixType::EQ == SubmixType)
+				{
+					UE_LOG(LogAudioMixer, Warning, TEXT("Failed to query EQ Submix when it was expected to be loaded."));
+				}
 			}
 		}
 
