Move analyzer callsite to use _GameThread version of FMixerDevice::AddPatchForAudioBus, and revert to original implementation of FMixerDevice::AddPatchForAudioBus as all original callers are from the Audio Render Thread

diff --git a/Engine/Source/Runtime/AudioAnalyzer/Private/AudioAnalyzer.cpp b/Engine/Source/Runtime/AudioAnalyzer/Private/AudioAnalyzer.cpp
--- a/Engine/Source/Runtime/AudioAnalyzer/Private/AudioAnalyzer.cpp
+++ b/Engine/Source/Runtime/AudioAnalyzer/Private/AudioAnalyzer.cpp
@@ -76,7 +76,7 @@ void UAudioAnalyzer::StartAnalyzing(UWorld* InWorld, UAudioBus* AudioBusToAnalyz
 	MixerDevice->StartAudioBus(AudioBusId, NumChannels, false);
 
 	// Get an output patch for the audio bus
-	PatchOutputStrongPtr = MixerDevice->AddPatchForAudioBus(AudioBusId);
+	PatchOutputStrongPtr = MixerDevice->AddPatchForAudioBus_GameThread(AudioBusId);
 
 	// Register this audio analyzer with the audio analyzer subsystem
 	// The subsystem will query this analyzer to see if it has enough audio to perform analysis.
diff --git a/Engine/Source/Runtime/AudioMixer/Private/AudioMixerDevice.cpp b/Engine/Source/Runtime/AudioMixer/Private/AudioMixerDevice.cpp
--- a/Engine/Source/Runtime/AudioMixer/Private/AudioMixerDevice.cpp
+++ b/Engine/Source/Runtime/AudioMixer/Private/AudioMixerDevice.cpp
@@ -2419,8 +2419,26 @@ namespace Audio
 	}
 
 	Audio::FPatchOutputStrongPtr FMixerDevice::AddPatchForAudioBus(uint32 InAudioBusId, float InPatchGain)
+	{
+		check(SourceManager);
+
+		const int32 NumChannels = SourceManager->GetAudioBusNumChannels(InAudioBusId);
+		if (NumChannels > 0)
+		{
+			const int32 NumOutputFrames = SourceManager->GetNumOutputFrames();
+			FPatchOutputStrongPtr StrongOutputPtr = MakeShareable(new FPatchOutput(NumOutputFrames * NumChannels, InPatchGain));
+			SourceManager->AddPatchOutputForAudioBus(InAudioBusId, StrongOutputPtr);
+			return StrongOutputPtr;
+		}
+
+		UE_LOG(LogAudioMixer, Warning, TEXT("Failed to add a patch output for audio bus: audio bus with id '%d' is inactive."), InAudioBusId);
+		return nullptr;
+	}
+
+	Audio::FPatchOutputStrongPtr FMixerDevice::AddPatchForAudioBus_GameThread(uint32 InAudioBusId, float InPatchGain)
 	{
 		FPatchOutputStrongPtr StrongOutputPtr;
+		
 		if (IsInGameThread())
 		{
 			FActiveBusData* BusData = ActiveAudioBuses_GameThread.Find(InAudioBusId);
@@ -2432,7 +2450,7 @@ namespace Audio
 
 				FAudioThread::RunCommandOnAudioThread([this, InAudioBusId, StrongOutputPtr]() mutable
 				{
-					SourceManager->AddPatchOutputForAudioBus(InAudioBusId, StrongOutputPtr);
+					SourceManager->AddPatchOutputForAudioBus_AudioThread(InAudioBusId, StrongOutputPtr);
 				});
 			}
 			else
diff --git a/Engine/Source/Runtime/AudioMixer/Private/AudioMixerSourceManager.cpp b/Engine/Source/Runtime/AudioMixer/Private/AudioMixerSourceManager.cpp
--- a/Engine/Source/Runtime/AudioMixer/Private/AudioMixerSourceManager.cpp
+++ b/Engine/Source/Runtime/AudioMixer/Private/AudioMixerSourceManager.cpp
@@ -1060,33 +1060,39 @@ namespace Audio
 		});
 	}
 
-	bool FMixerSourceManager::IsAudioBusActive(uint32 InAudioBusId)
+	bool FMixerSourceManager::IsAudioBusActive(uint32 InAudioBusId) const
 	{
 		AUDIO_MIXER_CHECK_GAME_THREAD(MixerDevice);
 		return AudioBusIds_AudioThread.Contains(InAudioBusId);
 	}
 
+	int32 FMixerSourceManager::GetAudioBusNumChannels(uint32 InAudioBusId) const
+	{
+		AUDIO_MIXER_CHECK_AUDIO_PLAT_THREAD(MixerDevice);
+		TSharedPtr<FMixerAudioBus> AudioBusPtr = AudioBuses.FindRef(InAudioBusId);
+		if (AudioBusPtr.IsValid())
+		{
+			return AudioBusPtr->GetNumChannels();
+		}
+
+		return 0;
+	}
+
 	void FMixerSourceManager::AddPatchOutputForAudioBus(uint32 InAudioBusId, FPatchOutputStrongPtr& InPatchOutputStrongPtr)
 	{
-		AudioMixerThreadCommand([this, InAudioBusId, InPatchOutputStrongPtr]() mutable
+		AUDIO_MIXER_CHECK_AUDIO_PLAT_THREAD(MixerDevice);
+		TSharedPtr<FMixerAudioBus> AudioBusPtr = AudioBuses.FindRef(InAudioBusId);
+		if (AudioBusPtr.IsValid())
 		{
-			TSharedPtr<FMixerAudioBus> AudioBusPtr = AudioBuses.FindRef(InAudioBusId);
-			if (AudioBusPtr.IsValid())
-			{
-				AudioBusPtr->AddNewPatchOutput(InPatchOutputStrongPtr);
-			}
-		});
+			AudioBusPtr->AddNewPatchOutput(InPatchOutputStrongPtr);
+		}
 	}
 
-	void FMixerSourceManager::AddPatchInputForAudioBus(uint32 InAudioBusId, FPatchInput& InPatchInput)
+	void FMixerSourceManager::AddPatchOutputForAudioBus_AudioThread(uint32 InAudioBusId, FPatchOutputStrongPtr& InPatchOutputStrongPtr)
 	{
-		AudioMixerThreadCommand([this, InAudioBusId, InPatchInput]() mutable
+		AudioMixerThreadCommand([this, InAudioBusId, InPatchOutputStrongPtr]() mutable
 		{
-			TSharedPtr<FMixerAudioBus> AudioBusPtr = AudioBuses.FindRef(InAudioBusId);
-			if (AudioBusPtr.IsValid())
-			{
-				AudioBusPtr->AddNewPatchInput(InPatchInput);
-			}
+			AddPatchOutputForAudioBus(InAudioBusId, InPatchOutputStrongPtr);
 		});
 	}
 
diff --git a/Engine/Source/Runtime/AudioMixer/Private/AudioMixerSourceManager.h b/Engine/Source/Runtime/AudioMixer/Private/AudioMixerSourceManager.h
--- a/Engine/Source/Runtime/AudioMixer/Private/AudioMixerSourceManager.h
+++ b/Engine/Source/Runtime/AudioMixer/Private/AudioMixerSourceManager.h
@@ -174,11 +174,18 @@ namespace Audio
 		void StopAudioBus(uint32 InAudioBusId);
 
 		// Queries if an audio bus is active
-		bool IsAudioBusActive(uint32 InAudioBusId);
+		bool IsAudioBusActive(uint32 InAudioBusId) const;
 
-		// Adds a patch output for an audio bus
+		// Returns the number of channels currently set for the audio bus associated with
+		// the provided BusId.  Returns 0 if the audio bus is inactive.
+		int32 GetAudioBusNumChannels(uint32 InAudioBusId) const;
+
+		// Adds a patch output for an audio bus from the Audio Render Thread
 		void AddPatchOutputForAudioBus(uint32 InAudioBusId, FPatchOutputStrongPtr& InPatchOutputStrongPtr);
 
+		// Adds a patch output for an audio bus from the Audio Thread
+		void AddPatchOutputForAudioBus_AudioThread(uint32 InAudioBusId, FPatchOutputStrongPtr& InPatchOutputStrongPtr);
+
 		// Adds a patch input for an audio bus
 		void AddPatchInputForAudioBus(uint32 InAudioBusId, FPatchInput& InPatchInput);
 
diff --git a/Engine/Source/Runtime/AudioMixer/Public/AudioMixerDevice.h b/Engine/Source/Runtime/AudioMixer/Public/AudioMixerDevice.h
--- a/Engine/Source/Runtime/AudioMixer/Public/AudioMixerDevice.h
+++ b/Engine/Source/Runtime/AudioMixer/Public/AudioMixerDevice.h
@@ -276,7 +276,9 @@ namespace Audio
 		void StartAudioBus(uint32 InAudioBusId, int32 InNumChannels, bool bInIsAutomatic);
 		void StopAudioBus(uint32 InAudioBusId);
 		bool IsAudioBusActive(uint32 InAudioBusId);
+
 		FPatchOutputStrongPtr AddPatchForAudioBus(uint32 InAudioBusId, float InPatchGain = 1.0f);
+		FPatchOutputStrongPtr AddPatchForAudioBus_GameThread(uint32 InAudioBusId, float InPatchGain = 1.0f);
 
 		// Clock Manager for quantized event handling on Audio Render Thread
 		FQuartzClockManager QuantizedEventClockManager;
