Fix log spam for missing AudioBus when attempting to connect via DynamicsProcessor by moving to caller, as it can check if it was provided a valid patch and warn accordingly (and throttle as desired).
#rnx

diff --git a/Engine/Source/Runtime/AudioMixer/Classes/SubmixEffects/AudioMixerSubmixEffectDynamicsProcessor.h b/Engine/Source/Runtime/AudioMixer/Classes/SubmixEffects/AudioMixerSubmixEffectDynamicsProcessor.h
--- a/Engine/Source/Runtime/AudioMixer/Classes/SubmixEffects/AudioMixerSubmixEffectDynamicsProcessor.h
+++ b/Engine/Source/Runtime/AudioMixer/Classes/SubmixEffects/AudioMixerSubmixEffectDynamicsProcessor.h
@@ -225,7 +225,7 @@ protected:
 	void OnDeviceCreated(Audio::FDeviceId InDeviceId);
 	void OnDeviceDestroyed(Audio::FDeviceId InDeviceId);
 	
-	Audio::AlignedFloatBuffer AudioExternal;
+	Audio::FAlignedFloatBuffer AudioExternal;
 
 	TArray<float> AudioKeyFrame;
 	TArray<float> AudioInputFrame;
@@ -240,6 +240,7 @@ private:
 		ESubmixEffectDynamicsKeySource Type = ESubmixEffectDynamicsKeySource::Default;
 		int32 NumChannels = 0;
 		uint32 ObjectId = INDEX_NONE;
+		bool bReportInactive = true;
 
 		mutable FCriticalSection MutateSourceCritSection;
 
@@ -255,6 +256,7 @@ private:
 				NumChannels = 0;
 				ObjectId = INDEX_NONE;
 				Type = ESubmixEffectDynamicsKeySource::Default;
+				bReportInactive = true;
 			}
 		}
 
@@ -270,6 +272,18 @@ private:
 			return NumChannels;
 		}
 
+		void SetReportInactive(bool bInReportInactive)
+		{
+			const FScopeLock ScopeLock(&MutateSourceCritSection);
+			bReportInactive = bInReportInactive;
+		}
+
+		bool ShouldReportInactive() const
+		{
+			const FScopeLock ScopeLock(&MutateSourceCritSection);
+			return bReportInactive;
+		}
+
 		ESubmixEffectDynamicsKeySource GetType() const
 		{
 			const FScopeLock ScopeLock(&MutateSourceCritSection);
@@ -295,6 +309,7 @@ private:
 					NumChannels = InNumChannels;
 
 					bResetPatch = true;
+					bReportInactive = true;
 				}
 			}
 
diff --git a/Engine/Source/Runtime/AudioMixer/Private/AudioMixerDevice.cpp b/Engine/Source/Runtime/AudioMixer/Private/AudioMixerDevice.cpp
--- a/Engine/Source/Runtime/AudioMixer/Private/AudioMixerDevice.cpp
+++ b/Engine/Source/Runtime/AudioMixer/Private/AudioMixerDevice.cpp
@@ -152,7 +152,6 @@ namespace Audio
 			SubmixMap.Reset();
 		}
 	}
-	
 
 	FMixerDevice::FMixerDevice(IAudioMixerPlatformInterface* InAudioMixerPlatform)
 		: QuantizedEventClockManager(this)
@@ -2407,7 +2406,7 @@ namespace Audio
 		}
 	}
 
-	bool FMixerDevice::IsAudioBusActive(uint32 InAudioBusId)
+	bool FMixerDevice::IsAudioBusActive(uint32 InAudioBusId) const
 	{
 		if (IsInGameThread())
 		{
@@ -2431,7 +2430,6 @@ namespace Audio
 			return StrongOutputPtr;
 		}
 
-		UE_LOG(LogAudioMixer, Warning, TEXT("Failed to add a patch output for audio bus: audio bus with id '%d' is inactive."), InAudioBusId);
 		return nullptr;
 	}
 
diff --git a/Engine/Source/Runtime/AudioMixer/Private/AudioMixerSourceManager.h b/Engine/Source/Runtime/AudioMixer/Private/AudioMixerSourceManager.h
--- a/Engine/Source/Runtime/AudioMixer/Private/AudioMixerSourceManager.h
+++ b/Engine/Source/Runtime/AudioMixer/Private/AudioMixerSourceManager.h
@@ -173,7 +173,7 @@ namespace Audio
 		// Stops an audio bus manually
 		void StopAudioBus(uint32 InAudioBusId);
 
-		// Queries if an audio bus is active
+		// Queries if an audio bus is active. Must be called from the audio thread.
 		bool IsAudioBusActive(uint32 InAudioBusId) const;
 
 		// Returns the number of channels currently set for the audio bus associated with
diff --git a/Engine/Source/Runtime/AudioMixer/Private/Effects/AudioMixerSubmixEffectDynamicsProcessor.cpp b/Engine/Source/Runtime/AudioMixer/Private/Effects/AudioMixerSubmixEffectDynamicsProcessor.cpp
--- a/Engine/Source/Runtime/AudioMixer/Private/Effects/AudioMixerSubmixEffectDynamicsProcessor.cpp
+++ b/Engine/Source/Runtime/AudioMixer/Private/Effects/AudioMixerSubmixEffectDynamicsProcessor.cpp
@@ -189,11 +189,20 @@ bool FSubmixEffectDynamicsProcessor::UpdateKeySourcePatch()
 			// should never be hit during Teardown.
 			if (Audio::FMixerDevice* MixerDevice = GetMixerDevice())
 			{
-				KeySource.Patch = MixerDevice->AddPatchForAudioBus(KeySource.GetObjectId(), 1.0f /* PatchGain */);
-				if (KeySource.Patch.IsValid())
+				const uint32 ObjectId = KeySource.GetObjectId();
+				if (ObjectId != INDEX_NONE)
 				{
-					DynamicsProcessor.SetKeyNumChannels(KeySource.GetNumChannels());
-					return true;
+					KeySource.Patch = MixerDevice->AddPatchForAudioBus(ObjectId, 1.0f /* PatchGain */);
+					if (KeySource.Patch.IsValid())
+					{
+						DynamicsProcessor.SetKeyNumChannels(KeySource.GetNumChannels());
+						return true;
+					}
+					else if (KeySource.ShouldReportInactive())
+					{
+						KeySource.SetReportInactive(false);
+						UE_LOG(LogAudioMixer, Warning, TEXT("DynamicsProcessor failed to add patch output to inactive AudioBus (ID: '%d')"), ObjectId);
+					}
 				}
 			}
 		}
@@ -207,18 +216,27 @@ bool FSubmixEffectDynamicsProcessor::UpdateKeySourcePatch()
 			// should never be hit during Teardown.
 			if (Audio::FMixerDevice* MixerDevice = GetMixerDevice())
 			{
-				KeySource.Patch = MixerDevice->AddPatchForSubmix(KeySource.GetObjectId(), 1.0f /* PatchGain */);
-				if (KeySource.Patch.IsValid())
+				const uint32 ObjectId = KeySource.GetObjectId();
+				if (ObjectId != INDEX_NONE)
 				{
-					Audio::FMixerSubmixPtr SubmixPtr = MixerDevice->FindSubmixInstanceByObjectId(KeySource.GetObjectId());
-					if (SubmixPtr.IsValid())
+					KeySource.Patch = MixerDevice->AddPatchForSubmix(ObjectId, 1.0f /* PatchGain */);
+					if (KeySource.Patch.IsValid())
 					{
-						const int32 SubmixNumChannels = SubmixPtr->GetNumOutputChannels();
-						KeySource.SetNumChannels(SubmixNumChannels);
-						DynamicsProcessor.SetKeyNumChannels(SubmixNumChannels);
-						return true;
+						Audio::FMixerSubmixPtr SubmixPtr = MixerDevice->FindSubmixInstanceByObjectId(KeySource.GetObjectId());
+						if (SubmixPtr.IsValid())
+						{
+							const int32 SubmixNumChannels = SubmixPtr->GetNumOutputChannels();
+							KeySource.SetNumChannels(SubmixNumChannels);
+							DynamicsProcessor.SetKeyNumChannels(SubmixNumChannels);
+							return true;
+						}
 					}
-				}
+					else if (KeySource.ShouldReportInactive())
+					{
+						KeySource.SetReportInactive(false);
+						UE_LOG(LogAudioMixer, Warning, TEXT("DynamicsProcessor failed to add patch output to inactive Submix (ID: '%d')"), ObjectId);
+					}
+				}	
 			}
 		}
 		break;
@@ -240,8 +258,8 @@ void FSubmixEffectDynamicsProcessor::OnProcessAudio(const FSoundEffectSubmixInpu
 
 	ensure(InData.NumChannels == OutData.NumChannels);
 
-	const Audio::AlignedFloatBuffer& InBuffer = *InData.AudioBuffer;
-	Audio::AlignedFloatBuffer& OutBuffer = *OutData.AudioBuffer;
+	const Audio::FAlignedFloatBuffer& InBuffer = *InData.AudioBuffer;
+	Audio::FAlignedFloatBuffer& OutBuffer = *OutData.AudioBuffer;
 
 	const bool bBypassDueToInvalidChannelCount = !ensure(InData.NumChannels <= AudioInputFrame.Num());
 	if (bBypassDueToInvalidChannelCount || bBypassSubmixDynamicsProcessor || bBypass)
diff --git a/Engine/Source/Runtime/AudioMixer/Public/AudioMixerDevice.h b/Engine/Source/Runtime/AudioMixer/Public/AudioMixerDevice.h
--- a/Engine/Source/Runtime/AudioMixer/Public/AudioMixerDevice.h
+++ b/Engine/Source/Runtime/AudioMixer/Public/AudioMixerDevice.h
@@ -279,7 +279,7 @@ namespace Audio
 		// Audio bus API
 		void StartAudioBus(uint32 InAudioBusId, int32 InNumChannels, bool bInIsAutomatic);
 		void StopAudioBus(uint32 InAudioBusId);
-		bool IsAudioBusActive(uint32 InAudioBusId);
+		bool IsAudioBusActive(uint32 InAudioBusId) const;
 
 		FPatchOutputStrongPtr AddPatchForAudioBus(uint32 InAudioBusId, float InPatchGain = 1.0f);
 		FPatchOutputStrongPtr AddPatchForAudioBus_GameThread(uint32 InAudioBusId, float InPatchGain = 1.0f);
