Fix audio stall caused by an ensure triggered by an audio task not calling FTaskTagScope Scope(ETaskTag::EAudioThread);

diff --git a/Engine/Source/Runtime/AudioMixer/Private/AudioMixerSourceManager.cpp b/Engine/Source/Runtime/AudioMixer/Private/AudioMixerSourceManager.cpp
--- a/Engine/Source/Runtime/AudioMixer/Private/AudioMixerSourceManager.cpp
+++ b/Engine/Source/Runtime/AudioMixer/Private/AudioMixerSourceManager.cpp
@@ -636,6 +636,7 @@ namespace Audio
 			// Unregister these source effect instances from their owning USoundEffectInstance on the next audio thread tick.
 			FAudioThread::RunCommandOnAudioThread([SourceEffects = MoveTemp(SourceInfo.SourceEffects)]() mutable
 			{
+				FTaskTagScope Scope(ETaskTag::EAudioThread);
 				for (int32 i = 0; i < SourceEffects.Num(); ++i)
 				{
 					USoundEffectPreset::UnregisterInstance(SourceEffects[i]);
