#ue5 - Backing out robo merged CL 16212946 as it was causing an assert/crash after every PIE run and when closing the editor in both Frosty and shootergame.

diff --git a/Engine/Source/Runtime/AudioMixer/Private/AudioMixerSourceManager.cpp b/Engine/Source/Runtime/AudioMixer/Private/AudioMixerSourceManager.cpp
--- a/Engine/Source/Runtime/AudioMixer/Private/AudioMixerSourceManager.cpp
+++ b/Engine/Source/Runtime/AudioMixer/Private/AudioMixerSourceManager.cpp
@@ -636,7 +636,6 @@ namespace Audio
 			// Unregister these source effect instances from their owning USoundEffectInstance on the next audio thread tick.
 			FAudioThread::RunCommandOnAudioThread([SourceEffects = MoveTemp(SourceInfo.SourceEffects)]() mutable
 			{
-				FTaskTagScope Scope(ETaskTag::EAudioThread);
 				for (int32 i = 0; i < SourceEffects.Num(); ++i)
 				{
 					USoundEffectPreset::UnregisterInstance(SourceEffects[i]);
