Forward all requests for resetting/unregistering effects to Game Thread prior to processing on AudioThread to avoid GC race condition and ensure about being called from incorrect thread

diff --git a/Engine/Source/Runtime/AudioMixer/Private/AudioMixerSourceManager.cpp b/Engine/Source/Runtime/AudioMixer/Private/AudioMixerSourceManager.cpp
--- a/Engine/Source/Runtime/AudioMixer/Private/AudioMixerSourceManager.cpp
+++ b/Engine/Source/Runtime/AudioMixer/Private/AudioMixerSourceManager.cpp
@@ -633,16 +633,24 @@ namespace Audio
 		{
 			FSourceInfo& SourceInfo = SourceInfos[SourceId];
 
-			// Unregister these source effect instances from their owning USoundEffectInstance on the next audio thread tick.
-			FAudioThread::RunCommandOnAudioThread([SourceEffects = MoveTemp(SourceInfo.SourceEffects)]() mutable
+			// Unregister these source effect instances from their owning USoundEffectInstance on the audio thread.
+			// Have to pass to Game Thread prior to processing on AudioThread to avoid race condition with GC.
+			// (RunCommandOnAudioThread is not safe to call from any thread other than the GameThread).
+			if (!SourceInfo.SourceEffects.IsEmpty())
 			{
-				for (int32 i = 0; i < SourceEffects.Num(); ++i)
+				AsyncTask(ENamedThreads::GameThread, [GTSourceEffects = MoveTemp(SourceInfo.SourceEffects)]() mutable
 				{
-					USoundEffectPreset::UnregisterInstance(SourceEffects[i]);
-				}
-			});
+					FAudioThread::RunCommandOnAudioThread([ATSourceEffects = MoveTemp(GTSourceEffects)]() mutable
+					{
+						for (const TSoundEffectSourcePtr& EffectPtr : ATSourceEffects)
+						{
+							USoundEffectPreset::UnregisterInstance(EffectPtr);
+						}
+					});
+				});
 
-			SourceInfo.SourceEffects.Reset();
+				SourceInfo.SourceEffects.Reset();
+			}
 			SourceInfo.SourceEffectPresets.Reset();
 		}
 	}
diff --git a/Engine/Source/Runtime/AudioMixer/Private/AudioMixerSubmix.cpp b/Engine/Source/Runtime/AudioMixer/Private/AudioMixerSubmix.cpp
--- a/Engine/Source/Runtime/AudioMixer/Private/AudioMixerSubmix.cpp
+++ b/Engine/Source/Runtime/AudioMixer/Private/AudioMixerSubmix.cpp
@@ -594,14 +594,22 @@ namespace Audio
 			FadeInfo.EffectChain.Reset();
 		}
 
-		// Unregister these source effect instances from their owning USoundEffectInstance on the next audio thread tick.
-		FAudioThread::RunCommandOnAudioThread([SubmixEffects = MoveTemp(SubmixEffectsToReset)]() mutable
+		// Unregister these source effect instances from their owning USoundEffectInstance on the audio thread.
+		// Have to pass to Game Thread prior to processing on AudioThread to avoid race condition with GC
+		// (RunCommandOnAudioThread is not safe to call from any thread other than the GameThread).
+		if (!SubmixEffectsToReset.IsEmpty())
 		{
-			for (TSoundEffectSubmixPtr& SubmixPtr : SubmixEffects)
+			AsyncTask(ENamedThreads::GameThread, [GTSubmixEffects = MoveTemp(SubmixEffectsToReset)]() mutable
 			{
-				USoundEffectPreset::UnregisterInstance(SubmixPtr);
-			}
-		});
+				FAudioThread::RunCommandOnAudioThread([ATSubmixEffects = MoveTemp(GTSubmixEffects)]() mutable
+				{
+					for (const TSoundEffectSubmixPtr& SubmixPtr : ATSubmixEffects)
+					{
+						USoundEffectPreset::UnregisterInstance(SubmixPtr);
+					}
+				});
+			});
+		}
 
 		NumSubmixEffects = 0;
 		EffectChains.Reset();
