Add AudioModulation destinations to Submix OutputVolume, WetLevel, and DryLevel

diff --git a/Engine/Source/Runtime/AudioMixer/Private/AudioMixerDevice.cpp b/Engine/Source/Runtime/AudioMixer/Private/AudioMixerDevice.cpp
--- a/Engine/Source/Runtime/AudioMixer/Private/AudioMixerDevice.cpp
+++ b/Engine/Source/Runtime/AudioMixer/Private/AudioMixerDevice.cpp
@@ -1365,6 +1365,65 @@ namespace Audio
 		}
 	}
 
+	void FMixerDevice::SetSubmixModulationSettings(USoundSubmix* InSoundSubmix, FSoundModulationDestinationSettings InOutputModulation, FSoundModulationDestinationSettings InWetLevelModulation, FSoundModulationDestinationSettings InDryLevelModulation)
+	{
+		if (!IsInAudioThread())
+		{
+			FMixerDevice* MixerDevice = this;
+
+			FAudioThread::RunCommandOnAudioThread([MixerDevice, InSoundSubmix, InOutputModulation, InWetLevelModulation, InDryLevelModulation]() 
+			{
+				MixerDevice->SetSubmixModulationSettings(InSoundSubmix, InOutputModulation, InWetLevelModulation, InDryLevelModulation);
+			});
+			return;
+		}
+
+		FModulationDestination VolumeModulation;
+		VolumeModulation.Init(DeviceID, FName("Volume"), false /* bInIsBuffered */, true /* bInValueLinear */);
+		VolumeModulation.UpdateModulator(InOutputModulation.Modulator);
+
+		FModulationDestination WetModulation;
+		WetModulation.Init(DeviceID, FName("Volume"), false /* bInIsBuffered */, true /* bInValueLinear */);
+		WetModulation.UpdateModulator(InWetLevelModulation.Modulator);
+
+		FModulationDestination DryModulation;
+		DryModulation.Init(DeviceID, FName("Volume"), false /* bInIsBuffered */, true /* bInValueLinear */);
+		DryModulation.UpdateModulator(InDryLevelModulation.Modulator);
+
+		FMixerSubmixPtr MixerSubmixPtr = GetSubmixInstance(InSoundSubmix).Pin();
+
+		if (MixerSubmixPtr.IsValid())
+		{
+			AudioRenderThreadCommand([MixerSubmixPtr, VolumeModulation, WetModulation, DryModulation]() 
+			{
+				MixerSubmixPtr->SetModulationSettings(VolumeModulation, WetModulation, DryModulation);
+			});
+		}
+	}
+
+	void FMixerDevice::SetSubmixModulationBaseLevels(USoundSubmix* InSoundSubmix, float InVolumeModBase, float InWetModBase, float InDryModBase)
+	{
+		if (!IsInAudioThread())
+		{
+			FMixerDevice* MixerDevice = this;
+
+			FAudioThread::RunCommandOnAudioThread([MixerDevice, InSoundSubmix, InVolumeModBase, InWetModBase, InDryModBase]() 
+			{
+				MixerDevice->SetSubmixModulationBaseLevels(InSoundSubmix, InVolumeModBase, InWetModBase, InDryModBase);
+			});
+			return;
+		}
+
+		FMixerSubmixPtr MixerSubmixPtr = GetSubmixInstance(InSoundSubmix).Pin();
+		if (MixerSubmixPtr.IsValid())
+		{
+			AudioRenderThreadCommand([MixerSubmixPtr, InVolumeModBase, InWetModBase, InDryModBase]() 
+			{
+				MixerSubmixPtr->SetModulationBaseLevels(InVolumeModBase, InWetModBase, InDryModBase);
+			});
+		}
+	}
+
 	bool FMixerDevice::GetCurrentSourceEffectChain(const uint32 SourceEffectChainId, TArray<FSourceEffectChainEntry>& OutCurrentSourceEffectChainEntries)
 	{
 		TArray<FSourceEffectChainEntry>* ExistingOverride = SourceEffectChainOverrides.Find(SourceEffectChainId);
diff --git a/Engine/Source/Runtime/AudioMixer/Private/AudioMixerSubmix.cpp b/Engine/Source/Runtime/AudioMixer/Private/AudioMixerSubmix.cpp
--- a/Engine/Source/Runtime/AudioMixer/Private/AudioMixerSubmix.cpp
+++ b/Engine/Source/Runtime/AudioMixer/Private/AudioMixerSubmix.cpp
@@ -8,6 +8,7 @@
 #include "AudioThread.h"
 #include "Sound/SoundEffectPreset.h"
 #include "Sound/SoundEffectSubmix.h"
+#include "Sound/SoundModulationDestination.h"
 #include "Sound/SoundSubmix.h"
 #include "Sound/SoundSubmixSend.h"
 #include "Misc/ScopeTryLock.h"
@@ -241,6 +242,23 @@ namespace Audio
 			CurrentDryLevel = FMath::Clamp(SoundSubmix->DryLevel, 0.0f, 1.0f);
 			TargetDryLevel = CurrentDryLevel;
 
+			FModulationDestination VolumeModulation;
+			VolumeModulation.Init(MixerDevice->DeviceID, FName("Volume"), false /* bInIsBuffered */, true /* bInValueLinear */);
+			VolumeModulation.UpdateModulator(SoundSubmix->OutputVolumeModulation.Modulator);
+			VolumeModBase = SoundSubmix->OutputVolumeModulation.Value;
+
+			FModulationDestination WetModulation;
+			WetModulation.Init(MixerDevice->DeviceID, FName("Volume"), false /* bInIsBuffered */, true /* bInValueLinear */);
+			WetModulation.UpdateModulator(SoundSubmix->WetLevelModulation.Modulator);
+			WetModBase = SoundSubmix->WetLevelModulation.Value;
+
+			FModulationDestination DryModulation;
+			DryModulation.Init(MixerDevice->DeviceID, FName("Volume"), false /* bInIsBuffered */, true /* bInValueLinear */);
+			DryModulation.UpdateModulator(SoundSubmix->DryLevelModulation.Modulator);
+			DryModBase = SoundSubmix->DryLevelModulation.Value;
+
+			SetModulationSettings(VolumeModulation, WetModulation, DryModulation);
+
 			FScopeLock ScopeLock(&EffectChainMutationCriticalSection);
 			{
     			NumSubmixEffects = 0;
@@ -1106,8 +1124,26 @@ namespace Audio
 
 		DryChannelBuffer.Reset();
 
+		// Update Dry Level using modulator
+		float ModulatedDryLevelStart = CurrentDryLevel;
+		float ModulatedDryLevelEnd = TargetDryLevel;
+
+		const bool bUseModulation = MixerDevice->IsModulationPluginEnabled() && MixerDevice->ModulationInterface.IsValid();
+		if (bUseModulation)
+		{
+			const float PreModulation = DryLevelMod.GetValue();
+			DryLevelMod.ProcessControl(DryModBase);
+			const float PostModulation = DryLevelMod.GetValue();
+
+			if (DryLevelMod.IsActive())
+			{
+				ModulatedDryLevelStart *= DryLevelMod.GetHasProcessed() ? PreModulation : PostModulation;
+				ModulatedDryLevelEnd *= PostModulation;
+			}
+		}
+
 		// Check if we need to allocate a dry buffer. This is stored here before effects processing. We mix in with wet buffer after effects processing.
-		if (!FMath::IsNearlyEqual(CurrentDryLevel, TargetDryLevel) || !FMath::IsNearlyZero(CurrentDryLevel))
+		if (!FMath::IsNearlyEqual(ModulatedDryLevelStart, ModulatedDryLevelEnd) || !FMath::IsNearlyZero(ModulatedDryLevelStart))
 		{
 			DryChannelBuffer.Append(InputBuffer);
 		}
@@ -1177,16 +1213,33 @@ namespace Audio
 					FMemory::Memcpy((void*)BufferPtr, (void*)SubmixChainMixBuffer.GetData(), sizeof(float)* NumSamples);
 				}
 
+				// Update Wet Level using modulator
+				float ModulatedWetLevelStart = CurrentWetLevel;
+				float ModulatedWetLevelEnd = TargetWetLevel;
+
+				if (bUseModulation)
+				{
+					const float PreModulation = WetLevelMod.GetValue();
+					WetLevelMod.ProcessControl(WetModBase);
+					const float PostModulation = WetLevelMod.GetValue();
+
+					if (WetLevelMod.IsActive())
+					{
+						ModulatedWetLevelStart *= WetLevelMod.GetHasProcessed() ? PreModulation : PostModulation;
+						ModulatedWetLevelEnd *= PostModulation;
+					}
+				}
+
 				// Apply the wet level here after processing effects. 
-				if (!FMath::IsNearlyEqual(TargetWetLevel, CurrentWetLevel) || !FMath::IsNearlyEqual(CurrentWetLevel, 1.0f))
+				if (!FMath::IsNearlyEqual(ModulatedWetLevelEnd, ModulatedWetLevelStart) || !FMath::IsNearlyEqual(ModulatedWetLevelStart, 1.0f))
 				{
-					if (FMath::IsNearlyEqual(TargetWetLevel, CurrentWetLevel))
+					if (FMath::IsNearlyEqual(ModulatedWetLevelEnd, ModulatedWetLevelStart))
 					{
-						MultiplyBufferByConstantInPlace(InputBuffer, TargetWetLevel);
+						MultiplyBufferByConstantInPlace(InputBuffer, ModulatedWetLevelEnd);
 					}
 					else
 					{
-						FadeBufferFast(InputBuffer, CurrentWetLevel, TargetWetLevel);
+						FadeBufferFast(InputBuffer, ModulatedWetLevelStart, ModulatedWetLevelEnd);
 						CurrentWetLevel = TargetWetLevel;
 					}
 				}
@@ -1197,14 +1250,14 @@ namespace Audio
 		if (DryChannelBuffer.Num() > 0)
 		{
 			// If we've already set the volume, only need to multiply by constant
-			if (FMath::IsNearlyEqual(TargetDryLevel, CurrentDryLevel))
+			if (FMath::IsNearlyEqual(ModulatedDryLevelEnd, ModulatedDryLevelStart))
 			{
-				MultiplyBufferByConstantInPlace(DryChannelBuffer, TargetDryLevel);
+				MultiplyBufferByConstantInPlace(DryChannelBuffer, ModulatedDryLevelEnd);
 			}
 			else
 			{
 				// To avoid popping, we do a fade on the buffer to the target volume
-				FadeBufferFast(DryChannelBuffer, CurrentDryLevel, TargetDryLevel);
+				FadeBufferFast(DryChannelBuffer, ModulatedDryLevelStart, ModulatedDryLevelEnd);
 				CurrentDryLevel = TargetDryLevel;
 			}
 			MixInBufferFast(DryChannelBuffer, InputBuffer);
@@ -1267,18 +1320,35 @@ namespace Audio
 			EnvelopeNumChannels = NumChannels;
 		}
 
+		// Update output volume using modulator
+		float ModulatedOutputVolumeStart = CurrentOutputVolume;
+		float ModulatedOutputVolumeEnd = TargetOutputVolume;
+
+		if (bUseModulation)
+		{
+			const float PreModulation = VolumeMod.GetValue();
+			VolumeMod.ProcessControl(VolumeModBase);
+			const float PostModulation = VolumeMod.GetValue();
+
+			if (VolumeMod.IsActive())
+			{
+				ModulatedOutputVolumeStart *= VolumeMod.GetHasProcessed() ? PreModulation : PostModulation;
+				ModulatedOutputVolumeEnd *= PostModulation;
+			}
+		}
+
 		// Now apply the output volume
-		if (!FMath::IsNearlyEqual(TargetOutputVolume, CurrentOutputVolume) || !FMath::IsNearlyEqual(CurrentOutputVolume, 1.0f))
+		if (!FMath::IsNearlyEqual(ModulatedOutputVolumeEnd, ModulatedOutputVolumeStart) || !FMath::IsNearlyEqual(ModulatedOutputVolumeStart, 1.0f))
 		{
 			// If we've already set the output volume, only need to multiply by constant
-			if (FMath::IsNearlyEqual(TargetOutputVolume, CurrentOutputVolume))
+			if (FMath::IsNearlyEqual(ModulatedOutputVolumeEnd, ModulatedOutputVolumeStart))
 			{
-				Audio::MultiplyBufferByConstantInPlace(InputBuffer, TargetOutputVolume);
+				Audio::MultiplyBufferByConstantInPlace(InputBuffer, ModulatedOutputVolumeEnd);
 			}
 			else
 			{
 				// To avoid popping, we do a fade on the buffer to the target volume
-				Audio::FadeBufferFast(InputBuffer, CurrentOutputVolume, TargetOutputVolume);
+				Audio::FadeBufferFast(InputBuffer, ModulatedOutputVolumeStart, ModulatedOutputVolumeEnd);
 				CurrentOutputVolume = TargetOutputVolume;
 			}
 		}
@@ -2050,6 +2120,20 @@ namespace Audio
 		TargetWetLevel = FMath::Clamp(InWetLevel, 0.0f, 1.0f);
 	}
 
+	void FMixerSubmix::SetModulationSettings(FModulationDestination InOutputModulation, FModulationDestination InWetLevelModulation, FModulationDestination InDryLevelModulation)
+	{
+		VolumeMod = InOutputModulation;
+		WetLevelMod = InWetLevelModulation;
+		DryLevelMod = InDryLevelModulation;
+	}
+
+	void FMixerSubmix::SetModulationBaseLevels(float InVolumeModBase, float InWetModBase, float InDryModBase)
+	{
+		VolumeModBase = InVolumeModBase;
+		WetModBase = InWetModBase;
+		DryModBase = InDryModBase;
+	}
+
 	void FMixerSubmix::BroadcastDelegates()
 	{
 		if (bIsEnvelopeFollowing)
diff --git a/Engine/Source/Runtime/AudioMixer/Public/AudioMixerDevice.h b/Engine/Source/Runtime/AudioMixer/Public/AudioMixerDevice.h
--- a/Engine/Source/Runtime/AudioMixer/Public/AudioMixerDevice.h
+++ b/Engine/Source/Runtime/AudioMixer/Public/AudioMixerDevice.h
@@ -135,6 +135,10 @@ namespace Audio
 		virtual void SetSubmixWetLevel(USoundSubmix* InSoundSubmix, float InWetLevel) override;
 		virtual void SetSubmixDryLevel(USoundSubmix* InSoundSubmix, float InDryLevel) override;
 
+		// Submix Modulation Settings
+		virtual void SetSubmixModulationSettings(USoundSubmix* InSoundSubmix, FSoundModulationDestinationSettings InOutputModulation, FSoundModulationDestinationSettings InWetLevelModulation, FSoundModulationDestinationSettings InDryLevelModulation) override;
+		virtual void SetSubmixModulationBaseLevels(USoundSubmix* InSoundSubmix, float InVolumeModBase, float InWetModBase, float InDryModBase) override;
+
 		// Submix effect chain override settings
 		virtual void SetSubmixEffectChainOverride(USoundSubmix* InSoundSubmix, const TArray<FSoundEffectSubmixPtr>& InSubmixEffectPresetChain, float InFadeTimeSec) override;
 		virtual void ClearSubmixEffectChainOverride(USoundSubmix* InSoundSubmix, float InFadeTimeSec) override;
diff --git a/Engine/Source/Runtime/AudioMixer/Public/AudioMixerSubmix.h b/Engine/Source/Runtime/AudioMixer/Public/AudioMixerSubmix.h
--- a/Engine/Source/Runtime/AudioMixer/Public/AudioMixerSubmix.h
+++ b/Engine/Source/Runtime/AudioMixer/Public/AudioMixerSubmix.h
@@ -8,6 +8,7 @@
 #include "IAudioEndpoint.h"
 #include "ISoundfieldEndpoint.h"
 #include "Sound/SoundSubmix.h"
+#include "Sound/SoundModulationDestination.h"
 #include "DSP/EnvelopeFollower.h"
 #include "DSP/MultithreadedPatching.h"
 #include "DSP/SpectrumAnalyzer.h"
@@ -137,6 +138,12 @@ namespace Audio
 		// Sets the wet level of the submix
 		void SetWetLevel(float InWetLevel);
 
+		// Update modulation settings of the submix
+		void SetModulationSettings(FModulationDestination InOutputModulation, FModulationDestination InWetLevelModulation, FModulationDestination InDryLevelModulation);
+
+		// Update modulation settings of the submix
+		void SetModulationBaseLevels(float InVolumeModBase, float InWetModeBase, float InDryModBase);
+
 		// Gets the submix channels channels
 		int32 GetSubmixChannels() const;
 
@@ -504,6 +511,14 @@ namespace Audio
 		float CurrentDryLevel;
 		float TargetDryLevel;
 
+		FModulationDestination VolumeMod;
+		FModulationDestination DryLevelMod;
+		FModulationDestination WetLevelMod;
+
+		float VolumeModBase;
+		float DryModBase;
+		float WetModBase;
+
 		// Envelope following data
 		float EnvelopeValues[AUDIO_MIXER_MAX_OUTPUT_CHANNELS];
 		TArray<Audio::FEnvelopeFollower> EnvelopeFollowers;
diff --git a/Engine/Source/Runtime/Engine/Classes/Sound/SoundSubmix.h b/Engine/Source/Runtime/Engine/Classes/Sound/SoundSubmix.h
--- a/Engine/Source/Runtime/Engine/Classes/Sound/SoundSubmix.h
+++ b/Engine/Source/Runtime/Engine/Classes/Sound/SoundSubmix.h
@@ -8,6 +8,7 @@
 #include "ISoundfieldEndpoint.h"
 #include "SampleBufferIO.h"
 #include "SoundEffectSubmix.h"
+#include "SoundModulationDestination.h"
 #include "SoundSubmixSend.h"
 #include "UObject/Object.h"
 #include "UObject/ObjectMacros.h"
@@ -303,6 +304,18 @@ public:
 	float DryLevelDB;
 #endif
 
+	/** Modulation to apply to the submix Output Volume (in dB)*/
+	UPROPERTY(EditAnywhere, BlueprintReadWrite, Category = Modulation, meta = (DisplayName = "Output Volume Modulation", AudioParam = "Volume"))
+	FSoundModulationDestinationSettings OutputVolumeModulation;
+
+	/** Modulation to apply to the submix Wet Level (in dB)*/
+	UPROPERTY(EditAnywhere, BlueprintReadWrite, Category = Modulation, meta = (DisplayName = "Wet Level Modulation", AudioParam = "Volume"))
+	FSoundModulationDestinationSettings WetLevelModulation;
+
+	/** Modulation to apply to the submix Dry Level (in dB)*/
+	UPROPERTY(EditAnywhere, BlueprintReadWrite, Category = Modulation, meta = (DisplayName = "Dry Level Modulation", AudioParam = "Volume"))
+	FSoundModulationDestinationSettings DryLevelModulation;
+
 	// Blueprint delegate for when a recorded file is finished exporting.
 	UPROPERTY(BlueprintAssignable)
 	FOnSubmixRecordedFileDone OnSubmixRecordedFileDone;
diff --git a/Engine/Source/Runtime/Engine/Private/SoundSubmix.cpp b/Engine/Source/Runtime/Engine/Private/SoundSubmix.cpp
--- a/Engine/Source/Runtime/Engine/Private/SoundSubmix.cpp
+++ b/Engine/Source/Runtime/Engine/Private/SoundSubmix.cpp
@@ -53,6 +53,9 @@ USoundSubmix::USoundSubmix(const FObjectInitializer& ObjectInitializer)
 	, DryLevelDB(-120.0f)
 #endif
 {
+	OutputVolumeModulation.Value = 0.f;
+	WetLevelModulation.Value = 0.f;
+	DryLevelModulation.Value = 0.f;
 }
 
 void USoundSubmix::PostLoad()
@@ -445,7 +448,32 @@ void USoundSubmix::PostEditChangeProperty(struct FPropertyChangedEvent& Property
 				});
 			}
 
-			if (PropertyChangedEvent.Property->GetFName() == GET_MEMBER_NAME_CHECKED(USoundSubmix, SubmixEffectChain))
+			FName MemberName = PropertyChangedEvent.MemberProperty->GetFName();
+
+			if (MemberName == GET_MEMBER_NAME_CHECKED(USoundSubmix, OutputVolumeModulation)
+				|| MemberName == GET_MEMBER_NAME_CHECKED(USoundSubmix, WetLevelModulation)
+				|| MemberName == GET_MEMBER_NAME_CHECKED(USoundSubmix, DryLevelModulation))
+			{
+				USoundSubmix* SoundSubmix = this;
+				FSoundModulationDestinationSettings NewVolumeMod = OutputVolumeModulation;
+				FSoundModulationDestinationSettings NewWetLevelMod = WetLevelModulation;
+				FSoundModulationDestinationSettings NewDryLevelMod = DryLevelModulation;
+
+				AudioDeviceManager->IterateOverAllDevices([SoundSubmix, NewVolumeMod, NewWetLevelMod, NewDryLevelMod](Audio::FDeviceId Id, FAudioDevice* Device)
+					{
+						Device->SetSubmixModulationSettings(SoundSubmix, NewVolumeMod, NewWetLevelMod, NewDryLevelMod);
+					});
+
+				float NewVolumeModBase = OutputVolumeModulation.Value;
+				float NewWetModBase = WetLevelModulation.Value;
+				float NewDryModBase = DryLevelModulation.Value;
+
+				AudioDeviceManager->IterateOverAllDevices([SoundSubmix, NewVolumeModBase, NewWetModBase, NewDryModBase](Audio::FDeviceId Id, FAudioDevice* Device) {
+					Device->SetSubmixModulationBaseLevels(SoundSubmix, NewVolumeModBase, NewWetModBase, NewDryModBase);
+					});
+			}
+
+			if (ChangedPropName == GET_MEMBER_NAME_CHECKED(USoundSubmix, SubmixEffectChain))
 			{
 				AudioDeviceManager->RegisterSoundSubmix(this);
 			}
diff --git a/Engine/Source/Runtime/Engine/Public/AudioDevice.h b/Engine/Source/Runtime/Engine/Public/AudioDevice.h
--- a/Engine/Source/Runtime/Engine/Public/AudioDevice.h
+++ b/Engine/Source/Runtime/Engine/Public/AudioDevice.h
@@ -20,6 +20,7 @@
 #include "Sound/SoundSubmix.h"
 #include "Sound/SoundSubmixSend.h"
 #include "Sound/SoundSourceBus.h"
+#include "Sound/SoundModulationDestination.h"
 #include "Subsystems/AudioEngineSubsystem.h"
 #include "Subsystems/SubsystemCollection.h"
 #include "AudioVirtualLoop.h"
@@ -1328,6 +1329,16 @@ public:
 		UE_LOG(LogAudio, Error, TEXT("Submixes are only supported in audio mixer."));
 	}
 
+	virtual void SetSubmixModulationSettings(USoundSubmix* InSoundSubmix, FSoundModulationDestinationSettings InOutputModulation, FSoundModulationDestinationSettings InWetLevelModulation, FSoundModulationDestinationSettings InDryLevelModulation)
+	{
+		UE_LOG(LogAudio, Error, TEXT("Submixes are only supported in audio mixer."));
+	}
+
+	virtual void SetSubmixModulationBaseLevels(USoundSubmix* InSoundSubmix, float InVolumeModBase, float InWetModBase, float InDryModBase)
+	{
+		UE_LOG(LogAudio, Error, TEXT("Submixes are only supported in audio mixer."));
+	}
+
 	/** Set the wet-dry level of the given submix */
 	virtual void SetSubmixOutputVolume(USoundSubmix* InSoundSubmix, float InOutputVolume)
 	{
