de-deprecated UQuartzClockHandle::ResetTransport() for now.

diff --git a/Engine/Source/Runtime/AudioMixer/Public/Quartz/AudioMixerClockHandle.h b/Engine/Source/Runtime/AudioMixer/Public/Quartz/AudioMixerClockHandle.h
--- a/Engine/Source/Runtime/AudioMixer/Public/Quartz/AudioMixerClockHandle.h
+++ b/Engine/Source/Runtime/AudioMixer/Public/Quartz/AudioMixerClockHandle.h
@@ -50,7 +50,7 @@ public:
 	void ResumeClock(const UObject* WorldContextObject, UQuartzClockHandle*& ClockHandle);
 
 
-	UFUNCTION(BlueprintCallable, Category = "Quartz Clock", meta = (DeprecatedFunction, DeprecationMessage="Please use ResetTransportQuantized instead", WorldContext = "WorldContextObject", AutoCreateRefTerm = "InDelegate", Keywords = "Transport, Counter"))
+	UFUNCTION(BlueprintCallable, Category = "Quartz Clock", meta = (/*DeprecatedFunction, DeprecationMessage="Please use ResetTransportQuantized instead",*/ WorldContext = "WorldContextObject", AutoCreateRefTerm = "InDelegate", Keywords = "Transport, Counter"))
 	void ResetTransport(const UObject* WorldContextObject, const FOnQuartzCommandEventBP& InDelegate);
 
 	UFUNCTION(BlueprintCallable, Category = "Quartz Clock", meta = (WorldContext = "WorldContextObject", AutoCreateRefTerm = "InDelegate", Keywords = "Transport, Counter"))
