Add IsModulationPluginEnabled() safeguards around modulation properties when initializing or editing a submix

diff --git a/Engine/Source/Runtime/AudioMixer/Private/AudioMixerDevice.cpp b/Engine/Source/Runtime/AudioMixer/Private/AudioMixerDevice.cpp
--- a/Engine/Source/Runtime/AudioMixer/Private/AudioMixerDevice.cpp
+++ b/Engine/Source/Runtime/AudioMixer/Private/AudioMixerDevice.cpp
@@ -1369,7 +1369,7 @@ namespace Audio
 		}
 	}
 
-	void FMixerDevice::SetSubmixModulationSettings(USoundSubmix* InSoundSubmix, FSoundModulationDestinationSettings InOutputModulation, FSoundModulationDestinationSettings InWetLevelModulation, FSoundModulationDestinationSettings InDryLevelModulation)
+	void FMixerDevice::UpdateSubmixModulationSettings(USoundSubmix* InSoundSubmix, USoundModulatorBase* InOutputModulation, USoundModulatorBase* InWetLevelModulation, USoundModulatorBase* InDryLevelModulation)
 	{
 		if (!IsInAudioThread())
 		{
@@ -1377,31 +1377,26 @@ namespace Audio
 
 			FAudioThread::RunCommandOnAudioThread([MixerDevice, InSoundSubmix, InOutputModulation, InWetLevelModulation, InDryLevelModulation]() 
 			{
-				MixerDevice->SetSubmixModulationSettings(InSoundSubmix, InOutputModulation, InWetLevelModulation, InDryLevelModulation);
+				MixerDevice->UpdateSubmixModulationSettings(InSoundSubmix, InOutputModulation, InWetLevelModulation, InDryLevelModulation);
 			});
 			return;
 		}
 
-		FModulationDestination VolumeModulation;
-		VolumeModulation.Init(DeviceID, FName("Volume"), false /* bInIsBuffered */, true /* bInValueLinear */);
-		VolumeModulation.UpdateModulator(InOutputModulation.Modulator);
-
-		FModulationDestination WetModulation;
-		WetModulation.Init(DeviceID, FName("Volume"), false /* bInIsBuffered */, true /* bInValueLinear */);
-		WetModulation.UpdateModulator(InWetLevelModulation.Modulator);
-
-		FModulationDestination DryModulation;
-		DryModulation.Init(DeviceID, FName("Volume"), false /* bInIsBuffered */, true /* bInValueLinear */);
-		DryModulation.UpdateModulator(InDryLevelModulation.Modulator);
-
-		FMixerSubmixPtr MixerSubmixPtr = GetSubmixInstance(InSoundSubmix).Pin();
-
-		if (MixerSubmixPtr.IsValid())
+		if (IsModulationPluginEnabled() && ModulationInterface.IsValid())
 		{
-			AudioRenderThreadCommand([MixerSubmixPtr, VolumeModulation, WetModulation, DryModulation]() 
+			FMixerSubmixPtr MixerSubmixPtr = GetSubmixInstance(InSoundSubmix).Pin();
+
+			if (MixerSubmixPtr.IsValid())
 			{
-				MixerSubmixPtr->SetModulationSettings(VolumeModulation, WetModulation, DryModulation);
-			});
+				USoundModulatorBase* VolumeMod = InOutputModulation;
+				USoundModulatorBase* WetMod = InOutputModulation;
+				USoundModulatorBase* DryMod = InOutputModulation;
+
+				AudioRenderThreadCommand([MixerSubmixPtr, VolumeMod, WetMod, DryMod]()
+				{
+					MixerSubmixPtr->UpdateModulationSettings(VolumeMod, WetMod, DryMod);
+				});
+			}
 		}
 	}
 
diff --git a/Engine/Source/Runtime/AudioMixer/Private/AudioMixerSubmix.cpp b/Engine/Source/Runtime/AudioMixer/Private/AudioMixerSubmix.cpp
--- a/Engine/Source/Runtime/AudioMixer/Private/AudioMixerSubmix.cpp
+++ b/Engine/Source/Runtime/AudioMixer/Private/AudioMixerSubmix.cpp
@@ -242,22 +242,20 @@ namespace Audio
 			CurrentDryLevel = FMath::Clamp(SoundSubmix->DryLevel, 0.0f, 1.0f);
 			TargetDryLevel = CurrentDryLevel;
 
-			FModulationDestination VolumeModulation;
-			VolumeModulation.Init(MixerDevice->DeviceID, FName("Volume"), false /* bInIsBuffered */, true /* bInValueLinear */);
-			VolumeModulation.UpdateModulator(SoundSubmix->OutputVolumeModulation.Modulator);
-			VolumeModBase = SoundSubmix->OutputVolumeModulation.Value;
-
-			FModulationDestination WetModulation;
-			WetModulation.Init(MixerDevice->DeviceID, FName("Volume"), false /* bInIsBuffered */, true /* bInValueLinear */);
-			WetModulation.UpdateModulator(SoundSubmix->WetLevelModulation.Modulator);
-			WetModBase = SoundSubmix->WetLevelModulation.Value;
+			if(MixerDevice->IsModulationPluginEnabled() && MixerDevice->ModulationInterface.IsValid())
+			{
+				VolumeMod.Init(MixerDevice->DeviceID, FName("Volume"), false /* bInIsBuffered */, true /* bInValueLinear */);
+				VolumeMod.UpdateModulator(SoundSubmix->OutputVolumeModulation.Modulator);
+				VolumeModBase = SoundSubmix->OutputVolumeModulation.Value;
 
-			FModulationDestination DryModulation;
-			DryModulation.Init(MixerDevice->DeviceID, FName("Volume"), false /* bInIsBuffered */, true /* bInValueLinear */);
-			DryModulation.UpdateModulator(SoundSubmix->DryLevelModulation.Modulator);
-			DryModBase = SoundSubmix->DryLevelModulation.Value;
+				WetLevelMod.Init(MixerDevice->DeviceID, FName("Volume"), false /* bInIsBuffered */, true /* bInValueLinear */);
+				WetLevelMod.UpdateModulator(SoundSubmix->WetLevelModulation.Modulator);
+				WetModBase = SoundSubmix->WetLevelModulation.Value;
 
-			SetModulationSettings(VolumeModulation, WetModulation, DryModulation);
+				DryLevelMod.Init(MixerDevice->DeviceID, FName("Volume"), false /* bInIsBuffered */, true /* bInValueLinear */);
+				DryLevelMod.UpdateModulator(SoundSubmix->DryLevelModulation.Modulator);
+				DryModBase = SoundSubmix->DryLevelModulation.Value;
+			}
 
 			FScopeLock ScopeLock(&EffectChainMutationCriticalSection);
 			{
@@ -2130,11 +2128,11 @@ namespace Audio
 		TargetWetLevel = FMath::Clamp(InWetLevel, 0.0f, 1.0f);
 	}
 
-	void FMixerSubmix::SetModulationSettings(FModulationDestination InOutputModulation, FModulationDestination InWetLevelModulation, FModulationDestination InDryLevelModulation)
+	void FMixerSubmix::UpdateModulationSettings(USoundModulatorBase* InOutputModulator, USoundModulatorBase* InWetLevelModulator, USoundModulatorBase* InDryLevelModulator)
 	{
-		VolumeMod = InOutputModulation;
-		WetLevelMod = InWetLevelModulation;
-		DryLevelMod = InDryLevelModulation;
+		VolumeMod.UpdateModulator(InOutputModulator);
+		WetLevelMod.UpdateModulator(InWetLevelModulator);
+		DryLevelMod.UpdateModulator(InDryLevelModulator);
 	}
 
 	void FMixerSubmix::SetModulationBaseLevels(float InVolumeModBase, float InWetModBase, float InDryModBase)
diff --git a/Engine/Source/Runtime/AudioMixer/Public/AudioMixerDevice.h b/Engine/Source/Runtime/AudioMixer/Public/AudioMixerDevice.h
--- a/Engine/Source/Runtime/AudioMixer/Public/AudioMixerDevice.h
+++ b/Engine/Source/Runtime/AudioMixer/Public/AudioMixerDevice.h
@@ -15,6 +15,7 @@
 // Forward Declarations
 class FOnSubmixEnvelopeBP;
 class IAudioMixerPlatformInterface;
+class USoundModulatorBase;
 
 namespace Audio
 {
@@ -136,7 +137,7 @@ namespace Audio
 		virtual void SetSubmixDryLevel(USoundSubmix* InSoundSubmix, float InDryLevel) override;
 
 		// Submix Modulation Settings
-		virtual void SetSubmixModulationSettings(USoundSubmix* InSoundSubmix, FSoundModulationDestinationSettings InOutputModulation, FSoundModulationDestinationSettings InWetLevelModulation, FSoundModulationDestinationSettings InDryLevelModulation) override;
+		virtual void UpdateSubmixModulationSettings(USoundSubmix* InSoundSubmix, USoundModulatorBase* InOutputModulation, USoundModulatorBase* InWetLevelModulation, USoundModulatorBase* InDryLevelModulation) override;
 		virtual void SetSubmixModulationBaseLevels(USoundSubmix* InSoundSubmix, float InVolumeModBase, float InWetModBase, float InDryModBase) override;
 
 		// Submix effect chain override settings
diff --git a/Engine/Source/Runtime/AudioMixer/Public/AudioMixerSubmix.h b/Engine/Source/Runtime/AudioMixer/Public/AudioMixerSubmix.h
--- a/Engine/Source/Runtime/AudioMixer/Public/AudioMixerSubmix.h
+++ b/Engine/Source/Runtime/AudioMixer/Public/AudioMixerSubmix.h
@@ -49,6 +49,7 @@ class FOnSubmixEnvelopeBP;
 class USoundEffectSubmix;
 class USoundSubmix;
 class USoundSubmixBase;
+class USoundModulatorBase;
 
 namespace Audio
 {
@@ -139,7 +140,7 @@ namespace Audio
 		void SetWetLevel(float InWetLevel);
 
 		// Update modulation settings of the submix
-		void SetModulationSettings(FModulationDestination InOutputModulation, FModulationDestination InWetLevelModulation, FModulationDestination InDryLevelModulation);
+		void UpdateModulationSettings(USoundModulatorBase* InOutputModulator, USoundModulatorBase* InWetLevelModulator, USoundModulatorBase* InDryLevelModulator);
 
 		// Update modulation settings of the submix
 		void SetModulationBaseLevels(float InVolumeModBase, float InWetModeBase, float InDryModBase);
diff --git a/Engine/Source/Runtime/Engine/Private/SoundSubmix.cpp b/Engine/Source/Runtime/Engine/Private/SoundSubmix.cpp
--- a/Engine/Source/Runtime/Engine/Private/SoundSubmix.cpp
+++ b/Engine/Source/Runtime/Engine/Private/SoundSubmix.cpp
@@ -455,22 +455,24 @@ void USoundSubmix::PostEditChangeProperty(struct FPropertyChangedEvent& Property
 				|| MemberName == GET_MEMBER_NAME_CHECKED(USoundSubmix, DryLevelModulation))
 			{
 				USoundSubmix* SoundSubmix = this;
-				FSoundModulationDestinationSettings NewVolumeMod = OutputVolumeModulation;
-				FSoundModulationDestinationSettings NewWetLevelMod = WetLevelModulation;
-				FSoundModulationDestinationSettings NewDryLevelMod = DryLevelModulation;
+
+				USoundModulatorBase* NewVolumeMod = OutputVolumeModulation.Modulator;
+				USoundModulatorBase* NewWetLevelMod = WetLevelModulation.Modulator;
+				USoundModulatorBase* NewDryLevelMod = DryLevelModulation.Modulator;
 
 				AudioDeviceManager->IterateOverAllDevices([SoundSubmix, NewVolumeMod, NewWetLevelMod, NewDryLevelMod](Audio::FDeviceId Id, FAudioDevice* Device)
-					{
-						Device->SetSubmixModulationSettings(SoundSubmix, NewVolumeMod, NewWetLevelMod, NewDryLevelMod);
-					});
+				{
+					Device->UpdateSubmixModulationSettings(SoundSubmix, NewVolumeMod, NewWetLevelMod, NewDryLevelMod);
+				});
 
 				float NewVolumeModBase = OutputVolumeModulation.Value;
 				float NewWetModBase = WetLevelModulation.Value;
 				float NewDryModBase = DryLevelModulation.Value;
 
-				AudioDeviceManager->IterateOverAllDevices([SoundSubmix, NewVolumeModBase, NewWetModBase, NewDryModBase](Audio::FDeviceId Id, FAudioDevice* Device) {
+				AudioDeviceManager->IterateOverAllDevices([SoundSubmix, NewVolumeModBase, NewWetModBase, NewDryModBase](Audio::FDeviceId Id, FAudioDevice* Device) 
+				{
 					Device->SetSubmixModulationBaseLevels(SoundSubmix, NewVolumeModBase, NewWetModBase, NewDryModBase);
-					});
+				});
 			}
 
 			if (ChangedPropName == GET_MEMBER_NAME_CHECKED(USoundSubmix, SubmixEffectChain))
diff --git a/Engine/Source/Runtime/Engine/Public/AudioDevice.h b/Engine/Source/Runtime/Engine/Public/AudioDevice.h
--- a/Engine/Source/Runtime/Engine/Public/AudioDevice.h
+++ b/Engine/Source/Runtime/Engine/Public/AudioDevice.h
@@ -49,6 +49,7 @@ class USoundConcurrency;
 class USoundEffectSourcePreset;
 class USoundEffectSubmixPreset;
 class USoundMix;
+class USoundModulatorBase;
 class USoundSubmixBase;
 class USoundSourceBus;
 class USoundWave;
@@ -1339,7 +1340,7 @@ public:
 		UE_LOG(LogAudio, Error, TEXT("Submixes are only supported in audio mixer."));
 	}
 
-	virtual void SetSubmixModulationSettings(USoundSubmix* InSoundSubmix, FSoundModulationDestinationSettings InOutputModulation, FSoundModulationDestinationSettings InWetLevelModulation, FSoundModulationDestinationSettings InDryLevelModulation)
+	virtual void UpdateSubmixModulationSettings(USoundSubmix* InSoundSubmix, USoundModulatorBase* InOutputModulation, USoundModulatorBase* InWetLevelModulation, USoundModulatorBase* InDryLevelModulation)
 	{
 		UE_LOG(LogAudio, Error, TEXT("Submixes are only supported in audio mixer."));
 	}
