Fixing metasounds cutting off too early when told they are done.

diff --git a/Engine/Plugins/Runtime/Metasound/Source/MetasoundGenerator/Private/MetasoundGenerator.cpp b/Engine/Plugins/Runtime/Metasound/Source/MetasoundGenerator/Private/MetasoundGenerator.cpp
--- a/Engine/Plugins/Runtime/Metasound/Source/MetasoundGenerator/Private/MetasoundGenerator.cpp
+++ b/Engine/Plugins/Runtime/Metasound/Source/MetasoundGenerator/Private/MetasoundGenerator.cpp
@@ -78,7 +78,8 @@ namespace Metasound
 	}
 
 	FMetasoundGenerator::FMetasoundGenerator(FMetasoundGeneratorInitParams&& InParams)
-		: bIsFinished(false)
+		: bIsFinishTriggered(false)
+		, bIsFinished(false)
 		, NumChannels(0)
 		, NumFramesPerExecute(0)
 		, NumSamplesPerExecute(0)
@@ -198,6 +199,12 @@ namespace Metasound
 
 	int32 FMetasoundGenerator::OnGenerateAudio(float* OutAudio, int32 NumSamplesRemaining)
 	{
+		// Defer finishing the metasound generator one block
+		if (bIsFinishTriggered)
+		{
+			bIsFinished = true;
+		}
+
 		if (NumSamplesRemaining <= 0)
 		{
 			return 0;
@@ -266,7 +273,7 @@ namespace Metasound
 
 		if (*OnFinishedTriggerRef)
 		{
-			bIsFinished = true;
+			bIsFinishTriggered = true;
 		}
 
 		return NumSamplesWritten;
diff --git a/Engine/Plugins/Runtime/Metasound/Source/MetasoundGenerator/Public/MetasoundGenerator.h b/Engine/Plugins/Runtime/Metasound/Source/MetasoundGenerator/Public/MetasoundGenerator.h
--- a/Engine/Plugins/Runtime/Metasound/Source/MetasoundGenerator/Public/MetasoundGenerator.h
+++ b/Engine/Plugins/Runtime/Metasound/Source/MetasoundGenerator/Public/MetasoundGenerator.h
@@ -157,6 +157,7 @@ namespace Metasound
 		FExecuter RootExecuter;
 
 		bool bIsGraphBuilding;
+		bool bIsFinishTriggered;
 		bool bIsFinished;
 
 		int32 NumChannels;
diff --git a/Engine/Source/Runtime/AudioMixer/Private/AudioMixerSourceBuffer.cpp b/Engine/Source/Runtime/AudioMixer/Private/AudioMixerSourceBuffer.cpp
--- a/Engine/Source/Runtime/AudioMixer/Private/AudioMixerSourceBuffer.cpp
+++ b/Engine/Source/Runtime/AudioMixer/Private/AudioMixerSourceBuffer.cpp
@@ -16,7 +16,7 @@ namespace Audio
 
 		if (LoopCount == Audio::LOOP_FOREVER)
 		{
-			bool bLooped = false;
+			bool bIsFinishedOrLooped = false;
 			for (uint32 Sample = 0; Sample < NumSampleToGet; ++Sample)
 			{
 				OutBufferPtr[Sample] = DataPtr[CurrentSample++] / 32768.0f;
@@ -25,10 +25,10 @@ namespace Audio
 				if (CurrentSample >= NumSamples)
 				{
 					CurrentSample = 0;
-					bLooped = true;
+					bIsFinishedOrLooped = true;
 				}
 			}
-			return bLooped;
+			return bIsFinishedOrLooped;
 		}
 		else if (CurrentSample < NumSamples)
 		{
@@ -99,6 +99,9 @@ namespace Audio
 			InitParams.bIsPreviewSound = InArgs.bIsPreviewSound;
 
 			SoundGenerator = InArgs.SoundWave->CreateSoundGenerator(InitParams);
+
+			// In the case of procedural audio generation, the mixer source buffer will never "loop" -- i.e. when it's done, it's done
+			LoopingMode = LOOP_Never;
 		}
 
 		const uint32 TotalSamples = MONO_PCM_BUFFER_SAMPLES * NumChannels;
@@ -316,11 +319,12 @@ namespace Audio
 		{
 			ProcessRealTimeSource();
 		}
-			}
+	}
 
 	bool FMixerSourceBuffer::ReadMoreRealtimeData(ICompressedAudioInfo* InDecoder, const int32 BufferIndex, EBufferReadMode BufferReadMode)
 	{
 		const int32 MaxSamples = MONO_PCM_BUFFER_SAMPLES * NumChannels;
+
 		SourceVoiceBuffers[BufferIndex]->AudioData.Reset();
 		SourceVoiceBuffers[BufferIndex]->AudioData.AddZeroed(MaxSamples);
 
@@ -332,13 +336,10 @@ namespace Audio
 			{
 				FProceduralAudioTaskData NewTaskData;
 
-				bool bIsFinished = false;
-
 				// Pass the generator instance to the async task
 				if (SoundGenerator.IsValid())
 				{
 					NewTaskData.SoundGenerator = SoundGenerator;
-					bIsFinished = SoundGenerator->IsFinished();
 				}
 				else
 				{
@@ -352,11 +353,9 @@ namespace Audio
 				NewTaskData.NumChannels = NumChannels;
 				check(!AsyncRealtimeAudioTask);
 				AsyncRealtimeAudioTask = CreateAudioTask(NewTaskData);
-
-				return bIsFinished;
 			}
 
-			return true;
+			return false;
 		}
 		else if (BufferType != EBufferType::PCMRealTime && BufferType != EBufferType::Streaming)
 		{
@@ -396,27 +395,27 @@ namespace Audio
 		return false;
 	}
 
-	void FMixerSourceBuffer::SubmitRealTimeSourceData(const bool bLooped)
+	void FMixerSourceBuffer::SubmitRealTimeSourceData(const bool bInIsFinishedOrLooped)
 	{
 		// Have we reached the end of the sound
-		if (bLooped)
+		if (bInIsFinishedOrLooped)
 		{
 			switch (LoopingMode)
 			{
-			case LOOP_Never:
-				// Play out any queued buffers - once there are no buffers left, the state check at the beginning of IsFinished will fire
-				bBufferFinished = true;
-				break;
-
-			case LOOP_WithNotification:
-				// If we have just looped, and we are looping, send notification
-				// This will trigger a WaveInstance->NotifyFinished() in the FXAudio2SoundSournce::IsFinished() function on main thread.
-				bLoopCallback = true;
-				break;
-
-			case LOOP_Forever:
-				// Let the sound loop indefinitely
-				break;
+				case LOOP_Never:
+					// Play out any queued buffers - once there are no buffers left, the state check at the beginning of IsFinished will fire
+					bBufferFinished = true;
+					break;
+
+				case LOOP_WithNotification:
+					// If we have just looped, and we are looping, send notification
+					// This will trigger a WaveInstance->NotifyFinished() in the FXAudio2SoundSournce::IsFinished() function on main thread.
+					bLoopCallback = true;
+					break;
+
+				case LOOP_Forever:
+					// Let the sound loop indefinitely
+					break;
 			}
 		}
 
@@ -433,33 +432,33 @@ namespace Audio
 		{
 			AsyncRealtimeAudioTask->EnsureCompletion();
 
-			bool bLooped = false;
+			bool bIsFinishedOrLooped = false;
 
 			switch (AsyncRealtimeAudioTask->GetType())
 			{
-			case EAudioTaskType::Decode:
-			{
-				FDecodeAudioTaskResults TaskResult;
-				AsyncRealtimeAudioTask->GetResult(TaskResult);
-
-				bLooped = TaskResult.bLooped;
-			}
-			break;
+				case EAudioTaskType::Decode:
+				{
+					FDecodeAudioTaskResults TaskResult;
+					AsyncRealtimeAudioTask->GetResult(TaskResult);
+					bIsFinishedOrLooped = TaskResult.bIsFinishedOrLooped;
+				}
+				break;
 
-			case EAudioTaskType::Procedural:
-			{
-				FProceduralAudioTaskResults TaskResult;
-				AsyncRealtimeAudioTask->GetResult(TaskResult);
+				case EAudioTaskType::Procedural:
+				{
+					FProceduralAudioTaskResults TaskResult;
+					AsyncRealtimeAudioTask->GetResult(TaskResult);
 
-				SourceVoiceBuffers[CurrentBuffer]->AudioData.SetNum(TaskResult.NumSamplesWritten);
-			}
-			break;
+					SourceVoiceBuffers[CurrentBuffer]->AudioData.SetNum(TaskResult.NumSamplesWritten);
+					bIsFinishedOrLooped = TaskResult.bIsFinished;
+				}
+				break;
 			}
 
 			delete AsyncRealtimeAudioTask;
 			AsyncRealtimeAudioTask = nullptr;
 
-			SubmitRealTimeSourceData(bLooped);
+			SubmitRealTimeSourceData(bIsFinishedOrLooped);
 		}
 
 		if (!AsyncRealtimeAudioTask)
@@ -481,12 +480,12 @@ namespace Audio
 				DataReadMode = EBufferReadMode::Asynchronous;
 			}
 
-			const bool bLooped = ReadMoreRealtimeData(DecompressionState, CurrentBuffer, DataReadMode);
+			const bool bIsFinishedOrLooped = ReadMoreRealtimeData(DecompressionState, CurrentBuffer, DataReadMode);
 
 			// If this was a synchronous read, then immediately write it
 			if (AsyncRealtimeAudioTask == nullptr && !bHasError)
 			{
-				SubmitRealTimeSourceData(bLooped);
+				SubmitRealTimeSourceData(bIsFinishedOrLooped);
 			}
 		}
 	}
diff --git a/Engine/Source/Runtime/AudioMixer/Private/AudioMixerSourceBuffer.h b/Engine/Source/Runtime/AudioMixer/Private/AudioMixerSourceBuffer.h
--- a/Engine/Source/Runtime/AudioMixer/Private/AudioMixerSourceBuffer.h
+++ b/Engine/Source/Runtime/AudioMixer/Private/AudioMixerSourceBuffer.h
@@ -117,7 +117,7 @@ namespace Audio
 
 		void SubmitInitialPCMBuffers();
 		void SubmitInitialRealtimeBuffers();
-		void SubmitRealTimeSourceData(const bool bLooped);
+		void SubmitRealTimeSourceData(const bool bFinishedOrLooped);
 		void ProcessRealTimeSource();
 		void SubmitBuffer(TSharedPtr<FMixerSourceVoiceBuffer, ESPMode::ThreadSafe> InSourceVoiceBuffer);
 		void DeleteDecoder();
diff --git a/Engine/Source/Runtime/AudioMixer/Private/AudioMixerSourceDecode.cpp b/Engine/Source/Runtime/AudioMixer/Private/AudioMixerSourceDecode.cpp
--- a/Engine/Source/Runtime/AudioMixer/Private/AudioMixerSourceDecode.cpp
+++ b/Engine/Source/Runtime/AudioMixer/Private/AudioMixerSourceDecode.cpp
@@ -62,6 +62,7 @@ public:
 					// Pre-zero the buffer before calling into the generator code as a convenience to implementers
 					FMemory::Memzero(ProceduralTaskData.AudioData, ProceduralTaskData.NumSamples * sizeof(float));
 					ProceduralResult.NumSamplesWritten = ProceduralTaskData.SoundGenerator->GetNextBuffer(ProceduralTaskData.AudioData, ProceduralTaskData.NumSamples);
+					ProceduralResult.bIsFinished = ProceduralTaskData.SoundGenerator->IsFinished();
 				}
 				else
 				{
@@ -150,11 +151,11 @@ public:
 				int32 NumBytesStreamed = kPCMBufferSize;
 				if (DecodeTaskData.BufferType == EBufferType::Streaming)
 				{
-					DecodeResult.bLooped = DecodeTaskData.DecompressionState->StreamCompressedData(DecodeBuffer.GetData(), DecodeTaskData.bLoopingMode, kPCMBufferSize, NumBytesStreamed);
+					DecodeResult.bIsFinishedOrLooped = DecodeTaskData.DecompressionState->StreamCompressedData(DecodeBuffer.GetData(), DecodeTaskData.bLoopingMode, kPCMBufferSize, NumBytesStreamed);
 				}
 				else
 				{
-					DecodeResult.bLooped = DecodeTaskData.DecompressionState->ReadCompressedData(DecodeBuffer.GetData(), DecodeTaskData.bLoopingMode, kPCMBufferSize);
+					DecodeResult.bIsFinishedOrLooped = DecodeTaskData.DecompressionState->ReadCompressedData(DecodeBuffer.GetData(), DecodeTaskData.bLoopingMode, kPCMBufferSize);
 				}
 
 				// Convert the decoded PCM data into a float buffer while still in the async task
diff --git a/Engine/Source/Runtime/AudioMixer/Private/AudioMixerSourceDecode.h b/Engine/Source/Runtime/AudioMixer/Private/AudioMixerSourceDecode.h
--- a/Engine/Source/Runtime/AudioMixer/Private/AudioMixerSourceDecode.h
+++ b/Engine/Source/Runtime/AudioMixer/Private/AudioMixerSourceDecode.h
@@ -109,9 +109,11 @@ namespace Audio
 	struct FProceduralAudioTaskResults
 	{
 		int32 NumSamplesWritten;
+		bool bIsFinished;
 
 		FProceduralAudioTaskResults()
 			: NumSamplesWritten(0)
+			, bIsFinished(false)
 		{}
 	};
 
@@ -119,10 +121,10 @@ namespace Audio
 	struct FDecodeAudioTaskResults
 	{
 		// Whether or not the audio buffer looped
-		bool bLooped;
+		bool bIsFinishedOrLooped;
 
 		FDecodeAudioTaskResults()
-			: bLooped(false)
+			: bIsFinishedOrLooped(false)
 		{}
 	};
 
diff --git a/Engine/Source/Runtime/AudioMixer/Private/AudioMixerSourceManager.cpp b/Engine/Source/Runtime/AudioMixer/Private/AudioMixerSourceManager.cpp
--- a/Engine/Source/Runtime/AudioMixer/Private/AudioMixerSourceManager.cpp
+++ b/Engine/Source/Runtime/AudioMixer/Private/AudioMixerSourceManager.cpp
@@ -1679,15 +1679,6 @@ namespace Audio
 				if (ensure(SourceInfo.MixerSourceBuffer.IsValid()))
 				{
 					SourceInfo.MixerSourceBuffer->OnBufferEnd();
-					if (SourceInfo.MixerSourceBuffer->IsGeneratorFinished())
-					{
-						SourceInfo.bIsPlaying = false;
-						SourceInfo.bIsPaused = false;
-						SourceInfo.bIsActive = false;
-						SourceInfo.bIsStopping = false;
-						SourceInfo.bIsLastBuffer = true;
-						break;
-					}
 				}
 			}
 
