Fix for crash and init order issue when sending a normal submix to a soundfield submix.

diff --git a/Engine/Source/Runtime/AudioMixer/Private/AudioMixerSubmix.cpp b/Engine/Source/Runtime/AudioMixer/Private/AudioMixerSubmix.cpp
--- a/Engine/Source/Runtime/AudioMixer/Private/AudioMixerSubmix.cpp
+++ b/Engine/Source/Runtime/AudioMixer/Private/AudioMixerSubmix.cpp
@@ -738,12 +738,6 @@ namespace Audio
 			FChildSubmixInfo& ChildSubmix = Iter.Value;
 			SetupSoundfieldEncodingForChild(ChildSubmix);
 		}
-
-		if ((ChildSubmixes.Num() > 0) && !SoundfieldStreams.Factory->ShouldEncodeAllStreamsIndependently(*SoundfieldStreams.Settings))
-		{
-			FAudioPluginInitializationParams InitParams = GetInitializationParamsForSoundfieldStream();
-			SoundfieldStreams.DownmixedChildrenEncoder = SoundfieldStreams.Factory->CreateEncoderStream(InitParams, *SoundfieldStreams.Settings);
-		}
 	}
 
 	void FMixerSubmix::SetupSoundfieldEncodingForChild(FChildSubmixInfo& InChild)
@@ -766,6 +760,16 @@ namespace Audio
 				InChild.Transcoder = GetTranscoderForChildSubmix(SubmixPtr);
 			}
 
+			// see if we need to initialize our child submix encoder
+			const bool bNeedsDownmixedChildEncoder = !SoundfieldStreams.DownmixedChildrenEncoder.IsValid();
+			const bool bUsingUniqueEncoderForEachChildSubmix = !SoundfieldStreams.Factory->ShouldEncodeAllStreamsIndependently(*SoundfieldStreams.Settings);
+
+			if (bNeedsDownmixedChildEncoder && bUsingUniqueEncoderForEachChildSubmix)
+			{
+				FAudioPluginInitializationParams InitParams = GetInitializationParamsForSoundfieldStream();
+				SoundfieldStreams.DownmixedChildrenEncoder = SoundfieldStreams.Factory->CreateEncoderStream(InitParams, *SoundfieldStreams.Settings);
+			}
+
 			// If neither of these are true, either we are downmixing all child audio and encoding it once, or
 			// this submix can handle the child's soundfield audio packet directly, so no encoder nor transcoder is needed.
 		}
@@ -830,8 +834,9 @@ namespace Audio
 			if (!ChildSubmixSharedPtr->IsSoundfieldSubmix())
 			{
 				// Reset the output scratch buffer so that we can call ProcessAudio on the ChildSubmix with it:
-				ScratchBuffer.Reset(NumSamples);
-				ScratchBuffer.AddZeroed(NumSamples);
+				ScratchBuffer.Reset(ChildSubmixSharedPtr->NumSamples);
+				ScratchBuffer.AddZeroed(ChildSubmixSharedPtr->NumSamples);
+				
 
 				// If this is true, the Soundfield Factory explicitly requested that a seperate encoder stream was set up for every
 				// non-soundfield child submix.
