Fix for ensure if binaural sound is sent directly to soundfield submix.
#p4v-preflight-copy 17088110

diff --git a/Engine/Source/Runtime/AudioMixer/Private/AudioMixerSource.cpp b/Engine/Source/Runtime/AudioMixer/Private/AudioMixerSource.cpp
--- a/Engine/Source/Runtime/AudioMixer/Private/AudioMixerSource.cpp
+++ b/Engine/Source/Runtime/AudioMixer/Private/AudioMixerSource.cpp
@@ -480,6 +480,22 @@ namespace Audio
 				InitParams.SubmixSends.Add(SubmixSend);
 				bPreviousBaseSubmixEnablement = InitParams.bEnableBaseSubmix;
 			}
+			else
+			{
+				// Warn about sending a source marked as Binaural directly to a soundfield submix:
+				// This is a bit of a gray area as soundfield submixes are intended to be their own spatial format
+				// So to send a source to this, and also flagging the source as Binaural are probably conflicting forms of spatialazition.
+				FMixerSubmixWeakPtr SubmixWeakPtr = MixerDevice->GetSubmixInstance(WaveInstance->SoundSubmix);
+
+				if (FMixerSubmixPtr SubmixPtr = SubmixWeakPtr.Pin())
+				{
+					if ((SubmixPtr->IsSoundfieldSubmix() || SubmixPtr->IsSoundfieldEndpointSubmix()))
+					{
+						UE_LOG(LogAudioMixer, Warning, TEXT("Ignoring soundfield Base Submix destination being set on SoundWave (%s) because spatializaition method is set to Binaural.")
+							, *InWaveInstance->GetName());
+					}
+				}
+			}
 
 			// Add submix sends for this source
 			for (FSoundSubmixSendInfo& SendInfo : WaveInstance->SoundSubmixSends)
diff --git a/Engine/Source/Runtime/AudioMixer/Private/AudioMixerSourceOutputBuffer.cpp b/Engine/Source/Runtime/AudioMixer/Private/AudioMixerSourceOutputBuffer.cpp
--- a/Engine/Source/Runtime/AudioMixer/Private/AudioMixerSourceOutputBuffer.cpp
+++ b/Engine/Source/Runtime/AudioMixer/Private/AudioMixerSourceOutputBuffer.cpp
@@ -518,9 +518,13 @@ namespace Audio
 
 	const ISoundfieldAudioPacket* FMixerSourceSubmixOutputBuffer::GetSoundfieldPacket(const FSoundfieldEncodingKey& InKey) const
 	{
-		check(EncodedSoundfieldDownmixes.Contains(InKey));
-		const FSoundfieldData& SoundfieldData = EncodedSoundfieldDownmixes[InKey];
-		return SoundfieldData.EncodedPacket.Get();
+		if(EncodedSoundfieldDownmixes.Contains(InKey))
+		{
+			const FSoundfieldData& SoundfieldData = EncodedSoundfieldDownmixes[InKey];
+			return SoundfieldData.EncodedPacket.Get();
+		}
+
+		return nullptr;
 	}
 
 	ISoundfieldAudioPacket* FMixerSourceSubmixOutputBuffer::GetSoundFieldPacket(const FSoundfieldEncodingKey& InKey)
