Fix for ensure when rendering offline and using dynamics processor submix effect.

diff --git a/Engine/Source/Runtime/AudioMixer/Private/AudioMixerDevice.cpp b/Engine/Source/Runtime/AudioMixer/Private/AudioMixerDevice.cpp
--- a/Engine/Source/Runtime/AudioMixer/Private/AudioMixerDevice.cpp
+++ b/Engine/Source/Runtime/AudioMixer/Private/AudioMixerDevice.cpp
@@ -2537,7 +2537,7 @@ namespace Audio
 
 	Audio::FPatchOutputStrongPtr FMixerDevice::AddPatchForSubmix(uint32 InObjectId, float InPatchGain)
 	{
-		if (!ensure(!IsInGameThread() && !IsInAudioThread()))
+		if (!ensure(IsAudioRenderingThread()))
 		{
 			return nullptr;
 		}
