Speculative fix for crash by handling edge case of null or empty resource

diff --git a/Engine/Source/Runtime/AudioMixer/Private/AudioMixerBuffer.cpp b/Engine/Source/Runtime/AudioMixer/Private/AudioMixerBuffer.cpp
--- a/Engine/Source/Runtime/AudioMixer/Private/AudioMixerBuffer.cpp
+++ b/Engine/Source/Runtime/AudioMixer/Private/AudioMixerBuffer.cpp
@@ -153,11 +153,21 @@ namespace Audio
 
 		FSoundQualityInfo QualityInfo;
 
+		if (!SoundWave->ResourceData || !SoundWave->ResourceSize)
+		{
+			UE_LOG(LogAudioMixer, Warning, TEXT("Failed to read compressed info of '%s' because there was no resource data or invalid resource size."), *ResourceName);
+			return false;
+		}
+
 		if (DecompressionState->ReadCompressedInfo(SoundWave->ResourceData, SoundWave->ResourceSize, &QualityInfo))
 		{
 			NumFrames = QualityInfo.SampleDataSize / (QualityInfo.NumChannels * sizeof(int16));
 			return true;
 		}
+		else
+		{
+			UE_LOG(LogAudioMixer, Warning, TEXT("Failed to read compressed info of '%s'."), *ResourceName);
+		}
 		return false;
 	}
 
diff --git a/Engine/Source/Runtime/Engine/Private/AudioWaveFormatParser.cpp b/Engine/Source/Runtime/Engine/Private/AudioWaveFormatParser.cpp
--- a/Engine/Source/Runtime/Engine/Private/AudioWaveFormatParser.cpp
+++ b/Engine/Source/Runtime/Engine/Private/AudioWaveFormatParser.cpp
@@ -17,6 +17,12 @@
 
 bool ParseWaveFormatHeader(const uint8* InSrcBufferData, uint32 InSrcBufferDataSize, FWaveFormatInfo& OutHeader)
 {
+	if (!InSrcBufferData || InSrcBufferDataSize < sizeof(OutHeader.RiffWaveHeader))
+	{
+		UE_LOG(LogAudio, Error, TEXT("Failed to parse wave format header. Buffer data was null or empty."));
+		return false;
+	}
+
 	OutHeader.DataStartOffset = 0;
 
 	uint32 CurrByte = 0;
