fix GC crash from USoundEffectPreset access on the audio render thread.

diff --git a/Engine/Source/Runtime/AudioMixer/Private/AudioMixerSourceManager.cpp b/Engine/Source/Runtime/AudioMixer/Private/AudioMixerSourceManager.cpp
--- a/Engine/Source/Runtime/AudioMixer/Private/AudioMixerSourceManager.cpp
+++ b/Engine/Source/Runtime/AudioMixer/Private/AudioMixerSourceManager.cpp
@@ -10,6 +10,7 @@
 #include "AudioMixer.h"
 #include "SoundFieldRendering.h"
 #include "ProfilingDebugging/CsvProfiler.h"
+#include "Async/Async.h"
 
 // Link to "Audio" profiling category
 CSV_DECLARE_CATEGORY_MODULE_EXTERN(AUDIOMIXERCORE_API, Audio);
@@ -586,10 +587,15 @@ namespace Audio
 		{
 			FSourceInfo& SourceInfo = SourceInfos[SourceId];
 
-			for (int32 i = 0; i < SourceInfo.SourceEffects.Num(); ++i)
+			// Unregister these source effect instances from their owning USoundEffectInstance on the next audio thread tick.
+			AsyncTask(ENamedThreads::AudioThread, [SourceEffects = MoveTemp(SourceInfo.SourceEffects)] () mutable
 			{
-				USoundEffectPreset::UnregisterInstance(SourceInfo.SourceEffects[i]);
-			}
+				for (int32 i = 0; i < SourceEffects.Num(); ++i)
+				{
+					USoundEffectPreset::UnregisterInstance(SourceEffects[i]);
+				}
+			});
+			
 			SourceInfo.SourceEffects.Reset();
 
 			for (int32 i = 0; i < SourceInfo.SourceEffectPresets.Num(); ++i)
diff --git a/Engine/Source/Runtime/Engine/Classes/Sound/SoundEffectPreset.h b/Engine/Source/Runtime/Engine/Classes/Sound/SoundEffectPreset.h
--- a/Engine/Source/Runtime/Engine/Classes/Sound/SoundEffectPreset.h
+++ b/Engine/Source/Runtime/Engine/Classes/Sound/SoundEffectPreset.h
@@ -104,6 +104,7 @@ public:
 
 	static void UnregisterInstance(TSoundEffectPtr InEffectPtr)
 	{
+		ensure(IsInAudioThread());
 		if (InEffectPtr.IsValid())
 		{
 			if (USoundEffectPreset* Preset = InEffectPtr->GetPreset())
@@ -117,6 +118,7 @@ public:
 
 	static void RegisterInstance(USoundEffectPreset& InPreset, TSoundEffectPtr InEffectPtr)
 	{
+		ensure(IsInAudioThread());
 		if (!InEffectPtr.IsValid())
 		{
 			return;
