Fixing binaural audio

diff --git a/Engine/Source/Runtime/AudioMixer/Private/AudioMixerDevice.cpp b/Engine/Source/Runtime/AudioMixer/Private/AudioMixerDevice.cpp
--- a/Engine/Source/Runtime/AudioMixer/Private/AudioMixerDevice.cpp
+++ b/Engine/Source/Runtime/AudioMixer/Private/AudioMixerDevice.cpp
@@ -998,7 +998,7 @@ namespace Audio
 				if (!IsMasterSubmixType(SubmixToLoad))
 				{
 					LoadSoundSubmix(*SubmixToLoad);
-					InitSoundfieldAndEndpointDataForSubmix(*SubmixToLoad, GetSubmixInstance(SubmixToLoad).Pin(), true);
+					InitSoundfieldAndEndpointDataForSubmix(*SubmixToLoad, GetSubmixInstance(SubmixToLoad).Pin(), false);
 				}
 			}
 			bSubmixRegistrationDisabled = false;
