Fixing device swap during pause

diff --git a/Engine/Source/Runtime/AudioMixer/Private/AudioMixerSourceManager.cpp b/Engine/Source/Runtime/AudioMixer/Private/AudioMixerSourceManager.cpp
--- a/Engine/Source/Runtime/AudioMixer/Private/AudioMixerSourceManager.cpp
+++ b/Engine/Source/Runtime/AudioMixer/Private/AudioMixerSourceManager.cpp
@@ -2535,8 +2535,8 @@ namespace Audio
 			{
 				FSourceInfo& SourceInfo = SourceInfos[SourceId];
 
-				// Don't need to do anything if it's not active
-				if (!SourceInfo.bIsActive)
+				// Don't need to do anything if it's not active or not paused. 
+				if (!SourceInfo.bIsActive && !SourceInfo.bIsPaused)
 				{
 					continue;
 				}
