Fix crash when downmixing in MixerAudioBus
#rnx

diff --git a/Engine/Source/Runtime/AudioMixer/Private/AudioMixerBus.cpp b/Engine/Source/Runtime/AudioMixer/Private/AudioMixerBus.cpp
--- a/Engine/Source/Runtime/AudioMixer/Private/AudioMixerBus.cpp
+++ b/Engine/Source/Runtime/AudioMixer/Private/AudioMixerBus.cpp
@@ -1,6 +1,8 @@
 // Copyright Epic Games, Inc. All Rights Reserved.
 
 #include "AudioMixerBus.h"
+
+#include "Algo/ForEach.h"
 #include "AudioMixerSourceManager.h"
 #include "DSP/BufferVectorOperations.h"
 
@@ -125,20 +127,17 @@ namespace Audio
 					const int32 NumSourceChannels = SourceManager->GetNumChannels(AudioBusSend.SourceId);
 					const int32 NumSourceSamples = NumSourceChannels * NumOutputFrames;
 
-					// If source channels are different than bus channels, we need to up-mix or down-mix
+					// Up-mix or down-mix if source channels differ from bus channels
 					if (NumSourceChannels != NumChannels)
 					{
-						Audio::FAlignedFloatBuffer ChannelMap;
+						FAlignedFloatBuffer ChannelMap;
 						SourceManager->Get2DChannelMap(AudioBusSend.SourceId, NumChannels, ChannelMap);
-
-						Audio::FAlignedFloatBuffer DownmixedBuffer;
-						DownmixedBuffer.AddUninitialized(NumOutputFrames * NumChannels);
-						Audio::DownmixBuffer(NumSourceChannels, NumChannels, SourceBufferPtr, DownmixedBuffer.GetData(), NumOutputFrames, ChannelMap.GetData());
-						Audio::MixInBufferFast(DownmixedBuffer.GetData(), BusDataBufferPtr, DownmixedBuffer.Num(), AudioBusSend.SendLevel);
+						Algo::ForEach(ChannelMap, [SendLevel = AudioBusSend.SendLevel](float& ChannelValue) { ChannelValue *= SendLevel; });
+						DownmixAndSumIntoBuffer(NumSourceChannels, NumChannels, SourceBufferPtr, BusDataBufferPtr, NumOutputFrames, ChannelMap.GetData());
 					}
 					else
 					{
-						Audio::MixInBufferFast(SourceBufferPtr, BusDataBufferPtr, NumOutputFrames * NumChannels, AudioBusSend.SendLevel);
+						MixInBufferFast(SourceBufferPtr, BusDataBufferPtr, NumOutputFrames * NumChannels, AudioBusSend.SendLevel);
 					}
 
 				}
@@ -154,15 +153,8 @@ namespace Audio
 
 	void FMixerAudioBus::CopyCurrentBuffer(Audio::FAlignedFloatBuffer& InChannelMap, int32 InNumOutputChannels, FAlignedFloatBuffer& OutBuffer, int32 NumOutputFrames) const
 	{
-		const float* RESTRICT CurrentBuffer = GetCurrentBusBuffer();
-
 		check(NumChannels != InNumOutputChannels);
-		Audio::FAlignedFloatBuffer DownmixedBuffer;
-		DownmixedBuffer.AddUninitialized(NumOutputFrames * NumChannels);
-
-		Audio::DownmixBuffer(NumChannels, InNumOutputChannels, CurrentBuffer, DownmixedBuffer.GetData(), NumOutputFrames, InChannelMap.GetData());
-
-		Audio::MixInBufferFast(DownmixedBuffer.GetData(), OutBuffer.GetData(), DownmixedBuffer.Num(), 1.0f);
+		DownmixAndSumIntoBuffer(NumChannels, InNumOutputChannels, MixedSourceData[CurrentBufferIndex], OutBuffer, InChannelMap.GetData());
 	}
 
 	void FMixerAudioBus::CopyCurrentBuffer(int32 InNumOutputChannels, FAlignedFloatBuffer& OutBuffer, int32 NumOutputFrames) const
diff --git a/Engine/Source/Runtime/AudioMixer/Public/AudioMixerBus.h b/Engine/Source/Runtime/AudioMixer/Public/AudioMixerBus.h
--- a/Engine/Source/Runtime/AudioMixer/Public/AudioMixerBus.h
+++ b/Engine/Source/Runtime/AudioMixer/Public/AudioMixerBus.h
@@ -91,7 +91,7 @@ namespace Audio
 
 		// The mixed source data. This is double-buffered to allow buses to send audio to themselves.
 		// Buses feed audio to each other by storing their previous buffer. Current buses mix in previous other buses (including themselves)
-		TArray<float> MixedSourceData[2];
+		FAlignedFloatBuffer MixedSourceData[2];
 
 		// The index of the bus data currently being rendered
 		int32 CurrentBufferIndex;
