Passing the radius as a channel position parameter for sound field format API.

diff --git a/Engine/Source/Runtime/AudioMixer/Private/AudioMixerSourceOutputBuffer.cpp b/Engine/Source/Runtime/AudioMixer/Private/AudioMixerSourceOutputBuffer.cpp
--- a/Engine/Source/Runtime/AudioMixer/Private/AudioMixerSourceOutputBuffer.cpp
+++ b/Engine/Source/Runtime/AudioMixer/Private/AudioMixerSourceOutputBuffer.cpp
@@ -399,6 +399,7 @@ namespace Audio
 				ChannelPosition.Channel = EAudioMixerChannel::FrontCenter;
 				ConvertCartesianToSpherical(InSpatParams.EmitterPosition, ChannelPosition.Azimuth, ChannelPosition.Elevation, ChannelPosition.Radius);
 
+				ChannelPosition.Radius = InSpatParams.Distance;
 				InputChannelPositions.Add(ChannelPosition);
 
 				SoundfieldPositionalData.ChannelPositions = &InputChannelPositions;
