Fixing adding and removing submixes

diff --git a/Engine/Source/Runtime/AudioMixer/Private/AudioMixerSubmix.cpp b/Engine/Source/Runtime/AudioMixer/Private/AudioMixerSubmix.cpp
--- a/Engine/Source/Runtime/AudioMixer/Private/AudioMixerSubmix.cpp
+++ b/Engine/Source/Runtime/AudioMixer/Private/AudioMixerSubmix.cpp
@@ -521,14 +521,24 @@ namespace Audio
 		}
 
 		++NumSubmixEffects;
-		for (FSubmixEffectFadeInfo& FadeInfo : EffectChains)
+		if (EffectChains.Num() > 0)
 		{
-			if (FadeInfo.bIsCurrentChain)
+			for (FSubmixEffectFadeInfo& FadeInfo : EffectChains)
 			{
-				FadeInfo.EffectChain.Add(InSoundEffectSubmix);
-				return;
+				if (FadeInfo.bIsCurrentChain)
+				{
+					FadeInfo.EffectChain.Add(InSoundEffectSubmix);
+					return;
+				}
 			}
 		}
+		else
+		{
+			FSubmixEffectFadeInfo& NewSubmixEffectChain = EffectChains.Add_GetRef(FSubmixEffectFadeInfo());
+			NewSubmixEffectChain.bIsCurrentChain = true;
+			NewSubmixEffectChain.FadeVolume = FDynamicParameter(1.0f);
+			NewSubmixEffectChain.EffectChain.Add(InSoundEffectSubmix);
+		}
 	}
 
 	void FMixerSubmix::RemoveSoundEffectSubmix(uint32 SubmixPresetId)
diff --git a/Engine/Source/Runtime/Engine/Classes/Sound/SoundEffectSubmix.h b/Engine/Source/Runtime/Engine/Classes/Sound/SoundEffectSubmix.h
--- a/Engine/Source/Runtime/Engine/Classes/Sound/SoundEffectSubmix.h
+++ b/Engine/Source/Runtime/Engine/Classes/Sound/SoundEffectSubmix.h
@@ -154,5 +154,4 @@ private:
 	}
 
 	FSoundEffectSubmixInitData InitData_Internal;
-	uint32 ParentPresetUniqueId = INDEX_NONE;
 };
