Fixing up naming consistency in submix bp api

diff --git a/Engine/Source/Runtime/AudioMixer/Private/AudioMixerBlueprintLibrary.cpp b/Engine/Source/Runtime/AudioMixer/Private/AudioMixerBlueprintLibrary.cpp
--- a/Engine/Source/Runtime/AudioMixer/Private/AudioMixerBlueprintLibrary.cpp
+++ b/Engine/Source/Runtime/AudioMixer/Private/AudioMixerBlueprintLibrary.cpp
@@ -145,6 +145,11 @@ int32 UAudioMixerBlueprintLibrary::AddSubmixEffect(const UObject* WorldContextOb
 }
 
 void UAudioMixerBlueprintLibrary::RemoveSubmixEffectPreset(const UObject* WorldContextObject, USoundSubmix* InSoundSubmix, USoundEffectSubmixPreset* InSubmixEffectPreset)
+{
+	RemoveSubmixEffect(WorldContextObject, InSoundSubmix, InSubmixEffectPreset);
+}
+
+void UAudioMixerBlueprintLibrary::RemoveSubmixEffect(const UObject* WorldContextObject, USoundSubmix* InSoundSubmix, USoundEffectSubmixPreset* InSubmixEffectPreset)
 {
 	if (Audio::FMixerDevice* MixerDevice = FAudioDeviceManager::GetAudioMixerDeviceFromWorldContext(WorldContextObject))
 	{
@@ -154,6 +159,11 @@ void UAudioMixerBlueprintLibrary::RemoveSubmixEffectPreset(const UObject* WorldC
 }
 
 void UAudioMixerBlueprintLibrary::RemoveSubmixEffectPresetAtIndex(const UObject* WorldContextObject, USoundSubmix* InSoundSubmix, int32 SubmixChainIndex)
+{
+	RemoveSubmixEffectAtIndex(WorldContextObject, InSoundSubmix, SubmixChainIndex);
+}
+
+void UAudioMixerBlueprintLibrary::RemoveSubmixEffectAtIndex(const UObject* WorldContextObject, USoundSubmix* InSoundSubmix, int32 SubmixChainIndex)
 {
 	if (Audio::FMixerDevice* MixerDevice = FAudioDeviceManager::GetAudioMixerDeviceFromWorldContext(WorldContextObject))
 	{
@@ -162,6 +172,11 @@ void UAudioMixerBlueprintLibrary::RemoveSubmixEffectPresetAtIndex(const UObject*
 }
 
 void UAudioMixerBlueprintLibrary::ReplaceSoundEffectSubmix(const UObject* WorldContextObject, USoundSubmix* InSoundSubmix, int32 SubmixChainIndex, USoundEffectSubmixPreset* SubmixEffectPreset)
+{
+	ReplaceSubmixEffect(WorldContextObject, InSoundSubmix, SubmixChainIndex, SubmixEffectPreset);
+}
+
+void UAudioMixerBlueprintLibrary::ReplaceSubmixEffect(const UObject* WorldContextObject, USoundSubmix* InSoundSubmix, int32 SubmixChainIndex, USoundEffectSubmixPreset* SubmixEffectPreset)
 {
 	if (!SubmixEffectPreset || !InSoundSubmix)
 	{
diff --git a/Engine/Source/Runtime/AudioMixer/Public/AudioMixerBlueprintLibrary.h b/Engine/Source/Runtime/AudioMixer/Public/AudioMixerBlueprintLibrary.h
--- a/Engine/Source/Runtime/AudioMixer/Public/AudioMixerBlueprintLibrary.h
+++ b/Engine/Source/Runtime/AudioMixer/Public/AudioMixerBlueprintLibrary.h
@@ -239,17 +239,29 @@ public:
 	UFUNCTION(BlueprintCallable, Category = "Audio|Effects", meta = (WorldContext = "WorldContextObject"))
 	static int32 AddSubmixEffect(const UObject* WorldContextObject, USoundSubmix* SoundSubmix, USoundEffectSubmixPreset* SubmixEffectPreset);
 
+	UE_DEPRECATED(4.27, "RemoveSubmixEffectPreset is deprecated, use RemoveSubmixEffect.")
+	UFUNCTION(BlueprintCallable, Category = "Audio|Effects", meta = (WorldContext = "WorldContextObject", DeprecatedFunction))
+	static void RemoveSubmixEffectPreset(const UObject* WorldContextObject, USoundSubmix* SoundSubmix, USoundEffectSubmixPreset* SubmixEffectPreset);
+
 	/** Removes all instances of a submix effect preset from the given submix. */
 	UFUNCTION(BlueprintCallable, Category = "Audio|Effects", meta = (WorldContext = "WorldContextObject"))
-	static void RemoveSubmixEffectPreset(const UObject* WorldContextObject, USoundSubmix* SoundSubmix, USoundEffectSubmixPreset* SubmixEffectPreset);
+	static void RemoveSubmixEffect(const UObject* WorldContextObject, USoundSubmix* SoundSubmix, USoundEffectSubmixPreset* SubmixEffectPreset);
+
+	UE_DEPRECATED(4.27, "RemoveSubmixEffectPresetAtIndex is deprecated, use RemoveSubmixEffectAtIndex.")
+	UFUNCTION(BlueprintCallable, Category = "Audio|Effects", meta = (WorldContext = "WorldContextObject", DeprecatedFunction))
+	static void RemoveSubmixEffectPresetAtIndex(const UObject* WorldContextObject, USoundSubmix* SoundSubmix, int32 SubmixChainIndex);
 
 	/** Removes the submix effect at the given submix chain index, if there is a submix effect at that index. */
 	UFUNCTION(BlueprintCallable, Category = "Audio|Effects", meta = (WorldContext = "WorldContextObject"))
-	static void RemoveSubmixEffectPresetAtIndex(const UObject* WorldContextObject, USoundSubmix* SoundSubmix, int32 SubmixChainIndex);
+	static void RemoveSubmixEffectAtIndex(const UObject* WorldContextObject, USoundSubmix* SoundSubmix, int32 SubmixChainIndex);
+
+	UE_DEPRECATED(4.27, "ReplaceSoundEffectSubmix is deprecated, use ReplaceSubmixEffect.")
+	UFUNCTION(BlueprintCallable, Category = "Audio|Effects", meta = (WorldContext = "WorldContextObject", DeprecatedFunction))
+	static void ReplaceSoundEffectSubmix(const UObject* WorldContextObject, USoundSubmix* InSoundSubmix, int32 SubmixChainIndex, USoundEffectSubmixPreset* SubmixEffectPreset);
 
 	/** Replaces the submix effect at the given submix chain index, adds the effect if there is none at that index. */
 	UFUNCTION(BlueprintCallable, Category = "Audio|Effects", meta = (WorldContext = "WorldContextObject"))
-	static void ReplaceSoundEffectSubmix(const UObject* WorldContextObject, USoundSubmix* InSoundSubmix, int32 SubmixChainIndex, USoundEffectSubmixPreset* SubmixEffectPreset);
+	static void ReplaceSubmixEffect(const UObject* WorldContextObject, USoundSubmix* InSoundSubmix, int32 SubmixChainIndex, USoundEffectSubmixPreset* SubmixEffectPreset);
 
 	/** Clears all submix effects on the given submix. */
 	UFUNCTION(BlueprintCallable, Category = "Audio|Effects", meta = (WorldContext = "WorldContextObject"))
