Fixing tooltip

diff --git a/Engine/Source/Runtime/AudioMixer/Classes/SubmixEffects/AudioMixerSubmixEffectDynamicsProcessor.h b/Engine/Source/Runtime/AudioMixer/Classes/SubmixEffects/AudioMixerSubmixEffectDynamicsProcessor.h
--- a/Engine/Source/Runtime/AudioMixer/Classes/SubmixEffects/AudioMixerSubmixEffectDynamicsProcessor.h
+++ b/Engine/Source/Runtime/AudioMixer/Classes/SubmixEffects/AudioMixerSubmixEffectDynamicsProcessor.h
@@ -141,7 +141,7 @@ struct AUDIOMIXER_API FSubmixEffectDynamicsProcessorSettings
 	UPROPERTY(EditAnywhere, BlueprintReadWrite, Category = Sidechain, meta = (EditCondition = "!bBypass"))
 	ESubmixEffectDynamicsKeySource KeySource = ESubmixEffectDynamicsKeySource::Default;
 
-	// If set, uses output of provided submix as modulator of input signal for dynamics processor (Uses input signal as default modulator)
+	// If set, uses output of provided audio bus as modulator of input signal for dynamics processor (Uses input signal as default modulator)
 	UPROPERTY(EditAnywhere, BlueprintReadWrite, Category = Sidechain, meta = (EditCondition = "!bBypass && KeySource == ESubmixEffectDynamicsKeySource::AudioBus", EditConditionHides))
 	TObjectPtr<UAudioBus> ExternalAudioBus = nullptr;
 
