Auto-disable feature for submixes

diff --git a/Engine/Source/Runtime/AudioMixer/Private/AudioMixerDevice.cpp b/Engine/Source/Runtime/AudioMixer/Private/AudioMixerDevice.cpp
--- a/Engine/Source/Runtime/AudioMixer/Private/AudioMixerDevice.cpp
+++ b/Engine/Source/Runtime/AudioMixer/Private/AudioMixerDevice.cpp
@@ -1374,6 +1374,50 @@ namespace Audio
 		}
 	}
 
+	void FMixerDevice::SetSubmixAutoDisable(USoundSubmix* InSoundSubmix, bool bInAutoDisable)
+	{
+		if (!IsInAudioThread())
+		{
+			FMixerDevice* MixerDevice = this;
+			FAudioThread::RunCommandOnAudioThread([MixerDevice, InSoundSubmix, bInAutoDisable]()
+			{
+				MixerDevice->SetSubmixAutoDisable(InSoundSubmix, bInAutoDisable);
+			});
+			return;
+		}
+
+		FMixerSubmixPtr MixerSubmixPtr = GetSubmixInstance(InSoundSubmix).Pin();
+		if (MixerSubmixPtr.IsValid())
+		{
+			AudioRenderThreadCommand([MixerSubmixPtr, bInAutoDisable]()
+			{
+				MixerSubmixPtr->SetAutoDisable(bInAutoDisable);
+			});
+		}
+	}
+
+	void FMixerDevice::SetSubmixAutoDisableTime(USoundSubmix* InSoundSubmix, float InDisableTime)
+	{
+		if (!IsInAudioThread())
+		{
+			FMixerDevice* MixerDevice = this;
+			FAudioThread::RunCommandOnAudioThread([MixerDevice, InSoundSubmix, InDisableTime]()
+			{
+				MixerDevice->SetSubmixAutoDisableTime(InSoundSubmix, InDisableTime);
+			});
+			return;
+		}
+
+		FMixerSubmixPtr MixerSubmixPtr = GetSubmixInstance(InSoundSubmix).Pin();
+		if (MixerSubmixPtr.IsValid())
+		{
+			AudioRenderThreadCommand([MixerSubmixPtr, InDisableTime]()
+			{
+				MixerSubmixPtr->SetAutoDisableTime(InDisableTime);
+			});
+		}
+	}
+
 	void FMixerDevice::UpdateSubmixModulationSettings(USoundSubmix* InSoundSubmix, USoundModulatorBase* InOutputModulation, USoundModulatorBase* InWetLevelModulation, USoundModulatorBase* InDryLevelModulation)
 	{
 		if (!IsInAudioThread())
diff --git a/Engine/Source/Runtime/AudioMixer/Private/AudioMixerSubmix.cpp b/Engine/Source/Runtime/AudioMixer/Private/AudioMixerSubmix.cpp
--- a/Engine/Source/Runtime/AudioMixer/Private/AudioMixerSubmix.cpp
+++ b/Engine/Source/Runtime/AudioMixer/Private/AudioMixerSubmix.cpp
@@ -6,6 +6,7 @@
 #include "AudioMixerDevice.h"
 #include "AudioMixerSourceVoice.h"
 #include "AudioThread.h"
+#include "DSP/FloatArrayMath.h"
 #include "Sound/SoundEffectPreset.h"
 #include "Sound/SoundEffectSubmix.h"
 #include "Sound/SoundModulationDestination.h"
@@ -33,6 +34,14 @@ FAutoConsoleVariableRef CVarBypassAllSubmixEffects(
 	TEXT("1: Submix Effects are disabled."),
 	ECVF_Default);
 
+static int32 LogSubmixEnablementCVar = 1;
+FAutoConsoleVariableRef CVarLogSubmixEnablement(
+	TEXT("au.LogSubmixEnablement"),
+	LogSubmixEnablementCVar,
+	TEXT("Enables logging of submix enablement state.\n")
+	TEXT("1: Submix enablement logging is on. 0: Submix enablement logging is off."),
+	ECVF_Default);
+
 // Define profiling categories for submixes. 
 DEFINE_STAT(STAT_AudioMixerSubmixes);
 DEFINE_STAT(STAT_AudioMixerEndpointSubmixes);
@@ -179,6 +188,11 @@ namespace Audio
 		, NumSubmixEffects(0)
 		, bIsRecording(false)
 		, bIsBackgroundMuted(false)
+		, bAutoDisable(true)
+		, bIsSilent(false)
+		, bIsCurrentlyDisabled(false)
+		, AutoDisableTime(0.01)
+		, SilenceTimeStartSeconds(-1.0)
 		, bIsSpectrumAnalyzing(false)
 	{
 	}
@@ -204,6 +218,8 @@ namespace Audio
 		check(IsInAudioThread());
 		if (InSoundSubmix != nullptr)
 		{
+			SubmixName = InSoundSubmix->GetName();
+
 			// This is a first init and needs to be synchronous
 			if (!OwningSubmixObject.IsValid())
 			{
@@ -227,6 +243,12 @@ namespace Audio
 		// Loop through the submix's presets and make new instances of effects in the same order as the presets
 		ClearSoundEffectSubmixes();
 
+		// Copy any base data
+		bAutoDisable = OwningSubmixObject->bAutoDisable;
+		AutoDisableTime = (double)OwningSubmixObject->AutoDisableTime;
+		bIsSilent = false;
+		bIsCurrentlyDisabled = false;
+		SilenceTimeStartSeconds = -1.0;
 
 		if (const USoundSubmix* SoundSubmix = Cast<const USoundSubmix>(OwningSubmixObject))
 		{
@@ -1031,6 +1053,46 @@ namespace Audio
 		return OwningSubmixObject.IsValid();
 	}
 
+	bool FMixerSubmix::IsRenderingAudio() const
+	{
+		// If this submix is not rendering any sources directly and silence has been detected, we need to check it's children submixes
+		if (MixerSourceVoices.Num() == 0 && SilenceTimeStartSeconds >= 0.0)
+		{
+			// Use the audio clock to check against the disablement timeout
+			double AudioClock = MixerDevice->GetAudioClock();
+			double TimeSinceSilent = AudioClock - SilenceTimeStartSeconds;
+
+			// If we're past the threshold for disablement, do one last check of any child submixes. Maybe they are now rendering audio.
+			if (TimeSinceSilent > AutoDisableTime)
+			{
+				for (auto& ChildSubmixEntry : ChildSubmixes)
+				{
+					TSharedPtr<Audio::FMixerSubmix, ESPMode::ThreadSafe> ChildSubmix = ChildSubmixEntry.Value.SubmixPtr.Pin();
+
+					// If any of the submix's children are rendering audio then this submix is also rendering audio
+					if (ChildSubmix->IsRenderingAudio())
+					{
+						return true;
+					}
+				}
+
+				// We're past the auto-disablement threshold and no child submixes are rendering audio, we're now silent
+				return false;
+			}
+		}
+		return true;
+	}
+
+	void FMixerSubmix::SetAutoDisable(bool bInAutoDisable)
+	{
+		bAutoDisable = bInAutoDisable;
+	}
+
+	void FMixerSubmix::SetAutoDisableTime(float InAutoDisableTime)
+	{
+		AutoDisableTime = InAutoDisableTime;
+	}
+
 	void FMixerSubmix::ProcessAudio(FAlignedFloatBuffer& OutAudioBuffer)
 	{
 		AUDIO_MIXER_CHECK_AUDIO_PLAT_THREAD(MixerDevice);
@@ -1054,6 +1116,39 @@ namespace Audio
 			return;
 		}
 
+		// As an optimization, early out if we're set to auto-disable and we're not rendering audio
+		if (bAutoDisable && !IsRenderingAudio())
+		{
+			if (!bIsCurrentlyDisabled)
+			{
+				bIsCurrentlyDisabled = true;
+
+				if (LogSubmixEnablementCVar == 1)
+				{
+					double AudioClock = MixerDevice->GetAudioClock();
+					double TimeSinceSilent = AudioClock - SilenceTimeStartSeconds;
+
+					UE_LOG(LogAudioMixer, Log, 
+						TEXT("Submix Disabled. Num Sources: %d, Time Silent: %.2f, Disablement Threshold: %.2f, Submix Name: %s"), 
+						MixerSourceVoices.Num(),
+						(float)TimeSinceSilent,
+						AutoDisableTime,
+						*SubmixName);
+				}
+			}
+
+			return;
+		}
+
+		if (bIsCurrentlyDisabled)
+		{
+			bIsCurrentlyDisabled = false;
+			if (LogSubmixEnablementCVar == 1)
+			{
+				UE_LOG(LogAudioMixer, Log, TEXT("Submix Enabled: %s"), *SubmixName);
+			}
+		}
+
 		// If this is a Soundfield Submix, process our soundfield and decode it to a OutAudioBuffer.
 		if (IsSoundfieldSubmix())
 		{
@@ -1097,9 +1192,7 @@ namespace Audio
 			// Pump pending command queues. For Soundfield Submixes this occurs in ProcessAudio(ISoundfieldAudioPacket&).
 			PumpCommandQueue();
 		}
-
-		
-
+	
 		const int32 NumOutputFrames = OutAudioBuffer.Num() / NumChannels;
 		NumSamples = NumChannels * NumOutputFrames;
 
@@ -1123,7 +1216,10 @@ namespace Audio
 				// forcibly deleted in editor, so submix validity (in addition to pointer validity) is checked before processing
 				if (ChildSubmix.IsValid() && ChildSubmix->IsValid())
 				{
-					ChildSubmix->ProcessAudio(InputBuffer);
+					if (ChildSubmix->IsRenderingAudio())
+					{
+						ChildSubmix->ProcessAudio(InputBuffer);
+					}
 				}
 				else
 				{
@@ -1388,6 +1484,26 @@ namespace Audio
 		// Mix the audio buffer of this submix with the audio buffer of the output buffer (i.e. with other submixes)
 		Audio::MixInBufferFast(InputBuffer, OutAudioBuffer);
 
+		// Once we've finished rendering submix audio, check if the output buffer is silent if we are auto-disabling
+		if (bAutoDisable)
+		{
+			// Get the mean squared value of the buffer. If the value is very small, then it's counted as "silent".
+			float BufferMean = 0.0f;
+			Audio::ArrayMeanSquared(OutAudioBuffer, BufferMean);
+
+			// If this is the first time we're silent track when it happens
+			if (BufferMean <= KINDA_SMALL_NUMBER && !bIsSilent)
+			{
+				bIsSilent = true;
+				SilenceTimeStartSeconds = MixerDevice->GetAudioClock();
+			}
+			else
+			{
+				bIsSilent = false;
+				SilenceTimeStartSeconds = -1.0;
+			}
+		}
+
 		// Now loop through any buffer listeners and feed the listeners the result of this audio callback
 		if(const USoundSubmix* SoundSubmix = Cast<const USoundSubmix>(OwningSubmixObject))
 		{
diff --git a/Engine/Source/Runtime/AudioMixer/Public/AudioMixerDevice.h b/Engine/Source/Runtime/AudioMixer/Public/AudioMixerDevice.h
--- a/Engine/Source/Runtime/AudioMixer/Public/AudioMixerDevice.h
+++ b/Engine/Source/Runtime/AudioMixer/Public/AudioMixerDevice.h
@@ -136,6 +136,10 @@ namespace Audio
 		virtual void SetSubmixWetLevel(USoundSubmix* InSoundSubmix, float InWetLevel) override;
 		virtual void SetSubmixDryLevel(USoundSubmix* InSoundSubmix, float InDryLevel) override;
 
+		// Submix auto-disable setteings
+		virtual void SetSubmixAutoDisable(USoundSubmix* InSoundSubmix, bool bInAutoDisable) override;
+		virtual void SetSubmixAutoDisableTime(USoundSubmix* InSoundSubmix, float InDisableTime) override;
+
 		// Submix Modulation Settings
 		virtual void UpdateSubmixModulationSettings(USoundSubmix* InSoundSubmix, USoundModulatorBase* InOutputModulation, USoundModulatorBase* InWetLevelModulation, USoundModulatorBase* InDryLevelModulation) override;
 		virtual void SetSubmixModulationBaseLevels(USoundSubmix* InSoundSubmix, float InVolumeModBase, float InWetModBase, float InDryModBase) override;
diff --git a/Engine/Source/Runtime/AudioMixer/Public/AudioMixerSubmix.h b/Engine/Source/Runtime/AudioMixer/Public/AudioMixerSubmix.h
--- a/Engine/Source/Runtime/AudioMixer/Public/AudioMixerSubmix.h
+++ b/Engine/Source/Runtime/AudioMixer/Public/AudioMixerSubmix.h
@@ -295,6 +295,15 @@ namespace Audio
 		//Returns true if this is an endpoint type that should no-op for this platform
 		bool IsDummyEndpointSubmix() const;
 
+		// Returns true if the submix is currently rendering audio. The current rendering time is passed in.
+		bool IsRenderingAudio() const;
+
+		// Set whether or not this submix is told to auto disable. 
+		void SetAutoDisable(bool bInAutoDisable);
+
+		// Sets the auto-disable time
+		void SetAutoDisableTime(float InAutoDisableTime);
+
 		// Get a unique key for this submix's format and settings.
 		// If another submix has an identical format and settings it will have an equivalent key.
 		FSoundfieldEncodingKey GetKeyForSubmixEncoding();
@@ -561,6 +570,24 @@ namespace Audio
 		// Whether or not this submix is muted.
 		uint8 bIsBackgroundMuted : 1;
 
+		// Whether or not auto-disablement is enabled. If true, the submix will disable itself.
+		uint8 bAutoDisable : 1;
+
+		// Whether or not the submix is currently rendering audio. I.e. audio was sent to it and mixing it, or any of its child submixes are rendering audio.
+		uint8 bIsSilent : 1;
+
+		// Whether or not we're currently disabled (i.e. the submix has been silent)
+		uint8 bIsCurrentlyDisabled : 1;
+
+		// The time to wait to disable the submix if the auto-disablement is active.
+		double AutoDisableTime;
+
+		// The time that the first full silent buffer was detected in the submix. Submix will auto-disable if the timeout is reached and the submix has bAutoDisable set to true.
+		double SilenceTimeStartSeconds;
+
+		// The name of this submix (the owning USoundSubmix)
+		FString SubmixName;
+
 		// Bool set to true when envelope following is enabled
 		FThreadSafeBool bIsEnvelopeFollowing;
 
diff --git a/Engine/Source/Runtime/Engine/Classes/Sound/SoundSubmix.h b/Engine/Source/Runtime/Engine/Classes/Sound/SoundSubmix.h
--- a/Engine/Source/Runtime/Engine/Classes/Sound/SoundSubmix.h
+++ b/Engine/Source/Runtime/Engine/Classes/Sound/SoundSubmix.h
@@ -169,6 +169,14 @@ public:
 	UEdGraph* SoundSubmixGraph;
 #endif
 
+	// Automatically disable the submix if no sound is rendering in the submix. If false, the submix will remain on regardless of if audio is playing in the submix.
+	UPROPERTY(EditAnywhere, BlueprintReadOnly, Category = AutoDisablement)
+	bool bAutoDisable = true;
+
+	/** The minimum amount of time to wait before automatically disabling a submix once no audio is rendering in it. Actual silence time may be slightly higher than this value. */
+	UPROPERTY(EditAnywhere, BlueprintReadWrite, Category = AutoDisablement, meta = (EditCondition = "bAutoDisable"))
+	float AutoDisableTime = 0.01f;
+
 	// Child submixes to this sound mix
 	UPROPERTY(VisibleAnywhere, BlueprintReadOnly, Category = SoundSubmix)
 	TArray<TObjectPtr<USoundSubmixBase>> ChildSubmixes;
diff --git a/Engine/Source/Runtime/Engine/Private/SoundSubmix.cpp b/Engine/Source/Runtime/Engine/Private/SoundSubmix.cpp
--- a/Engine/Source/Runtime/Engine/Private/SoundSubmix.cpp
+++ b/Engine/Source/Runtime/Engine/Private/SoundSubmix.cpp
@@ -377,65 +377,65 @@ void USoundSubmix::PostEditChangeProperty(struct FPropertyChangedEvent& Property
 {
 	if (PropertyChangedEvent.Property != nullptr)
 	{
-		FName ChangedPropName = PropertyChangedEvent.Property->GetFName();
+		// Force the properties to be initialized for this SoundSubmix on all active audio devices
+		if (FAudioDeviceManager* AudioDeviceManager = FAudioDeviceManager::Get())
+		{
+			FName MemberName = PropertyChangedEvent.MemberProperty->GetFName();
 
-		bool bUpdateSubmixGain = false;
+			bool bUpdateSubmixGain = false;
 
-		if (ChangedPropName == GET_MEMBER_NAME_CHECKED(USoundSubmix, OutputVolume))
-		{
-			OutputVolumeDB = Audio::ConvertToDecibels(OutputVolume);
-			bUpdateSubmixGain = true;
-		}
-		else if (ChangedPropName == GET_MEMBER_NAME_CHECKED(USoundSubmix, WetLevel))
-		{
-			WetLevelDB = Audio::ConvertToDecibels(OutputVolume);
-			bUpdateSubmixGain = true;
-		}
-		else if (ChangedPropName == GET_MEMBER_NAME_CHECKED(USoundSubmix, DryLevel))
-		{
-			DryLevelDB = Audio::ConvertToDecibels(DryLevel);
-			bUpdateSubmixGain = true;
-		}
-		else if (ChangedPropName == GET_MEMBER_NAME_CHECKED(USoundSubmix, OutputVolumeDB))
-		{
-			if (OutputVolumeDB <= -160.f)
+			if (MemberName == GET_MEMBER_NAME_CHECKED(USoundSubmix, OutputVolume))
 			{
-				OutputVolume = 0.0f;
+				OutputVolumeDB = Audio::ConvertToDecibels(OutputVolume);
+				bUpdateSubmixGain = true;
 			}
-			else
+			else if (MemberName == GET_MEMBER_NAME_CHECKED(USoundSubmix, WetLevel))
 			{
-				OutputVolume = Audio::ConvertToLinear(OutputVolumeDB);
+				WetLevelDB = Audio::ConvertToDecibels(OutputVolume);
+				bUpdateSubmixGain = true;
 			}
-			bUpdateSubmixGain = true;
-		}
-		else if (ChangedPropName == GET_MEMBER_NAME_CHECKED(USoundSubmix, WetLevelDB))
-		{
-			if (WetLevelDB <= -120.f)
+			else if (MemberName == GET_MEMBER_NAME_CHECKED(USoundSubmix, DryLevel))
 			{
-				WetLevel = 0.0f;
+				DryLevelDB = Audio::ConvertToDecibels(DryLevel);
+				bUpdateSubmixGain = true;
 			}
-			else
+			else if (MemberName == GET_MEMBER_NAME_CHECKED(USoundSubmix, OutputVolumeDB))
 			{
-				WetLevel = Audio::ConvertToLinear(WetLevelDB);
+				if (OutputVolumeDB <= -160.f)
+				{
+					OutputVolume = 0.0f;
+				}
+				else
+				{
+					OutputVolume = Audio::ConvertToLinear(OutputVolumeDB);
+				}
+				bUpdateSubmixGain = true;
 			}
-			bUpdateSubmixGain = true;
-		}
-		else if (ChangedPropName == GET_MEMBER_NAME_CHECKED(USoundSubmix, DryLevelDB))
-		{
-			if (DryLevelDB <= -120.0f)
+			else if (MemberName == GET_MEMBER_NAME_CHECKED(USoundSubmix, WetLevelDB))
 			{
-				DryLevel = 0.0f;
+				if (WetLevelDB <= -120.f)
+				{
+					WetLevel = 0.0f;
+				}
+				else
+				{
+					WetLevel = Audio::ConvertToLinear(WetLevelDB);
+				}
+				bUpdateSubmixGain = true;
 			}
-			else
+			else if (MemberName == GET_MEMBER_NAME_CHECKED(USoundSubmix, DryLevelDB))
 			{
-				DryLevel = Audio::ConvertToLinear(DryLevelDB);
+				if (DryLevelDB <= -120.0f)
+				{
+					DryLevel = 0.0f;
+				}
+				else
+				{
+					DryLevel = Audio::ConvertToLinear(DryLevelDB);
+				}
+				bUpdateSubmixGain = true;
 			}
-			bUpdateSubmixGain = true;
-		}
 
-		// Force the properties to be initialized for this SoundSubmix on all active audio devices
-		if (FAudioDeviceManager* AudioDeviceManager = FAudioDeviceManager::Get())
-		{
 			if (bUpdateSubmixGain)
 			{
 				const float NewOutputVolume = OutputVolume;
@@ -448,7 +448,25 @@ void USoundSubmix::PostEditChangeProperty(struct FPropertyChangedEvent& Property
 				});
 			}
 
-			FName MemberName = PropertyChangedEvent.MemberProperty->GetFName();
+			if (MemberName == GET_MEMBER_NAME_CHECKED(USoundSubmix, bAutoDisable))
+			{
+				const bool NewAutoDisable = bAutoDisable;
+				USoundSubmix* SoundSubmix = this;
+				AudioDeviceManager->IterateOverAllDevices([SoundSubmix, NewAutoDisable](Audio::FDeviceId Id, FAudioDevice* Device)
+				{
+					Device->SetSubmixAutoDisable(SoundSubmix, NewAutoDisable);
+				});
+			}
+
+			if (MemberName == GET_MEMBER_NAME_CHECKED(USoundSubmix, AutoDisableTime))
+			{
+				const float NewAutoDisableTime = AutoDisableTime;
+				USoundSubmix* SoundSubmix = this;
+				AudioDeviceManager->IterateOverAllDevices([SoundSubmix, NewAutoDisableTime](Audio::FDeviceId Id, FAudioDevice* Device)
+				{
+					Device->SetSubmixAutoDisableTime(SoundSubmix, NewAutoDisableTime);
+				});
+			}
 
 			if (MemberName == GET_MEMBER_NAME_CHECKED(USoundSubmix, OutputVolumeModulation)
 				|| MemberName == GET_MEMBER_NAME_CHECKED(USoundSubmix, WetLevelModulation)
@@ -475,7 +493,7 @@ void USoundSubmix::PostEditChangeProperty(struct FPropertyChangedEvent& Property
 				});
 			}
 
-			if (ChangedPropName == GET_MEMBER_NAME_CHECKED(USoundSubmix, SubmixEffectChain))
+			if (MemberName == GET_MEMBER_NAME_CHECKED(USoundSubmix, SubmixEffectChain))
 			{
 				AudioDeviceManager->RegisterSoundSubmix(this);
 			}
diff --git a/Engine/Source/Runtime/Engine/Public/AudioDevice.h b/Engine/Source/Runtime/Engine/Public/AudioDevice.h
--- a/Engine/Source/Runtime/Engine/Public/AudioDevice.h
+++ b/Engine/Source/Runtime/Engine/Public/AudioDevice.h
@@ -1402,6 +1402,18 @@ public:
 		UE_LOG(LogAudio, Error, TEXT("Submixes are only supported in audio mixer."));
 	}
 
+	/** Set whether or not a submix is auto-disabled. */
+	virtual void SetSubmixAutoDisable(USoundSubmix* InSoundSubmix, bool bInAutoDisable)
+	{
+		UE_LOG(LogAudio, Error, TEXT("Submixes are only supported in audio mixer."));
+	}
+
+	/** Set what the auto-disable time is. */
+	virtual void SetSubmixAutoDisableTime(USoundSubmix* InSoundSubmix, float InDisableTime)
+	{
+		UE_LOG(LogAudio, Error, TEXT("Submixes are only supported in audio mixer."));
+	}
+
 	virtual void UpdateSubmixModulationSettings(USoundSubmix* InSoundSubmix, USoundModulatorBase* InOutputModulation, USoundModulatorBase* InWetLevelModulation, USoundModulatorBase* InDryLevelModulation)
 	{
 		UE_LOG(LogAudio, Error, TEXT("Submixes are only supported in audio mixer."));
