Fix for invalid Shared Pointer deref (build/cis)

diff --git a/Engine/Source/Runtime/AudioMixer/Private/AudioMixerSubmix.cpp b/Engine/Source/Runtime/AudioMixer/Private/AudioMixerSubmix.cpp
--- a/Engine/Source/Runtime/AudioMixer/Private/AudioMixerSubmix.cpp
+++ b/Engine/Source/Runtime/AudioMixer/Private/AudioMixerSubmix.cpp
@@ -1070,7 +1070,7 @@ namespace Audio
 					TSharedPtr<Audio::FMixerSubmix, ESPMode::ThreadSafe> ChildSubmix = ChildSubmixEntry.Value.SubmixPtr.Pin();
 
 					// If any of the submix's children are rendering audio then this submix is also rendering audio
-					if (ChildSubmix->IsRenderingAudio())
+					if (ChildSubmix && ChildSubmix->IsRenderingAudio())
 					{
 						return true;
 					}
