Tooltip fix

diff --git a/Engine/Source/Runtime/AudioMixer/Public/AudioMixerBlueprintLibrary.h b/Engine/Source/Runtime/AudioMixer/Public/AudioMixerBlueprintLibrary.h
--- a/Engine/Source/Runtime/AudioMixer/Public/AudioMixerBlueprintLibrary.h
+++ b/Engine/Source/Runtime/AudioMixer/Public/AudioMixerBlueprintLibrary.h
@@ -303,7 +303,7 @@ public:
 	 *
 	 *  @param InNumSemitones - The number of semitones to represent.
 	 *  @param InStartingMuiscalNote - The name of the first note in the array.
-	 *  @param InStartingOctave - The octave of the first note in the arrya.
+	 *  @param InStartingOctave - The octave of the first note in the array.
 	 *  @param InAttackTimeMsec - The attack time (in milliseconds) to apply to each band's envelope tracker.
 	 *  @param InReleaseTimeMsec - The release time (in milliseconds) to apply to each band's envelope tracker.
 	 */
