Fixing tooltips

diff --git a/Engine/Source/Runtime/AudioMixer/Public/AudioMixerBlueprintLibrary.h b/Engine/Source/Runtime/AudioMixer/Public/AudioMixerBlueprintLibrary.h
--- a/Engine/Source/Runtime/AudioMixer/Public/AudioMixerBlueprintLibrary.h
+++ b/Engine/Source/Runtime/AudioMixer/Public/AudioMixerBlueprintLibrary.h
@@ -295,7 +295,7 @@ public:
 	UFUNCTION(BlueprintCallable, Category = "Audio|Analysis", meta = (WorldContext = "WorldContextObject", AdvancedDisplay = 1))
 	static void StartAnalyzingOutput(const UObject* WorldContextObject, USoundSubmix* SubmixToAnalyze = nullptr, EFFTSize FFTSize = EFFTSize::DefaultSize, EFFTPeakInterpolationMethod InterpolationMethod = EFFTPeakInterpolationMethod::Linear, EFFTWindowType WindowType = EFFTWindowType::Hann, float HopSize = 0, EAudioSpectrumType SpectrumType = EAudioSpectrumType::MagnitudeSpectrum);
 
-	/** Start spectrum analysis of the audio output. By leaving the Submix To Stop Analyzing blank, you can analyze the master output of the game. */
+	/** Stop spectrum analysis. */
 	UFUNCTION(BlueprintCallable, Category = "Audio|Analysis", meta = (WorldContext = "WorldContextObject", AdvancedDisplay = 1))
 	static void StopAnalyzingOutput(const UObject* WorldContextObject, USoundSubmix* SubmixToStopAnalyzing = nullptr);
 
@@ -331,11 +331,11 @@ public:
 	UFUNCTION(BlueprintPure, Category = "Audio|Analysis", meta = (AdvancedDisplay = 2))
 	static TArray<FSoundSubmixSpectralAnalysisBandSettings> MakePresetSpectralAnalysisBandSettings(EAudioSpectrumBandPresetType InBandPresetType, int32 InNumBands = 10, int32 InAttackTimeMsec = 10, int32 InReleaseTimeMsec = 10);
 
-	/** Start spectrum analysis of the audio output. By leaving the Submix To Analyze blank, you can analyze the master output of the game. */
+	/** Retrieve the magnitudes for the given frequencies. */
 	UFUNCTION(BlueprintCallable, Category = "Audio|Analysis", meta = (WorldContext = "WorldContextObject", AdvancedDisplay = 3))
 	static void GetMagnitudeForFrequencies(const UObject* WorldContextObject, const TArray<float>& Frequencies, TArray<float>& Magnitudes, USoundSubmix* SubmixToAnalyze = nullptr);
 
-	/** Start spectrum analysis of the audio output. By leaving the Submix To Analyze blank, you can analyze the master output of the game. */
+	/** Retrieve the phases for the given frequencies. */
 	UFUNCTION(BlueprintCallable, Category = "Audio|Analysis", meta = (WorldContext = "WorldContextObject", AdvancedDisplay = 3))
 	static void GetPhaseForFrequencies(const UObject* WorldContextObject, const TArray<float>& Frequencies, TArray<float>& Phases, USoundSubmix* SubmixToAnalyze = nullptr);
 
diff --git a/Engine/Source/Runtime/Engine/Classes/Sound/SoundSubmix.h b/Engine/Source/Runtime/Engine/Classes/Sound/SoundSubmix.h
--- a/Engine/Source/Runtime/Engine/Classes/Sound/SoundSubmix.h
+++ b/Engine/Source/Runtime/Engine/Classes/Sound/SoundSubmix.h
@@ -393,7 +393,7 @@ public:
 
 	void StartSpectralAnalysis(FAudioDevice* InDevice, EFFTSize FFTSize = EFFTSize::DefaultSize, EFFTPeakInterpolationMethod InterpolationMethod = EFFTPeakInterpolationMethod::Linear, EFFTWindowType WindowType = EFFTWindowType::Hann, float HopSize = 0, EAudioSpectrumType SpectrumType = EAudioSpectrumType::MagnitudeSpectrum);
 
-	/** Start spectrum analysis of the audio output. */
+	/** Stop spectrum analysis of the audio output. */
 	UFUNCTION(BlueprintCallable, Category = "Audio|Analysis", meta = (WorldContext = "WorldContextObject", AdvancedDisplay = 1))
 	void StopSpectralAnalysis(const UObject* WorldContextObject);
 
