Fixing reverb  plugins with new auto-disablement feature.

diff --git a/Engine/Plugins/Runtime/Oculus/OculusAudio/Source/OculusAudio/Private/OculusAudioReverb.cpp b/Engine/Plugins/Runtime/Oculus/OculusAudio/Source/OculusAudio/Private/OculusAudioReverb.cpp
--- a/Engine/Plugins/Runtime/Oculus/OculusAudio/Source/OculusAudio/Private/OculusAudioReverb.cpp
+++ b/Engine/Plugins/Runtime/Oculus/OculusAudio/Source/OculusAudio/Private/OculusAudioReverb.cpp
@@ -120,6 +120,7 @@ USoundSubmix* OculusAudioReverb::GetSubmix()
 		ReverbSubmix = NewObject<USoundSubmix>(USoundSubmix::StaticClass(), *DefaultSubmixName);
 		ReverbSubmix->bMuteWhenBackgrounded = true;
 	}
+	ReverbSubmix->bAutoDisable = false;
 
 	return ReverbSubmix;
 }
\ No newline at end of file
diff --git a/Engine/Plugins/Runtime/ResonanceAudio/Source/ResonanceAudio/Private/ResonanceAudioReverb.cpp b/Engine/Plugins/Runtime/ResonanceAudio/Source/ResonanceAudio/Private/ResonanceAudioReverb.cpp
--- a/Engine/Plugins/Runtime/ResonanceAudio/Source/ResonanceAudio/Private/ResonanceAudioReverb.cpp
+++ b/Engine/Plugins/Runtime/ResonanceAudio/Source/ResonanceAudio/Private/ResonanceAudioReverb.cpp
@@ -189,6 +189,8 @@ namespace ResonanceAudio
 			ReverbSubmix->bMuteWhenBackgrounded = true;
 		}
 
+		ReverbSubmix->bAutoDisable = false;
+
 		// SubmixEffect is required to be initialized for ReverbPluginPreset to be set
 		if (!SubmixEffect.IsValid())
 		{
diff --git a/Engine/Plugins/Runtime/Steam/SteamAudio/Source/SteamAudio/Private/PhononReverb.cpp b/Engine/Plugins/Runtime/Steam/SteamAudio/Source/SteamAudio/Private/PhononReverb.cpp
--- a/Engine/Plugins/Runtime/Steam/SteamAudio/Source/SteamAudio/Private/PhononReverb.cpp
+++ b/Engine/Plugins/Runtime/Steam/SteamAudio/Source/SteamAudio/Private/PhononReverb.cpp
@@ -438,6 +438,7 @@ namespace SteamAudio
 			ReverbSubmix = NewObject<USoundSubmix>(USoundSubmix::StaticClass(), *DefaultSubmixName);
 			ReverbSubmix->bMuteWhenBackgrounded = true;
 		}
+		ReverbSubmix->bAutoDisable = false;
 
 		bool bFoundPreset = false;
 		for (USoundEffectSubmixPreset* Preset : ReverbSubmix->SubmixEffectChain)
diff --git a/Engine/Source/Runtime/AudioMixer/Private/AudioMixerSubmix.cpp b/Engine/Source/Runtime/AudioMixer/Private/AudioMixerSubmix.cpp
--- a/Engine/Source/Runtime/AudioMixer/Private/AudioMixerSubmix.cpp
+++ b/Engine/Source/Runtime/AudioMixer/Private/AudioMixerSubmix.cpp
@@ -1067,6 +1067,12 @@ namespace Audio
 
 	bool FMixerSubmix::IsRenderingAudio() const
 	{
+		// If we're told to not auto-disable we act as if we're always rendering audio
+		if (!bAutoDisable)
+		{
+			return true;
+		}
+
 		// If this submix is not rendering any sources directly and silence has been detected, we need to check it's children submixes
 		if (MixerSourceVoices.Num() == 0 && SilenceTimeStartSeconds >= 0.0)
 		{
