- Fix Quartz play commands being canceled but then resumed in the same frame but still being marked as canceled, presumably exacerbated by saturating voice count.
- Remove duplicate logic from bad merge
#rnx

diff --git a/Engine/Source/Runtime/AudioMixer/Private/AudioMixerSource.cpp b/Engine/Source/Runtime/AudioMixer/Private/AudioMixerSource.cpp
--- a/Engine/Source/Runtime/AudioMixer/Private/AudioMixerSource.cpp
+++ b/Engine/Source/Runtime/AudioMixer/Private/AudioMixerSource.cpp
@@ -376,12 +376,6 @@ namespace Audio
 				InitParams.QuantizedRequestData = *WaveInstance->QuantizedRequestData;
 			}
 
-			// Copy quantization request data
-			if (WaveInstance->QuantizedRequestData)
-			{
-				InitParams.QuantizedRequestData = *WaveInstance->QuantizedRequestData;
-			}
-
 			if (WaveInstance->bIsAmbisonics && (WaveData->NumChannels != 4))
 			{
 				UE_LOG(LogAudioMixer, Warning, TEXT("Sound wave %s was flagged as being ambisonics but had a channel count of %d. Currently the audio engine only supports FOA sources that have four channels."), *InWaveInstance->GetName(), WaveData->NumChannels);
diff --git a/Engine/Source/Runtime/AudioMixer/Private/AudioMixerSourceManager.cpp b/Engine/Source/Runtime/AudioMixer/Private/AudioMixerSourceManager.cpp
--- a/Engine/Source/Runtime/AudioMixer/Private/AudioMixerSourceManager.cpp
+++ b/Engine/Source/Runtime/AudioMixer/Private/AudioMixerSourceManager.cpp
@@ -612,6 +612,8 @@ namespace Audio
 		SourceInfo.bHasPreDistanceAttenuationSend = false;
 		SourceInfo.bModFiltersUpdated = false;
 
+		SourceInfo.QuantizedCommandHandle.Reset();
+
 #if AUDIO_MIXER_ENABLE_DEBUG_MODE
 		SourceInfo.bIsDebugMode = false;
 		SourceInfo.DebugName = FString();
@@ -2768,6 +2770,8 @@ namespace Audio
 
 		SourceInfo.bIsPausedForQuantization = false;
 		SourceInfo.bIsActive = !SourceInfo.bIsPaused;
+
+		SourceInfo.QuantizedCommandHandle.Reset();
 	}
 
 	const float* FMixerSourceManager::GetPreDistanceAttenuationBuffer(const int32 SourceId) const
diff --git a/Engine/Source/Runtime/AudioMixer/Private/Quartz/AudioMixerQuantizedCommands.cpp b/Engine/Source/Runtime/AudioMixer/Private/Quartz/AudioMixerQuantizedCommands.cpp
--- a/Engine/Source/Runtime/AudioMixer/Private/Quartz/AudioMixerQuantizedCommands.cpp
+++ b/Engine/Source/Runtime/AudioMixer/Private/Quartz/AudioMixerQuantizedCommands.cpp
@@ -24,6 +24,7 @@ namespace Audio
 	{
 		OwningClockPtr = InCommandInitInfo.OwningClockPointer;
 		SourceID = InCommandInitInfo.SourceID;
+		bIsCanceled = false;
 
 		// access source manager through owning clock (via clock manager)
 		FMixerSourceManager* SourceManager = OwningClockPtr->GetSourceManager();
diff --git a/Engine/Source/Runtime/Engine/Classes/Sound/QuartzQuantizationUtilities.cpp b/Engine/Source/Runtime/Engine/Classes/Sound/QuartzQuantizationUtilities.cpp
--- a/Engine/Source/Runtime/Engine/Classes/Sound/QuartzQuantizationUtilities.cpp
+++ b/Engine/Source/Runtime/Engine/Classes/Sound/QuartzQuantizationUtilities.cpp
@@ -576,6 +576,12 @@ namespace Audio
 		return false;
 	}
 
+	void FQuartzQuantizedCommandHandle::Reset()
+	{
+		MixerDevice = nullptr;
+		CommandPtr.Reset();
+		OwningClockName = FName();
+	}
 
 	FQuartzLatencyTimer::FQuartzLatencyTimer()
 		: JourneyStartCycles(-1)
diff --git a/Engine/Source/Runtime/Engine/Classes/Sound/QuartzQuantizationUtilities.h b/Engine/Source/Runtime/Engine/Classes/Sound/QuartzQuantizationUtilities.h
--- a/Engine/Source/Runtime/Engine/Classes/Sound/QuartzQuantizationUtilities.h
+++ b/Engine/Source/Runtime/Engine/Classes/Sound/QuartzQuantizationUtilities.h
@@ -563,8 +563,11 @@ namespace Audio
 		TSharedPtr<IQuartzQuantizedCommand> CommandPtr{ nullptr };
 		FMixerDevice* MixerDevice{ nullptr };
 
-		// attempts to cancel the command. Returns true if the cancellation was successful.
+		// Attempts to cancel the command. Returns true if the cancellation was successful.
 		bool Cancel();
+
+		// Resets the handle to initial state.
+		void Reset();
 	};
 } // namespace Audio
 
