Bandaid to prevent OOMs in case where the hardware has stopped receiving data.

diff --git a/Engine/Source/Runtime/AudioMixer/Private/AudioMixerSourceManager.cpp b/Engine/Source/Runtime/AudioMixer/Private/AudioMixerSourceManager.cpp
--- a/Engine/Source/Runtime/AudioMixer/Private/AudioMixerSourceManager.cpp
+++ b/Engine/Source/Runtime/AudioMixer/Private/AudioMixerSourceManager.cpp
@@ -88,6 +88,13 @@ FAutoConsoleVariableRef CVarCommandBufferFlushWaitTimeMs(
 	TEXT("How long to wait for the command buffer flush to complete.\n"),
 	ECVF_Default);
 
+static int32 CommandBufferMaxSizeInMbCvar = 10;
+FAutoConsoleVariableRef CVarCommandBufferMaxSizeMb(
+	TEXT("au.CommandBufferMaxSizeInMb"),
+	CommandBufferMaxSizeInMbCvar,
+	TEXT("How big to allow the command buffer to grow before ignoring more commands"),
+	ECVF_Default);
+
 // +/- 4 Octaves (default)
 static float MaxModulationPitchRangeFreqCVar = 16.0f;
 static float MinModulationPitchRangeFreqCVar = 0.0625f;
@@ -2919,22 +2926,28 @@ namespace Audio
 
 		// Add the function to the command queue:
 		int32 AudioThreadCommandIndex = !RenderThreadCommandBufferIndex.GetValue();
+		SIZE_T CurrentBufferSizeInBytes = CommandBuffers[AudioThreadCommandIndex].SourceCommandQueue.GetAllocatedSize();
 		
 #if !NO_LOGGING
-		static uint32 WarnSize = 1024 * 1024;
-		SIZE_T Size = CommandBuffers[AudioThreadCommandIndex].SourceCommandQueue.GetAllocatedSize();
-		if (Size > WarnSize )
+		static SIZE_T WarnSize = 1024 * 1024;
+		if (CurrentBufferSizeInBytes > WarnSize )
 		{		
 			SIZE_T Num = CommandBuffers[AudioThreadCommandIndex].SourceCommandQueue.Num();
 			// NOTE: Although not really and error we want this to show up in shipping builds.
-			UE_LOG(LogAudioMixer, Error, TEXT("Command Queue has grown to %uk bytes, containing %d cmds, last pump was %fms ago."), 
-				Size >> 10, Num, FPlatformTime::ToMilliseconds64(FPlatformTime::Cycles64() - LastPumpTimeInCycles));
+			UE_LOG(LogAudioMixer, Error, TEXT("Command Queue has grown to %ukb, containing %d cmds, last pump was %.2fms ago."), 
+				CurrentBufferSizeInBytes >> 10, Num, FPlatformTime::ToMilliseconds64(FPlatformTime::Cycles64() - LastPumpTimeInCycles));
 			WarnSize *= 2;
 		}
 #endif //!NO_LOGGING
 
-		CommandBuffers[AudioThreadCommandIndex].SourceCommandQueue.Add(MoveTemp(InFunction));
-		NumCommands.Increment();
+		// Before adding further commands, ensure we're not growing outside any sensible size for these buffers.
+		// On shipping builds, this will just stop us crashing from growing out of control and OOMing the machine.		
+		const SIZE_T MaxBufferSizeInBytes = CommandBufferMaxSizeInMbCvar << 20;
+		if (ensureMsgf(CurrentBufferSizeInBytes < MaxBufferSizeInBytes, TEXT("Command buffer grown to %umb, preventing any more adds! Likely cause, the h/w has stopped consuming data."), CurrentBufferSizeInBytes >>20))
+		{
+			CommandBuffers[AudioThreadCommandIndex].SourceCommandQueue.Add(MoveTemp(InFunction));
+			NumCommands.Increment();
+		}
 	}
 
 	void FMixerSourceManager::PumpCommandQueue()
