Fix CIS warning

diff --git a/Engine/Source/Runtime/AudioMixer/Private/AudioMixerSourceManager.cpp b/Engine/Source/Runtime/AudioMixer/Private/AudioMixerSourceManager.cpp
--- a/Engine/Source/Runtime/AudioMixer/Private/AudioMixerSourceManager.cpp
+++ b/Engine/Source/Runtime/AudioMixer/Private/AudioMixerSourceManager.cpp
@@ -2942,7 +2942,7 @@ namespace Audio
 
 		// Before adding further commands, ensure we're not growing outside any sensible size for these buffers.
 		// On shipping builds, this will just stop us crashing from growing out of control and OOMing the machine.		
-		const SIZE_T MaxBufferSizeInBytes = CommandBufferMaxSizeInMbCvar << 20;
+		const SIZE_T MaxBufferSizeInBytes = ((SIZE_T)CommandBufferMaxSizeInMbCvar) << 20;
 		if (ensureMsgf(CurrentBufferSizeInBytes < MaxBufferSizeInBytes, TEXT("Command buffer grown to %umb, preventing any more adds! Likely cause, the h/w has stopped consuming data."), CurrentBufferSizeInBytes >>20))
 		{
 			CommandBuffers[AudioThreadCommandIndex].SourceCommandQueue.Add(MoveTemp(InFunction));
