- Removing unnecessary call to source override in source manager.
- And removed the GetUseSpatialization condition for releasing SourceDataOverride sources.

diff --git a/Engine/Source/Runtime/AudioMixer/Private/AudioMixerSource.cpp b/Engine/Source/Runtime/AudioMixer/Private/AudioMixerSource.cpp
--- a/Engine/Source/Runtime/AudioMixer/Private/AudioMixerSource.cpp
+++ b/Engine/Source/Runtime/AudioMixer/Private/AudioMixerSource.cpp
@@ -1201,7 +1201,7 @@ namespace Audio
 		if (MixerSourceVoice)
 		{
 			// Release the source using the propagation interface
-			if (AudioDevice->SourceDataOverridePluginInterface && WaveInstance->GetUseSpatialization())
+			if (AudioDevice->SourceDataOverridePluginInterface)
 			{
 				uint32 SourceId = MixerSourceVoice->GetSourceId();
 				AudioDevice->SourceDataOverridePluginInterface->OnReleaseSource(SourceId);
diff --git a/Engine/Source/Runtime/AudioMixer/Private/AudioMixerSourceManager.cpp b/Engine/Source/Runtime/AudioMixer/Private/AudioMixerSourceManager.cpp
--- a/Engine/Source/Runtime/AudioMixer/Private/AudioMixerSourceManager.cpp
+++ b/Engine/Source/Runtime/AudioMixer/Private/AudioMixerSourceManager.cpp
@@ -278,7 +278,6 @@ namespace Audio
 			SourceInfo.bUseHRTFSpatializer = false;
 			SourceInfo.bUseOcclusionPlugin = false;
 			SourceInfo.bUseReverbPlugin = false;
-			SourceInfo.bUseSourceDataOverridePlugin = false;
 			SourceInfo.bHasStarted = false;
 			SourceInfo.bEnableBusSends = false;
 			SourceInfo.bEnableBaseSubmix = false;
@@ -541,11 +540,6 @@ namespace Audio
 			MixerDevice->ReverbPluginInterface->OnReleaseSource(SourceId);
 		}
 
-		if (SourceInfo.bUseSourceDataOverridePlugin)
-		{
-			MixerDevice->SourceDataOverridePluginInterface->OnReleaseSource(SourceId);
-		}
-
 		// Delete the source effects
 		SourceInfo.SourceEffectChainId = INDEX_NONE;
 		ResetSourceEffectChain(SourceId);
@@ -609,7 +603,6 @@ namespace Audio
 		SourceInfo.bIsExternalSend = false;
 		SourceInfo.bUseOcclusionPlugin = false;
 		SourceInfo.bUseReverbPlugin = false;
-		SourceInfo.bUseSourceDataOverridePlugin = false;
 		SourceInfo.bHasStarted = false;
 		SourceInfo.bEnableBusSends = false;
 		SourceInfo.bEnableBaseSubmix = false;
@@ -871,12 +864,6 @@ namespace Audio
 			SourceInfo.SourceBufferListener = InitParams.SourceBufferListener;
 			SourceInfo.bShouldSourceBufferListenerZeroBuffer = InitParams.bShouldSourceBufferListenerZeroBuffer;
 
-			if (InitParams.SourceDataOverridePluginSettings != nullptr)
-			{
-				MixerDevice->SourceDataOverridePluginInterface->OnInitSource(SourceId, InitParams.AudioComponentUserID, InitParams.SourceDataOverridePluginSettings);
-				SourceInfo.bUseSourceDataOverridePlugin = true;
-			}
-
 			// Default all sounds to not consider effect chain tails when playing
 			SourceInfo.bEffectTailsDone = true;
 
diff --git a/Engine/Source/Runtime/AudioMixer/Private/AudioMixerSourceManager.h b/Engine/Source/Runtime/AudioMixer/Private/AudioMixerSourceManager.h
--- a/Engine/Source/Runtime/AudioMixer/Private/AudioMixerSourceManager.h
+++ b/Engine/Source/Runtime/AudioMixer/Private/AudioMixerSourceManager.h
@@ -453,7 +453,6 @@ namespace Audio
 			uint8 bIsExternalSend:1;
 			uint8 bUseOcclusionPlugin:1;
 			uint8 bUseReverbPlugin:1;
-			uint8 bUseSourceDataOverridePlugin:1;
 			uint8 bIsDone:1;
 			uint8 bIsLastBuffer:1;
 			uint8 bEnableBusSends : 1;
