Deprecate old submix volume, wet, and dry mix properties and replace with audio modulation destinations

diff --git a/Engine/Source/Runtime/AudioMixer/Private/AudioMixerSubmix.cpp b/Engine/Source/Runtime/AudioMixer/Private/AudioMixerSubmix.cpp
--- a/Engine/Source/Runtime/AudioMixer/Private/AudioMixerSubmix.cpp
+++ b/Engine/Source/Runtime/AudioMixer/Private/AudioMixerSubmix.cpp
@@ -253,26 +253,23 @@ namespace Audio
 
 		if (const USoundSubmix* SoundSubmix = Cast<const USoundSubmix>(OwningSubmixObject))
 		{
-			CurrentOutputVolume = FMath::Clamp(SoundSubmix->OutputVolume, 0.0f, 1.0f);
-			TargetOutputVolume = CurrentOutputVolume;
+			VolumeModBase = FMath::Clamp(SoundSubmix->OutputVolumeModulation.Value, 0.0f, 1.0f);
+			WetModBase = FMath::Clamp(SoundSubmix->WetLevelModulation.Value, 0.0f, 1.0f);
+			DryModBase = FMath::Clamp(SoundSubmix->DryLevelModulation.Value, 0.0f, 1.0f);
 
-			// Set the initialized output volume
-			CurrentWetLevel = FMath::Clamp(SoundSubmix->WetLevel, 0.0f, 1.0f);
-			TargetWetLevel = CurrentWetLevel;
+			TargetOutputVolume = VolumeModBase;
+			TargetWetLevel = WetModBase;
+			TargetDryLevel = DryModBase;
 
-			CurrentDryLevel = FMath::Clamp(SoundSubmix->DryLevel, 0.0f, 1.0f);
-			TargetDryLevel = CurrentDryLevel;
+			CurrentOutputVolume = TargetOutputVolume;
+			CurrentDryLevel = TargetDryLevel;
+			CurrentWetLevel = TargetWetLevel;
 
 			if(MixerDevice->IsModulationPluginEnabled() && MixerDevice->ModulationInterface.IsValid())
 			{
 				VolumeMod.Init(MixerDevice->DeviceID, FName("Volume"), false /* bInIsBuffered */, true /* bInValueLinear */);
-				VolumeModBase = SoundSubmix->OutputVolumeModulation.Value;
-
 				WetLevelMod.Init(MixerDevice->DeviceID, FName("Volume"), false /* bInIsBuffered */, true /* bInValueLinear */);
-				WetModBase = SoundSubmix->WetLevelModulation.Value;
-
 				DryLevelMod.Init(MixerDevice->DeviceID, FName("Volume"), false /* bInIsBuffered */, true /* bInValueLinear */);
-				DryModBase = SoundSubmix->DryLevelModulation.Value;
 
 				USoundModulatorBase* VolumeModulator = SoundSubmix->OutputVolumeModulation.Modulator;
 				USoundModulatorBase* WetLevelModulator = SoundSubmix->WetLevelModulation.Modulator;
@@ -1269,25 +1266,11 @@ namespace Audio
 		DryChannelBuffer.Reset();
 
 		// Update Dry Level using modulator
-		float ModulatedDryLevelStart = CurrentDryLevel;
-		float ModulatedDryLevelEnd = TargetDryLevel;
-
-		const bool bUseModulation = MixerDevice->IsModulationPluginEnabled() && MixerDevice->ModulationInterface.IsValid();
-		if (bUseModulation)
-		{
-			const float PreModulation = DryLevelMod.GetValue();
-			DryLevelMod.ProcessControl(DryModBase);
-			const float PostModulation = DryLevelMod.GetValue();
-
-			if (DryLevelMod.IsActive())
-			{
-				ModulatedDryLevelStart *= DryLevelMod.GetHasProcessed() ? PreModulation : PostModulation;
-				ModulatedDryLevelEnd *= PostModulation;
-			}
-		}
+		DryLevelMod.ProcessControl(DryModBase);
+		TargetDryLevel = DryLevelMod.GetValue();
 
 		// Check if we need to allocate a dry buffer. This is stored here before effects processing. We mix in with wet buffer after effects processing.
-		if (!FMath::IsNearlyEqual(ModulatedDryLevelStart, ModulatedDryLevelEnd) || !FMath::IsNearlyZero(ModulatedDryLevelStart))
+		if (!FMath::IsNearlyEqual(TargetDryLevel, CurrentDryLevel) || !FMath::IsNearlyZero(CurrentDryLevel))
 		{
 			DryChannelBuffer.Append(InputBuffer);
 		}
@@ -1358,32 +1341,19 @@ namespace Audio
 				}
 
 				// Update Wet Level using modulator
-				float ModulatedWetLevelStart = CurrentWetLevel;
-				float ModulatedWetLevelEnd = TargetWetLevel;
-
-				if (bUseModulation)
-				{
-					const float PreModulation = WetLevelMod.GetValue();
-					WetLevelMod.ProcessControl(WetModBase);
-					const float PostModulation = WetLevelMod.GetValue();
-
-					if (WetLevelMod.IsActive())
-					{
-						ModulatedWetLevelStart *= WetLevelMod.GetHasProcessed() ? PreModulation : PostModulation;
-						ModulatedWetLevelEnd *= PostModulation;
-					}
-				}
+				WetLevelMod.ProcessControl(WetModBase);
+				TargetWetLevel = WetLevelMod.GetValue();
 
 				// Apply the wet level here after processing effects. 
-				if (!FMath::IsNearlyEqual(ModulatedWetLevelEnd, ModulatedWetLevelStart) || !FMath::IsNearlyEqual(ModulatedWetLevelStart, 1.0f))
+				if (!FMath::IsNearlyEqual(TargetWetLevel, CurrentWetLevel) || !FMath::IsNearlyEqual(CurrentWetLevel, 1.0f))
 				{
-					if (FMath::IsNearlyEqual(ModulatedWetLevelEnd, ModulatedWetLevelStart))
+					if (FMath::IsNearlyEqual(TargetWetLevel, CurrentWetLevel))
 					{
-						MultiplyBufferByConstantInPlace(InputBuffer, ModulatedWetLevelEnd);
+						MultiplyBufferByConstantInPlace(InputBuffer, TargetWetLevel);
 					}
 					else
 					{
-						FadeBufferFast(InputBuffer, ModulatedWetLevelStart, ModulatedWetLevelEnd);
+						FadeBufferFast(InputBuffer, CurrentWetLevel, TargetWetLevel);
 						CurrentWetLevel = TargetWetLevel;
 					}
 				}
@@ -1394,14 +1364,14 @@ namespace Audio
 		if (DryChannelBuffer.Num() > 0)
 		{
 			// If we've already set the volume, only need to multiply by constant
-			if (FMath::IsNearlyEqual(ModulatedDryLevelEnd, ModulatedDryLevelStart))
+			if (FMath::IsNearlyEqual(TargetDryLevel, CurrentDryLevel))
 			{
-				MultiplyBufferByConstantInPlace(DryChannelBuffer, ModulatedDryLevelEnd);
+				MultiplyBufferByConstantInPlace(DryChannelBuffer, TargetDryLevel);
 			}
 			else
 			{
 				// To avoid popping, we do a fade on the buffer to the target volume
-				FadeBufferFast(DryChannelBuffer, ModulatedDryLevelStart, ModulatedDryLevelEnd);
+				FadeBufferFast(DryChannelBuffer, CurrentDryLevel, TargetDryLevel);
 				CurrentDryLevel = TargetDryLevel;
 			}
 			MixInBufferFast(DryChannelBuffer, InputBuffer);
@@ -1467,34 +1437,21 @@ namespace Audio
 		}
 
 		// Update output volume using modulator
-		float ModulatedOutputVolumeStart = CurrentOutputVolume;
-		float ModulatedOutputVolumeEnd = TargetOutputVolume;
-
-		if (bUseModulation)
-		{
-			const float PreModulation = VolumeMod.GetValue();
-			VolumeMod.ProcessControl(VolumeModBase);
-			const float PostModulation = VolumeMod.GetValue();
-
-			if (VolumeMod.IsActive())
-			{
-				ModulatedOutputVolumeStart *= VolumeMod.GetHasProcessed() ? PreModulation : PostModulation;
-				ModulatedOutputVolumeEnd *= PostModulation;
-			}
-		}
+		VolumeMod.ProcessControl(VolumeModBase);
+		TargetOutputVolume = VolumeMod.GetValue();
 
 		// Now apply the output volume
-		if (!FMath::IsNearlyEqual(ModulatedOutputVolumeEnd, ModulatedOutputVolumeStart) || !FMath::IsNearlyEqual(ModulatedOutputVolumeStart, 1.0f))
+		if (!FMath::IsNearlyEqual(TargetOutputVolume, CurrentOutputVolume) || !FMath::IsNearlyEqual(CurrentOutputVolume, 1.0f))
 		{
 			// If we've already set the output volume, only need to multiply by constant
-			if (FMath::IsNearlyEqual(ModulatedOutputVolumeEnd, ModulatedOutputVolumeStart))
+			if (FMath::IsNearlyEqual(TargetOutputVolume, CurrentOutputVolume))
 			{
-				Audio::MultiplyBufferByConstantInPlace(InputBuffer, ModulatedOutputVolumeEnd);
+				Audio::MultiplyBufferByConstantInPlace(InputBuffer, TargetOutputVolume);
 			}
 			else
 			{
 				// To avoid popping, we do a fade on the buffer to the target volume
-				Audio::FadeBufferFast(InputBuffer, ModulatedOutputVolumeStart, ModulatedOutputVolumeEnd);
+				Audio::FadeBufferFast(InputBuffer, CurrentOutputVolume, TargetOutputVolume);
 				CurrentOutputVolume = TargetOutputVolume;
 			}
 		}
@@ -2292,17 +2249,17 @@ namespace Audio
 
 	void FMixerSubmix::SetOutputVolume(float InOutputVolume)
 	{
-		TargetOutputVolume = FMath::Clamp(InOutputVolume, 0.0f, 1.0f);
+		VolumeModBase = FMath::Clamp(InOutputVolume, 0.0f, 1.0f);
 	}
 
 	void FMixerSubmix::SetDryLevel(float InDryLevel)
 	{
-		TargetDryLevel = FMath::Clamp(InDryLevel, 0.0f, 1.0f);
+		DryModBase = FMath::Clamp(InDryLevel, 0.0f, 1.0f);
 	}
 
 	void FMixerSubmix::SetWetLevel(float InWetLevel)
 	{
-		TargetWetLevel = FMath::Clamp(InWetLevel, 0.0f, 1.0f);
+		WetModBase = FMath::Clamp(InWetLevel, 0.0f, 1.0f);
 	}
 
 	void FMixerSubmix::UpdateModulationSettings(USoundModulatorBase* InOutputModulator, USoundModulatorBase* InWetLevelModulator, USoundModulatorBase* InDryLevelModulator)
diff --git a/Engine/Source/Runtime/AudioMixer/Public/AudioMixerSubmix.h b/Engine/Source/Runtime/AudioMixer/Public/AudioMixerSubmix.h
--- a/Engine/Source/Runtime/AudioMixer/Public/AudioMixerSubmix.h
+++ b/Engine/Source/Runtime/AudioMixer/Public/AudioMixerSubmix.h
@@ -529,9 +529,9 @@ namespace Audio
 		FModulationDestination DryLevelMod;
 		FModulationDestination WetLevelMod;
 
-		float VolumeModBase;
-		float DryModBase;
-		float WetModBase;
+		float VolumeModBase = 1.f;
+		float DryModBase = 0.f;
+		float WetModBase = 1.f;
 
 		// Envelope following data
 		float EnvelopeValues[AUDIO_MIXER_MAX_OUTPUT_CHANNELS];
diff --git a/Engine/Source/Runtime/Engine/Classes/Sound/SoundSubmix.h b/Engine/Source/Runtime/Engine/Classes/Sound/SoundSubmix.h
--- a/Engine/Source/Runtime/Engine/Classes/Sound/SoundSubmix.h
+++ b/Engine/Source/Runtime/Engine/Classes/Sound/SoundSubmix.h
@@ -283,46 +283,28 @@ public:
 	UPROPERTY(EditAnywhere, BlueprintReadWrite, Category = EnvelopeFollower, meta = (ClampMin = "0", UIMin = "0"))
 	int32 EnvelopeFollowerReleaseTime;
 
-	/** Whether to treat submix gain levels as linear or decibel values. */
-	UPROPERTY(EditAnywhere, Category = SubmixLevel, meta = (InlineCategoryProperty))
-	EGainParamMode GainMode;
-
-	/** The output volume of the submix. Applied after submix effects and analysis are performed.*/
-	UPROPERTY(EditAnywhere, BlueprintReadWrite, Category = SubmixLevel, meta = (ClampMin = "0.0", ClampMax = "1.0", UIMin = "0.0", UIMax = "1.0", EditCondition = "GainMode == EGainParamMode::Linear", DisplayName = "Output Volume", EditConditionHides))
+	/** Deprecated -- The output volume of the submix. Applied after submix effects and analysis are performed.*/
+	UPROPERTY(meta = (DeprecatedProperty, DeprecationMessage = "5.0 - Removed in favor of OutputVolumeModulation."))
 	float OutputVolume;
 
-	/** The wet level of the submix. Applied after submix effects and analysis are performed. */
-	UPROPERTY(EditAnywhere, BlueprintReadWrite, Category = SubmixLevel, meta = (ClampMin = "0.0", ClampMax = "1.0", UIMin = "0.0", UIMax = "1.0", EditCondition = "GainMode == EGainParamMode::Linear", DisplayName = "Wet Level", EditConditionHides))
+	/** Deprecated -- The wet level of the submix. Applied after submix effects and analysis are performed. */
+	UPROPERTY(meta = (DeprecatedProperty, DeprecationMessage = "5.0 - Removed in favor of WetLevelModulation."))
 	float WetLevel;
 
-	/** The dry level of the submix. Applied before submix effects and analysis are performed. */
-	UPROPERTY(EditAnywhere, BlueprintReadWrite, Category = SubmixLevel, meta = (ClampMin = "0.0", ClampMax = "1.0", UIMin = "0.0", UIMax = "1.0", EditCondition = "GainMode == EGainParamMode::Linear", DisplayName = "Dry Level", EditConditionHides))
+	/** Deprecated -- The dry level of the submix. Applied before submix effects and analysis are performed. */
+	UPROPERTY(meta = (DeprecatedProperty, DeprecationMessage = "5.0 - Removed in favor of DryLevelModulation."))
 	float DryLevel;
 
-#if WITH_EDITORONLY_DATA
-	/** The output volume of the submix (in dB). */
-	UPROPERTY(EditAnywhere, BlueprintReadWrite, Category = SubmixLevel, meta = (ClampMin = "-160.0", ClampMax = "0.0", UIMin = "-60.0", UIMax = "0.0", EditCondition = "GainMode == EGainParamMode::Decibels", DisplayName = "Output Volume (dB)", EditConditionHides))
-	float OutputVolumeDB;
-
-	/** The wet level of the submix  (in dB). Applied after submix effects and analysis are performed. */
-	UPROPERTY(EditAnywhere, BlueprintReadWrite, Category = SubmixLevel, meta = (ClampMin = "-160.0", ClampMax = "0.0", UIMin = "-60.0", UIMax = "0.0", EditCondition = "GainMode == EGainParamMode::Decibels", DisplayName = "Wet Level (dB)", EditConditionHides))
-	float WetLevelDB;
-
-	/** The dry level of the submix  (in dB)s. Applied before submix effects and analysis are performed. */
-	UPROPERTY(EditAnywhere, BlueprintReadWrite, Category = SubmixLevel, meta = (ClampMin = "-160.0", ClampMax = "0.0", UIMin = "-60.0", UIMax = "0.0", EditCondition = "GainMode == EGainParamMode::Decibels", DisplayName = "Dry Level (dB)", EditConditionHides))
-	float DryLevelDB;
-#endif
-
-	/** Modulation to apply to the submix Output Volume (in dB)*/
-	UPROPERTY(EditAnywhere, BlueprintReadWrite, Category = Modulation, meta = (DisplayName = "Output Volume Modulation", AudioParam = "Volume"))
+	/** The output volume of the submix. Applied after submix effects and analysis are performed.*/
+	UPROPERTY(EditAnywhere, BlueprintReadWrite, Category = SubmixLevel, meta = (DisplayName = "Output Volume", AudioParam = "Volume", ClampMin = "0.0", ClampMax = "1.0", UIMin = "0.0", UIMax = "1.0"))
 	FSoundModulationDestinationSettings OutputVolumeModulation;
 
-	/** Modulation to apply to the submix Wet Level (in dB)*/
-	UPROPERTY(EditAnywhere, BlueprintReadWrite, Category = Modulation, meta = (DisplayName = "Wet Level Modulation", AudioParam = "Volume"))
+	/** The wet level of the submix. Applied after submix effects and analysis are performed. */
+	UPROPERTY(EditAnywhere, BlueprintReadWrite, Category = SubmixLevel, meta = (DisplayName = "Wet Level", AudioParam = "Volume", ClampMin = "0.0", ClampMax = "1.0", UIMin = "0.0", UIMax = "1.0"))
 	FSoundModulationDestinationSettings WetLevelModulation;
 
-	/** Modulation to apply to the submix Dry Level (in dB)*/
-	UPROPERTY(EditAnywhere, BlueprintReadWrite, Category = Modulation, meta = (DisplayName = "Dry Level Modulation", AudioParam = "Volume"))
+	/** The dry level of the submix. Applied before submix effects and analysis are performed. */
+	UPROPERTY(EditAnywhere, BlueprintReadWrite, Category = SubmixLevel, meta = (DisplayName = "Dry Level", AudioParam = "Volume", ClampMin = "0.0", ClampMax = "1.0", UIMin = "0.0", UIMax = "1.0"))
 	FSoundModulationDestinationSettings DryLevelModulation;
 	
 	/** Optional Audio Link Settings Object */
@@ -408,7 +390,8 @@ public:
 	static FSoundSpectrumAnalyzerDelegateSettings GetSpectrumAnalysisDelegateSettings(const TArray<FSoundSubmixSpectralAnalysisBandSettings>& InBandSettings, float UpdateRate, float DecibelNoiseFloor, bool bDoNormalize, bool bDoAutoRange, float AutoRangeAttackTime, float AutoRangeReleaseTime);
 protected:
 
-	virtual void PostLoad() override;
+	virtual void Serialize(FArchive& Ar) override;
+
 #if WITH_EDITOR
 	virtual void PostEditChangeProperty(struct FPropertyChangedEvent& PropertyChangedEvent) override;
 #endif
diff --git a/Engine/Source/Runtime/Engine/Private/SoundSubmix.cpp b/Engine/Source/Runtime/Engine/Private/SoundSubmix.cpp
--- a/Engine/Source/Runtime/Engine/Private/SoundSubmix.cpp
+++ b/Engine/Source/Runtime/Engine/Private/SoundSubmix.cpp
@@ -43,31 +43,41 @@ USoundSubmix::USoundSubmix(const FObjectInitializer& ObjectInitializer)
 	, AmbisonicsPluginSettings(nullptr)
 	, EnvelopeFollowerAttackTime(10)
 	, EnvelopeFollowerReleaseTime(500)
-	, GainMode(EGainParamMode::Linear)
 	, OutputVolume(1.0f)
 	, WetLevel(1.0f)
 	, DryLevel(0.0f)
-#if WITH_EDITOR
-	, OutputVolumeDB(0.0f)
-	, WetLevelDB(0.0f)
-	, DryLevelDB(-120.0f)
-#endif
 {
-	OutputVolumeModulation.Value = 0.f;
-	WetLevelModulation.Value = 0.f;
+	OutputVolumeModulation.Value = 1.f;
+	WetLevelModulation.Value = 1.f;
 	DryLevelModulation.Value = 0.f;
 }
 
-void USoundSubmix::PostLoad()
+void USoundSubmix::Serialize(FArchive& Ar)
 {
-	Super::PostLoad();
+	Super::Serialize(Ar);
 
-#if WITH_EDITOR
-	OutputVolumeDB = Audio::ConvertToDecibels(OutputVolume);
-	WetLevelDB = Audio::ConvertToDecibels(WetLevel);
-	DryLevelDB = Audio::ConvertToDecibels(DryLevel);
-#endif
+#if WITH_EDITORONLY_DATA
+	if (Ar.IsLoading() || Ar.IsSaving())
+	{
+		if (OutputVolume < 1.0f)
+		{
+			OutputVolumeModulation.Value = OutputVolume;
+			OutputVolume = 1.0f;
+		}
+
+		if (WetLevel < 1.0f)
+		{
+			OutputVolumeModulation.Value = WetLevel;
+			WetLevel = 1.0f;
+		}
 
+		if (DryLevel > 0.0f)
+		{
+			OutputVolumeModulation.Value = DryLevel;
+			DryLevel = 0.0f;
+		}
+	}
+#endif // WITH_EDITORONLY_DATA
 }
 
 UEndpointSubmix::UEndpointSubmix(const FObjectInitializer& ObjectInitializer)
@@ -382,72 +392,6 @@ void USoundSubmix::PostEditChangeProperty(struct FPropertyChangedEvent& Property
 		{
 			FName MemberName = PropertyChangedEvent.MemberProperty->GetFName();
 
-			bool bUpdateSubmixGain = false;
-
-			if (MemberName == GET_MEMBER_NAME_CHECKED(USoundSubmix, OutputVolume))
-			{
-				OutputVolumeDB = Audio::ConvertToDecibels(OutputVolume);
-				bUpdateSubmixGain = true;
-			}
-			else if (MemberName == GET_MEMBER_NAME_CHECKED(USoundSubmix, WetLevel))
-			{
-				WetLevelDB = Audio::ConvertToDecibels(OutputVolume);
-				bUpdateSubmixGain = true;
-			}
-			else if (MemberName == GET_MEMBER_NAME_CHECKED(USoundSubmix, DryLevel))
-			{
-				DryLevelDB = Audio::ConvertToDecibels(DryLevel);
-				bUpdateSubmixGain = true;
-			}
-			else if (MemberName == GET_MEMBER_NAME_CHECKED(USoundSubmix, OutputVolumeDB))
-			{
-				if (OutputVolumeDB <= -160.f)
-				{
-					OutputVolume = 0.0f;
-				}
-				else
-				{
-					OutputVolume = Audio::ConvertToLinear(OutputVolumeDB);
-				}
-				bUpdateSubmixGain = true;
-			}
-			else if (MemberName == GET_MEMBER_NAME_CHECKED(USoundSubmix, WetLevelDB))
-			{
-				if (WetLevelDB <= -120.f)
-				{
-					WetLevel = 0.0f;
-				}
-				else
-				{
-					WetLevel = Audio::ConvertToLinear(WetLevelDB);
-				}
-				bUpdateSubmixGain = true;
-			}
-			else if (MemberName == GET_MEMBER_NAME_CHECKED(USoundSubmix, DryLevelDB))
-			{
-				if (DryLevelDB <= -120.0f)
-				{
-					DryLevel = 0.0f;
-				}
-				else
-				{
-					DryLevel = Audio::ConvertToLinear(DryLevelDB);
-				}
-				bUpdateSubmixGain = true;
-			}
-
-			if (bUpdateSubmixGain)
-			{
-				const float NewOutputVolume = OutputVolume;
-				const float NewWetLevel = WetLevel;
-				const float NewDryLevel = DryLevel;
-				USoundSubmix* SoundSubmix = this;
-				AudioDeviceManager->IterateOverAllDevices([SoundSubmix, NewOutputVolume, NewWetLevel, NewDryLevel](Audio::FDeviceId Id, FAudioDevice* Device)
-				{
-					Device->SetSubmixWetDryLevel(SoundSubmix, NewOutputVolume, NewWetLevel, NewDryLevel);
-				});
-			}
-
 			if (MemberName == GET_MEMBER_NAME_CHECKED(USoundSubmix, bAutoDisable))
 			{
 				const bool NewAutoDisable = bAutoDisable;
