Fix sound submixes using linear volume values in decibel contexts when the modulation plugin is enabled.

diff --git a/Engine/Source/Editor/AudioEditor/Private/SoundModulationDestinationLayout.cpp b/Engine/Source/Editor/AudioEditor/Private/SoundModulationDestinationLayout.cpp
--- a/Engine/Source/Editor/AudioEditor/Private/SoundModulationDestinationLayout.cpp
+++ b/Engine/Source/Editor/AudioEditor/Private/SoundModulationDestinationLayout.cpp
@@ -123,33 +123,39 @@ namespace ModDestinationLayoutUtils
 		OutParamName = ModDestinationLayoutUtils::GetParameterNameFromMetaData(StructPropertyHandle);
 		if (OutParamName != FName())
 		{
-			// If parameter was provided, it overrides ClampMin/Max.  User data however overrides UIMin/Max if its in clamp range.
+			// If parameter was provided, it overrides ClampMin/Max.
 			const Audio::FModulationParameter& Parameter = Audio::GetModulationParameter(OutParamName);
-			UIMinValue = Parameter.MinValue;
-			UIMaxValue = Parameter.MaxValue;
-			ClampMinValue = UIMinValue;
-			ClampMaxValue = UIMaxValue;
-			OutUnitDisplayText = Parameter.UnitDisplayName;
-			if (bClampValuesSet)
-			{
-				UE_LOG(LogAudioEditor, Warning, TEXT("ClampMin/Max overridden by AudioModulation plugin asset with ParamName '%s'."), *OutParamName.ToString());
-			}
 
-			if (StructPropertyHandle->HasMetaData("UIMin"))
+			// if no valid parameter was found & the user has specified their own clamping, don't override it with the default parameter range of [0,1]
+			if (false == (Parameter.ParameterName == FName() && bClampValuesSet))
 			{
-				float NewMin = UIMinValue;
-				FString ParamString = StructPropertyHandle->GetMetaData("UIMin");
-				NewMin = FCString::Atof(*ParamString);
-				UIMinValue = FMath::Clamp(NewMin, ClampMinValue, ClampMaxValue);
+				UIMinValue = Parameter.MinValue;
+				UIMaxValue = Parameter.MaxValue;
+				ClampMinValue = UIMinValue;
+				ClampMaxValue = UIMaxValue;
+				OutUnitDisplayText = Parameter.UnitDisplayName;
+				if (bClampValuesSet)
+				{
+					UE_LOG(LogAudioEditor, Warning, TEXT("ClampMin/Max overridden by AudioModulation plugin asset with ParamName '%s'."), *OutParamName.ToString());
+				}
 			}
+		}
 
-			if (StructPropertyHandle->HasMetaData("UIMax"))
-			{
-				float NewMax = UIMaxValue;
-				FString ParamString = StructPropertyHandle->GetMetaData("UIMax");
-				NewMax = FCString::Atof(*ParamString);
-				UIMaxValue = FMath::Clamp(NewMax, ClampMinValue, ClampMaxValue);
-			}
+		// User data overrides UIMin/Max if its in clamp range.
+		if (StructPropertyHandle->HasMetaData("UIMin"))
+		{
+			float NewMin = UIMinValue;
+			FString ParamString = StructPropertyHandle->GetMetaData("UIMin");
+			NewMin = FCString::Atof(*ParamString);
+			UIMinValue = FMath::Clamp(NewMin, ClampMinValue, ClampMaxValue);
+		}
+
+		if (StructPropertyHandle->HasMetaData("UIMax"))
+		{
+			float NewMax = UIMaxValue;
+			FString ParamString = StructPropertyHandle->GetMetaData("UIMax");
+			NewMax = FCString::Atof(*ParamString);
+			UIMaxValue = FMath::Clamp(NewMax, ClampMinValue, ClampMaxValue);
 		}
 
 		ValueHandle->SetInstanceMetaData("ClampMin", FString::Printf(TEXT("%f"), ClampMinValue));
diff --git a/Engine/Source/Runtime/AudioMixer/Private/AudioMixerSubmix.cpp b/Engine/Source/Runtime/AudioMixer/Private/AudioMixerSubmix.cpp
--- a/Engine/Source/Runtime/AudioMixer/Private/AudioMixerSubmix.cpp
+++ b/Engine/Source/Runtime/AudioMixer/Private/AudioMixerSubmix.cpp
@@ -253,13 +253,13 @@ namespace Audio
 
 		if (const USoundSubmix* SoundSubmix = Cast<const USoundSubmix>(OwningSubmixObject))
 		{
-			VolumeModBase = FMath::Clamp(SoundSubmix->OutputVolumeModulation.Value, 0.0f, 1.0f);
-			WetModBase = FMath::Clamp(SoundSubmix->WetLevelModulation.Value, 0.0f, 1.0f);
-			DryModBase = FMath::Clamp(SoundSubmix->DryLevelModulation.Value, 0.0f, 1.0f);
+			VolumeModBaseDb = FMath::Clamp(SoundSubmix->OutputVolumeModulation.Value, MIN_VOLUME_DECIBELS, 0.0f);
+			WetModBaseDb = FMath::Clamp(SoundSubmix->WetLevelModulation.Value, MIN_VOLUME_DECIBELS, 0.0f);
+			DryModBaseDb = FMath::Clamp(SoundSubmix->DryLevelModulation.Value, MIN_VOLUME_DECIBELS, 0.0f);
 
-			TargetOutputVolume = VolumeModBase;
-			TargetWetLevel = WetModBase;
-			TargetDryLevel = DryModBase;
+			TargetOutputVolume = Audio::ConvertToLinear(VolumeModBaseDb);
+			TargetWetLevel = Audio::ConvertToLinear(WetModBaseDb);
+			TargetDryLevel = Audio::ConvertToLinear(DryModBaseDb);
 
 			CurrentOutputVolume = TargetOutputVolume;
 			CurrentDryLevel = TargetDryLevel;
@@ -1266,8 +1266,17 @@ namespace Audio
 		DryChannelBuffer.Reset();
 
 		// Update Dry Level using modulator
-		DryLevelMod.ProcessControl(DryModBase);
-		TargetDryLevel = DryLevelMod.GetValue();
+		if (MixerDevice->IsModulationPluginEnabled() && MixerDevice->ModulationInterface.IsValid())
+		{
+			DryLevelMod.ProcessControl(DryModBaseDb);
+			TargetDryLevel = DryLevelMod.GetValue();
+		}
+		else
+		{
+			TargetDryLevel = Audio::ConvertToLinear(DryModBaseDb);
+		}
+
+		TargetDryLevel *= DryLevelModifier;
 
 		// Check if we need to allocate a dry buffer. This is stored here before effects processing. We mix in with wet buffer after effects processing.
 		if (!FMath::IsNearlyEqual(TargetDryLevel, CurrentDryLevel) || !FMath::IsNearlyZero(CurrentDryLevel))
@@ -1341,8 +1350,17 @@ namespace Audio
 				}
 
 				// Update Wet Level using modulator
-				WetLevelMod.ProcessControl(WetModBase);
-				TargetWetLevel = WetLevelMod.GetValue();
+				if (MixerDevice->IsModulationPluginEnabled() && MixerDevice->ModulationInterface.IsValid())
+				{
+					WetLevelMod.ProcessControl(WetModBaseDb);
+					TargetWetLevel = WetLevelMod.GetValue();
+				}
+				else
+				{
+					TargetWetLevel = Audio::ConvertToLinear(WetModBaseDb);
+				}
+
+				TargetWetLevel *= WetLevelModifier;
 
 				// Apply the wet level here after processing effects. 
 				if (!FMath::IsNearlyEqual(TargetWetLevel, CurrentWetLevel) || !FMath::IsNearlyEqual(CurrentWetLevel, 1.0f))
@@ -1437,8 +1455,17 @@ namespace Audio
 		}
 
 		// Update output volume using modulator
-		VolumeMod.ProcessControl(VolumeModBase);
-		TargetOutputVolume = VolumeMod.GetValue();
+		if(MixerDevice->IsModulationPluginEnabled() && MixerDevice->ModulationInterface.IsValid())
+		{
+			VolumeMod.ProcessControl(VolumeModBaseDb);
+			TargetOutputVolume = VolumeMod.GetValue();
+		}
+		else
+		{
+			TargetOutputVolume = Audio::ConvertToLinear(VolumeModBaseDb);
+		}
+
+		TargetOutputVolume *= VolumeModifier;
 
 		// Now apply the output volume
 		if (!FMath::IsNearlyEqual(TargetOutputVolume, CurrentOutputVolume) || !FMath::IsNearlyEqual(CurrentOutputVolume, 1.0f))
@@ -2249,17 +2276,17 @@ namespace Audio
 
 	void FMixerSubmix::SetOutputVolume(float InOutputVolume)
 	{
-		VolumeModBase = FMath::Clamp(InOutputVolume, 0.0f, 1.0f);
+		VolumeModifier = FMath::Clamp(InOutputVolume, 0.0f, 1.0f);
 	}
 
 	void FMixerSubmix::SetDryLevel(float InDryLevel)
 	{
-		DryModBase = FMath::Clamp(InDryLevel, 0.0f, 1.0f);
+		DryLevelModifier = FMath::Clamp(InDryLevel, 0.0f, 1.0f);
 	}
 
 	void FMixerSubmix::SetWetLevel(float InWetLevel)
 	{
-		WetModBase = FMath::Clamp(InWetLevel, 0.0f, 1.0f);
+		WetLevelModifier = FMath::Clamp(InWetLevel, 0.0f, 1.0f);
 	}
 
 	void FMixerSubmix::UpdateModulationSettings(USoundModulatorBase* InOutputModulator, USoundModulatorBase* InWetLevelModulator, USoundModulatorBase* InDryLevelModulator)
@@ -2269,11 +2296,11 @@ namespace Audio
 		DryLevelMod.UpdateModulator(InDryLevelModulator);
 	}
 
-	void FMixerSubmix::SetModulationBaseLevels(float InVolumeModBase, float InWetModBase, float InDryModBase)
+	void FMixerSubmix::SetModulationBaseLevels(float InVolumeModBaseDb, float InWetModBaseDb, float InDryModBaseDb)
 	{
-		VolumeModBase = InVolumeModBase;
-		WetModBase = InWetModBase;
-		DryModBase = InDryModBase;
+		VolumeModBaseDb = InVolumeModBaseDb;
+		WetModBaseDb = InWetModBaseDb;
+		DryModBaseDb = InDryModBaseDb;
 	}
 
 	void FMixerSubmix::BroadcastDelegates()
diff --git a/Engine/Source/Runtime/AudioMixer/Public/AudioMixerSubmix.h b/Engine/Source/Runtime/AudioMixer/Public/AudioMixerSubmix.h
--- a/Engine/Source/Runtime/AudioMixer/Public/AudioMixerSubmix.h
+++ b/Engine/Source/Runtime/AudioMixer/Public/AudioMixerSubmix.h
@@ -3,6 +3,7 @@
 #pragma once
 
 #include "AudioMixer.h"
+#include "AudioDefines.h"
 #include "CoreMinimal.h"
 #include "SampleBuffer.h"
 #include "IAudioEndpoint.h"
@@ -131,20 +132,20 @@ namespace Audio
 		// Removes the given submix from this submix's children
 		void RemoveChildSubmix(TWeakPtr<FMixerSubmix, ESPMode::ThreadSafe> SubmixWeakPtr);
 
-		// Sets the output level of the submix
+		// Sets the output level of the submix in linear gain
 		void SetOutputVolume(float InOutputLevel);
 
-		// Sets the static output volume of the submix
+		// Sets the static output volume of the submix in linear gain
 		void SetDryLevel(float InDryLevel);
 
-		// Sets the wet level of the submix
+		// Sets the wet level of the submix in linear gain
 		void SetWetLevel(float InWetLevel);
 
 		// Update modulation settings of the submix
 		void UpdateModulationSettings(USoundModulatorBase* InOutputModulator, USoundModulatorBase* InWetLevelModulator, USoundModulatorBase* InDryLevelModulator);
 
-		// Update modulation settings of the submix
-		void SetModulationBaseLevels(float InVolumeModBase, float InWetModeBase, float InDryModBase);
+		// Update modulation settings of the submix with Decibel values
+		void SetModulationBaseLevels(float InVolumeModBaseDb, float InWetModeBaseDb, float InDryModBaseDb);
 
 		// Gets the submix channels channels
 		int32 GetSubmixChannels() const;
@@ -529,9 +530,14 @@ namespace Audio
 		FModulationDestination DryLevelMod;
 		FModulationDestination WetLevelMod;
 
-		float VolumeModBase = 1.f;
-		float DryModBase = 0.f;
-		float WetModBase = 1.f;
+		float VolumeModBaseDb = 0.f;
+		float DryModBaseDb = MIN_VOLUME_DECIBELS;
+		float WetModBaseDb = 0.f;
+
+		// modifiers set from BP code
+		float VolumeModifier = 1.f;
+		float DryLevelModifier = 1.f;
+		float WetLevelModifier = 1.f;
 
 		// Envelope following data
 		float EnvelopeValues[AUDIO_MIXER_MAX_OUTPUT_CHANNELS];
diff --git a/Engine/Source/Runtime/AudioMixerCore/Public/AudioDefines.h b/Engine/Source/Runtime/AudioMixerCore/Public/AudioDefines.h
--- a/Engine/Source/Runtime/AudioMixerCore/Public/AudioDefines.h
+++ b/Engine/Source/Runtime/AudioMixerCore/Public/AudioDefines.h
@@ -27,6 +27,9 @@
 #define MIN_PITCH						0.4f
 #define MAX_PITCH						2.0f
 
+#define MIN_VOLUME_LINEAR				SMALL_NUMBER
+#define MIN_VOLUME_DECIBELS				-160.f
+
 #define MIN_SOUND_PRIORITY				0.0f
 #define MAX_SOUND_PRIORITY				100.0f
 
diff --git a/Engine/Source/Runtime/Engine/Classes/Sound/SoundSubmix.h b/Engine/Source/Runtime/Engine/Classes/Sound/SoundSubmix.h
--- a/Engine/Source/Runtime/Engine/Classes/Sound/SoundSubmix.h
+++ b/Engine/Source/Runtime/Engine/Classes/Sound/SoundSubmix.h
@@ -295,16 +295,16 @@ public:
 	UPROPERTY(meta = (DeprecatedProperty, DeprecationMessage = "5.0 - Removed in favor of DryLevelModulation."))
 	float DryLevel;
 
-	/** The output volume of the submix. Applied after submix effects and analysis are performed.*/
-	UPROPERTY(EditAnywhere, BlueprintReadWrite, Category = SubmixLevel, meta = (DisplayName = "Output Volume", AudioParam = "Volume", ClampMin = "0.0", ClampMax = "1.0", UIMin = "0.0", UIMax = "1.0"))
+	/** The output volume of the submix in Decibels. Applied after submix effects and analysis are performed.*/
+	UPROPERTY(EditAnywhere, BlueprintReadWrite, Category = SubmixLevel, meta = (DisplayName = "Output Volume (dB)", AudioParam = "Volume", ClampMin = "-96.0", ClampMax = "0.0", UIMin = "-96.0", UIMax = "0.0"))
 	FSoundModulationDestinationSettings OutputVolumeModulation;
 
-	/** The wet level of the submix. Applied after submix effects and analysis are performed. */
-	UPROPERTY(EditAnywhere, BlueprintReadWrite, Category = SubmixLevel, meta = (DisplayName = "Wet Level", AudioParam = "Volume", ClampMin = "0.0", ClampMax = "1.0", UIMin = "0.0", UIMax = "1.0"))
+	/** The wet level of the submixin Decibels. Applied after submix effects and analysis are performed. */
+	UPROPERTY(EditAnywhere, BlueprintReadWrite, Category = SubmixLevel, meta = (DisplayName = "Wet Level (dB)", AudioParam = "Volume", ClampMin = "-96.0", ClampMax = "0.0", UIMin = "-96.0", UIMax = "0.0"))
 	FSoundModulationDestinationSettings WetLevelModulation;
 
-	/** The dry level of the submix. Applied before submix effects and analysis are performed. */
-	UPROPERTY(EditAnywhere, BlueprintReadWrite, Category = SubmixLevel, meta = (DisplayName = "Dry Level", AudioParam = "Volume", ClampMin = "0.0", ClampMax = "1.0", UIMin = "0.0", UIMax = "1.0"))
+	/** The dry level of the submix in Decibels. Applied before submix effects and analysis are performed. */
+	UPROPERTY(EditAnywhere, BlueprintReadWrite, Category = SubmixLevel, meta = (DisplayName = "Dry Level (dB)", AudioParam = "Volume", ClampMin = "-96.0", ClampMax = "0.0", UIMin = "-96.0", UIMax = "0.0"))
 	FSoundModulationDestinationSettings DryLevelModulation;
 	
 	/** Optional Audio Link Settings Object */
@@ -381,10 +381,18 @@ public:
 
 	void StopSpectralAnalysis(FAudioDevice* InDevice);
 
-	/** Sets the output volume of the submix. This dynamic volume acts as a multiplier on the OutputVolume property of this submix.  */
-	UFUNCTION(BlueprintCallable, Category = "Audio", meta = (WorldContext = "WorldContextObject"))
+	/** Sets the output volume of the submix in linear gain. This dynamic volume acts as a multiplier on the OutputVolume property of this submix.  */
+	UFUNCTION(BlueprintCallable, Category = "Audio", meta = (WorldContext = "WorldContextObject", DisplayName = "SetSubmixOutputVolume (linear gain)"))
 	void SetSubmixOutputVolume(const UObject* WorldContextObject, float InOutputVolume);
 
+	/** Sets the output volume of the submix in linear gain. This dynamic level acts as a multiplier on the WetLevel property of this submix.  */
+	UFUNCTION(BlueprintCallable, Category = "Audio", meta = (WorldContext = "WorldContextObject", DisplayName = "SetSubmixWetLevel (linear gain)"))
+	void SetSubmixWetLevel(const UObject* WorldContextObject, float InWetLevel);
+
+	/** Sets the output volume of the submix in linear gain. This dynamic level acts as a multiplier on the DryLevel property of this submix.  */
+	UFUNCTION(BlueprintCallable, Category = "Audio", meta = (WorldContext = "WorldContextObject", DisplayName = "SetSubmixDryLevel (linear gain)"))
+	void SetSubmixDryLevel(const UObject* WorldContextObject, float InDryLevel);
+
 	static FSoundSpectrumAnalyzerSettings GetSpectrumAnalyzerSettings(EFFTSize FFTSize, EFFTPeakInterpolationMethod InterpolationMethod, EFFTWindowType WindowType, float HopSize, EAudioSpectrumType SpectrumType);
 
 	static FSoundSpectrumAnalyzerDelegateSettings GetSpectrumAnalysisDelegateSettings(const TArray<FSoundSubmixSpectralAnalysisBandSettings>& InBandSettings, float UpdateRate, float DecibelNoiseFloor, bool bDoNormalize, bool bDoAutoRange, float AutoRangeAttackTime, float AutoRangeReleaseTime);
diff --git a/Engine/Source/Runtime/Engine/Private/SoundSubmix.cpp b/Engine/Source/Runtime/Engine/Private/SoundSubmix.cpp
--- a/Engine/Source/Runtime/Engine/Private/SoundSubmix.cpp
+++ b/Engine/Source/Runtime/Engine/Private/SoundSubmix.cpp
@@ -43,13 +43,13 @@ USoundSubmix::USoundSubmix(const FObjectInitializer& ObjectInitializer)
 	, AmbisonicsPluginSettings(nullptr)
 	, EnvelopeFollowerAttackTime(10)
 	, EnvelopeFollowerReleaseTime(500)
-	, OutputVolume(1.0f)
-	, WetLevel(1.0f)
-	, DryLevel(0.0f)
+	, OutputVolume(-1.0f)
+	, WetLevel(-1.0f)
+	, DryLevel(-1.0f)
 {
-	OutputVolumeModulation.Value = 1.f;
-	WetLevelModulation.Value = 1.f;
-	DryLevelModulation.Value = 0.f;
+	OutputVolumeModulation.Value = 0.f;
+	WetLevelModulation.Value = 0.f;
+	DryLevelModulation.Value = -96.f;
 }
 
 void USoundSubmix::Serialize(FArchive& Ar)
@@ -59,22 +59,39 @@ void USoundSubmix::Serialize(FArchive& Ar)
 #if WITH_EDITORONLY_DATA
 	if (Ar.IsLoading() || Ar.IsSaving())
 	{
-		if (OutputVolume < 1.0f)
+		// convert any old deprecated values to the new value
+		if (OutputVolume > 0.0f)
 		{
-			OutputVolumeModulation.Value = OutputVolume;
-			OutputVolume = 1.0f;
+			OutputVolumeModulation.Value = Audio::ConvertToDecibels(OutputVolume);
+			OutputVolume = -1.0f;
 		}
 
-		if (WetLevel < 1.0f)
+		if (WetLevel > 0.0f)
 		{
-			OutputVolumeModulation.Value = WetLevel;
-			WetLevel = 1.0f;
+			WetLevelModulation.Value = Audio::ConvertToDecibels(WetLevel);
+			WetLevel = -1.0f;
 		}
 
 		if (DryLevel > 0.0f)
 		{
-			OutputVolumeModulation.Value = DryLevel;
-			DryLevel = 0.0f;
+			DryLevelModulation.Value = Audio::ConvertToDecibels(DryLevel);
+			DryLevel = -1.0f;
+		}
+
+		// fix values previously saved as linear values
+		if (OutputVolumeModulation.Value > 0.0f)
+		{
+			OutputVolumeModulation.Value = Audio::ConvertToDecibels(OutputVolumeModulation.Value);
+		}
+
+		if (WetLevelModulation.Value > 0.0f)
+		{
+			WetLevelModulation.Value = Audio::ConvertToDecibels(WetLevelModulation.Value);
+		}
+
+		if (DryLevelModulation.Value > 0.0f)
+		{
+			DryLevelModulation.Value = Audio::ConvertToDecibels(DryLevelModulation.Value);
 		}
 	}
 #endif // WITH_EDITORONLY_DATA
@@ -382,6 +399,38 @@ void USoundSubmix::SetSubmixOutputVolume(const UObject* WorldContextObject, floa
 	}
 }
 
+void USoundSubmix::SetSubmixWetLevel(const UObject* WorldContextObject, float InWetLevel)
+{
+	if (!GEngine)
+	{
+		return;
+	}
+
+	if (UWorld* ThisWorld = GEngine->GetWorldFromContextObject(WorldContextObject, EGetWorldErrorMode::LogAndReturnNull))
+	{
+		if (FAudioDevice* AudioDevice = ThisWorld->GetAudioDeviceRaw())
+		{
+			AudioDevice->SetSubmixWetLevel(this, InWetLevel);
+		}
+	}
+}
+
+void USoundSubmix::SetSubmixDryLevel(const UObject* WorldContextObject, float InDryLevel)
+{
+	if (!GEngine)
+	{
+		return;
+	}
+
+	if (UWorld* ThisWorld = GEngine->GetWorldFromContextObject(WorldContextObject, EGetWorldErrorMode::LogAndReturnNull))
+	{
+		if (FAudioDevice* AudioDevice = ThisWorld->GetAudioDeviceRaw())
+		{
+			AudioDevice->SetSubmixDryLevel(this, InDryLevel);
+		}
+	}
+}
+
 #if WITH_EDITOR
 void USoundSubmix::PostEditChangeProperty(struct FPropertyChangedEvent& PropertyChangedEvent)
 {
