Initialize Submix Volume Modulators on all types of submixes to ensure volume is interpreted correctly

diff --git a/Engine/Source/Runtime/AudioMixer/Private/AudioMixerSubmix.cpp b/Engine/Source/Runtime/AudioMixer/Private/AudioMixerSubmix.cpp
--- a/Engine/Source/Runtime/AudioMixer/Private/AudioMixerSubmix.cpp
+++ b/Engine/Source/Runtime/AudioMixer/Private/AudioMixerSubmix.cpp
@@ -251,6 +251,13 @@ namespace Audio
 		bIsCurrentlyDisabled = false;
 		SilenceTimeStartSeconds = -1.0;
 
+		if (MixerDevice->IsModulationPluginEnabled() && MixerDevice->ModulationInterface.IsValid())
+		{
+			VolumeMod.Init(MixerDevice->DeviceID, FName("Volume"), false /* bInIsBuffered */, true /* bInValueLinear */);
+			WetLevelMod.Init(MixerDevice->DeviceID, FName("Volume"), false /* bInIsBuffered */, true /* bInValueLinear */);
+			DryLevelMod.Init(MixerDevice->DeviceID, FName("Volume"), false /* bInIsBuffered */, true /* bInValueLinear */);
+		}
+
 		if (const USoundSubmix* SoundSubmix = Cast<const USoundSubmix>(OwningSubmixObject))
 		{
 			VolumeModBaseDb = FMath::Clamp(SoundSubmix->OutputVolumeModulation.Value, MIN_VOLUME_DECIBELS, 0.0f);
@@ -267,10 +274,6 @@ namespace Audio
 
 			if(MixerDevice->IsModulationPluginEnabled() && MixerDevice->ModulationInterface.IsValid())
 			{
-				VolumeMod.Init(MixerDevice->DeviceID, FName("Volume"), false /* bInIsBuffered */, true /* bInValueLinear */);
-				WetLevelMod.Init(MixerDevice->DeviceID, FName("Volume"), false /* bInIsBuffered */, true /* bInValueLinear */);
-				DryLevelMod.Init(MixerDevice->DeviceID, FName("Volume"), false /* bInIsBuffered */, true /* bInValueLinear */);
-
 				USoundModulatorBase* VolumeModulator = SoundSubmix->OutputVolumeModulation.Modulator;
 				USoundModulatorBase* WetLevelModulator = SoundSubmix->WetLevelModulation.Modulator;
 				USoundModulatorBase* DryLevelModulator = SoundSubmix->DryLevelModulation.Modulator;
