PR #8277 : Fix SubmixEffect Fade Issue

diff --git a/Engine/Source/Runtime/AudioMixer/Private/AudioMixerSubmix.cpp b/Engine/Source/Runtime/AudioMixer/Private/AudioMixerSubmix.cpp
--- a/Engine/Source/Runtime/AudioMixer/Private/AudioMixerSubmix.cpp
+++ b/Engine/Source/Runtime/AudioMixer/Private/AudioMixerSubmix.cpp
@@ -1307,8 +1307,6 @@ namespace Audio
 				InputData.ListenerTransforms = MixerDevice->GetListenerTransforms();
 				InputData.AudioClock = MixerDevice->GetAudioClock();
 
-				SubmixChainMixBuffer.Reset(NumSamples);
-				SubmixChainMixBuffer.AddZeroed(NumSamples);
 				bool bProcessedAnEffect = false;
 
 				for (int32 EffectChainIndex = EffectChains.Num() - 1; EffectChainIndex >= 0; --EffectChainIndex)
@@ -1334,19 +1332,18 @@ namespace Audio
 					// Prepare the scratch buffer for effect chain processing
 					EffectChainOutputBuffer.SetNumUninitialized(NumSamples);
 
-					bProcessedAnEffect |= GenerateEffectChainAudio(InputData, InputBuffer, FadeInfo.EffectChain, EffectChainOutputBuffer);
 
 					float StartFadeVolume = FadeInfo.FadeVolume.GetValue();
 					FadeInfo.FadeVolume.Update(DeltaTimeSec);
 					float EndFadeVolume = FadeInfo.FadeVolume.GetValue();
 
-					MixInBufferFast(EffectChainOutputBuffer, SubmixChainMixBuffer, StartFadeVolume, EndFadeVolume);
+					bProcessedAnEffect |= GenerateEffectChainAudio(InputData, InputBuffer, FadeInfo.EffectChain, StartFadeVolume, EndFadeVolume, EffectChainOutputBuffer);
 				}
 
 				// If we processed any effects, write over the old input buffer vs mixing into it. This is basically the "wet channel" audio in a submix.
 				if (bProcessedAnEffect)
 				{
-					FMemory::Memcpy((void*)BufferPtr, (void*)SubmixChainMixBuffer.GetData(), sizeof(float)* NumSamples);
+					FMemory::Memcpy((void*)BufferPtr, (void*)EffectChainOutputBuffer.GetData(), sizeof(float)* NumSamples);
 				}
 
 				// Update Wet Level using modulator
@@ -1525,7 +1522,7 @@ namespace Audio
 		}
 	}
 
-	bool FMixerSubmix::GenerateEffectChainAudio(FSoundEffectSubmixInputData& InputData, FAlignedFloatBuffer& InAudioBuffer, TArray<FSoundEffectSubmixPtr>& InEffectChain, FAlignedFloatBuffer& OutBuffer)
+	bool FMixerSubmix::GenerateEffectChainAudio(FSoundEffectSubmixInputData& InputData, AlignedFloatBuffer& InAudioBuffer, TArray<FSoundEffectSubmixPtr>& InEffectChain, const float InStartGain, const float InEndGain, AlignedFloatBuffer& OutBuffer)
 	{
 		// Reset the output scratch buffer
 		ScratchBuffer.Reset(NumSamples);
@@ -1568,7 +1565,10 @@ namespace Audio
 			}
 
 			// Copy the output to the input
-			FMemory::Memcpy((void*)InAudioBuffer.GetData(), (void*)OutputData.AudioBuffer->GetData(), sizeof(float) * NumSamples);
+			SubmixChainMixBuffer.Reset(NumSamples);
+			SubmixChainMixBuffer.AddZeroed(NumSamples);
+			BufferSubtractFast(*OutputData.AudioBuffer, InAudioBuffer, SubmixChainMixBuffer);
+			MixInBufferFast(SubmixChainMixBuffer, InAudioBuffer, InStartGain, InEndGain);
 
 			// Mix in the dry signal directly
 			const float DryLevel = SubmixEffect->GetDryLevel();
diff --git a/Engine/Source/Runtime/AudioMixer/Public/AudioMixerSubmix.h b/Engine/Source/Runtime/AudioMixer/Public/AudioMixerSubmix.h
--- a/Engine/Source/Runtime/AudioMixer/Public/AudioMixerSubmix.h
+++ b/Engine/Source/Runtime/AudioMixer/Public/AudioMixerSubmix.h
@@ -357,7 +357,7 @@ namespace Audio
 		void SubmixCommand(TFunction<void()> Command);
 
 		// Generates audio from the given effect chain into the given buffer
-		bool GenerateEffectChainAudio(FSoundEffectSubmixInputData& InputData, FAlignedFloatBuffer& InAudioBuffer, TArray<FSoundEffectSubmixPtr>& InEffectChain, FAlignedFloatBuffer& OutBuffer);
+		bool GenerateEffectChainAudio(FSoundEffectSubmixInputData& InputData, AlignedFloatBuffer& InAudioBuffer, TArray<FSoundEffectSubmixPtr>& InEffectChain, const float InStartGain, const float InEndGain, AlignedFloatBuffer& OutBuffer);
 
 		// This mixer submix's Id
 		uint32 Id;
