Fix for detecting submix buffer silence for auto disablement

diff --git a/Engine/Source/Runtime/AudioMixer/Private/AudioMixerSubmix.cpp b/Engine/Source/Runtime/AudioMixer/Private/AudioMixerSubmix.cpp
--- a/Engine/Source/Runtime/AudioMixer/Private/AudioMixerSubmix.cpp
+++ b/Engine/Source/Runtime/AudioMixer/Private/AudioMixerSubmix.cpp
@@ -35,12 +35,12 @@ FAutoConsoleVariableRef CVarBypassAllSubmixEffects(
 	TEXT("1: Submix Effects are disabled."),
 	ECVF_Default);
 
-static int32 LogSubmixEnablementCVar = 1;
+static int32 LogSubmixEnablementCVar = 0;
 FAutoConsoleVariableRef CVarLogSubmixEnablement(
-	TEXT("au.LogSubmixEnablement"),
+	TEXT("au.LogSubmixAutoDisable"),
 	LogSubmixEnablementCVar,
-	TEXT("Enables logging of submix enablement state.\n")
-	TEXT("1: Submix enablement logging is on. 0: Submix enablement logging is off."),
+	TEXT("Enables logging of submix disable and enable state.\n")
+	TEXT("1: Submix enablement logging is on. 0: Submix enablement/disablement logging is off."),
 	ECVF_Default);
 
 // Define profiling categories for submixes. 
@@ -192,7 +192,7 @@ namespace Audio
 		, bAutoDisable(true)
 		, bIsSilent(false)
 		, bIsCurrentlyDisabled(false)
-		, AutoDisableTime(0.01)
+		, AutoDisableTime(0.1)
 		, SilenceTimeStartSeconds(-1.0)
 		, bIsSpectrumAnalyzing(false)
 	{
@@ -1141,27 +1141,23 @@ namespace Audio
 			{
 				bIsCurrentlyDisabled = true;
 
-				if (LogSubmixEnablementCVar == 1)
-				{
-					double AudioClock = MixerDevice->GetAudioClock();
-					double TimeSinceSilent = AudioClock - SilenceTimeStartSeconds;
-
-					UE_LOG(LogAudioMixer, Verbose, 
-						TEXT("Submix Disabled. Num Sources: %d, Time Silent: %.2f, Disablement Threshold: %.2f, Submix Name: %s"), 
-						MixerSourceVoices.Num(),
-						(float)TimeSinceSilent,
-						AutoDisableTime,
-						*SubmixName);
-				}
+				UE_CLOG(LogSubmixEnablementCVar == 1, LogAudioMixer, Display,
+					TEXT("Submix Disabled. Num Sources: %d, Time Silent: %.2f, Disablement Threshold: %.2f, Submix Name: %s"), 
+					MixerSourceVoices.Num(),
+					(float)(MixerDevice->GetAudioClock() - SilenceTimeStartSeconds),
+					AutoDisableTime,
+					*SubmixName);
 			}
 
+			// Even though we're silent, broadcast the buffer to any listeners (will be a silent buffer)
+			SendAudioToSubmixBufferListeners(OutAudioBuffer);
 			return;
 		}
 
 		if (bIsCurrentlyDisabled)
 		{
 			bIsCurrentlyDisabled = false;
-			UE_CLOG(LogSubmixEnablementCVar == 1, LogAudioMixer, Verbose, TEXT("Submix Enabled: %s"), *SubmixName);
+			UE_CLOG(LogSubmixEnablementCVar == 1, LogAudioMixer, Display, TEXT("Submix Re-Enabled: %s"), *SubmixName);
 		}
 
 		// If this is a Soundfield Submix, process our soundfield and decode it to a OutAudioBuffer.
@@ -1452,7 +1448,7 @@ namespace Audio
 		}
 
 		// Update output volume using modulator
-		if(MixerDevice->IsModulationPluginEnabled() && MixerDevice->ModulationInterface.IsValid())
+		if (MixerDevice->IsModulationPluginEnabled() && MixerDevice->ModulationInterface.IsValid())
 		{
 			VolumeMod.ProcessControl(VolumeModBaseDb);
 			TargetOutputVolume = VolumeMod.GetValue();
@@ -1486,12 +1482,20 @@ namespace Audio
 		// Once we've finished rendering submix audio, check if the output buffer is silent if we are auto-disabling
 		if (bAutoDisable)
 		{
-			// Get the mean squared value of the buffer. If the value is very small, then it's counted as "silent".
-			float BufferMean = 0.0f;
-			Audio::ArrayMeanSquared(OutAudioBuffer, BufferMean);
+			// Extremely cheap silent buffer detection: as soon as we hit a sample which isn't silent, we flag we're not silent
+			bool bIsNowSilent = true;
+			for (float& SampleValue : OutAudioBuffer)
+			{
+				// As soon as we hit a non-silent sample, we're not silent
+				if (SampleValue > SMALL_NUMBER)
+				{
+					bIsNowSilent = false;
+					break;
+				}
+			}
 
 			// If this is the first time we're silent track when it happens
-			if (BufferMean <= KINDA_SMALL_NUMBER && !bIsSilent)
+			if (bIsNowSilent && !bIsSilent)
 			{
 				bIsSilent = true;
 				SilenceTimeStartSeconds = MixerDevice->GetAudioClock();
@@ -1503,8 +1507,13 @@ namespace Audio
 			}
 		}
 
+		SendAudioToSubmixBufferListeners(OutAudioBuffer);
+	}
+
+	void FMixerSubmix::SendAudioToSubmixBufferListeners(FAlignedFloatBuffer& OutAudioBuffer)
+	{
 		// Now loop through any buffer listeners and feed the listeners the result of this audio callback
-		if(const USoundSubmix* SoundSubmix = Cast<const USoundSubmix>(OwningSubmixObject))
+		if (const USoundSubmix* SoundSubmix = Cast<const USoundSubmix>(OwningSubmixObject))
 		{
 			CSV_SCOPED_TIMING_STAT(Audio, SubmixBufferListeners);
 			SCOPE_CYCLE_COUNTER(STAT_AudioMixerSubmixBufferListeners);
diff --git a/Engine/Source/Runtime/AudioMixer/Public/AudioMixerSubmix.h b/Engine/Source/Runtime/AudioMixer/Public/AudioMixerSubmix.h
--- a/Engine/Source/Runtime/AudioMixer/Public/AudioMixerSubmix.h
+++ b/Engine/Source/Runtime/AudioMixer/Public/AudioMixerSubmix.h
@@ -203,6 +203,8 @@ namespace Audio
 		void ProcessAudio(FAlignedFloatBuffer& OutAudio);
 		void ProcessAudio(ISoundfieldAudioPacket& OutputAudio);
 
+		void SendAudioToSubmixBufferListeners(FAlignedFloatBuffer& OutAudioBuffer);
+
 		// This should be called if this submix doesn't send it's audio to a parent submix,
 		// but rather an external endpoint.
 		void ProcessAudioAndSendToEndpoint();
diff --git a/Engine/Source/Runtime/Engine/Classes/Sound/SoundSubmix.h b/Engine/Source/Runtime/Engine/Classes/Sound/SoundSubmix.h
--- a/Engine/Source/Runtime/Engine/Classes/Sound/SoundSubmix.h
+++ b/Engine/Source/Runtime/Engine/Classes/Sound/SoundSubmix.h
@@ -170,11 +170,11 @@ public:
 	UEdGraph* SoundSubmixGraph;
 #endif
 
-	// Automatically disable the submix if no sound is rendering in the submix. If false, the submix will remain on regardless of if audio is playing in the submix.
+	// Auto-manage enabling and disabling the submix as a CPU optimization. It will be disabled if the submix and all child submixes are silent. It will re-enable if a sound is sent to the submix or a child submix is audible.
 	UPROPERTY(EditAnywhere, BlueprintReadOnly, Category = AutoDisablement)
 	bool bAutoDisable = true;
 
-	/** The minimum amount of time to wait before automatically disabling a submix once no audio is rendering in it. Actual silence time may be slightly higher than this value. */
+	// The minimum amount of time to wait before automatically disabling a submix if it is silent. Will immediately re-enable if source audio is sent to it. 
 	UPROPERTY(EditAnywhere, BlueprintReadWrite, Category = AutoDisablement, meta = (EditCondition = "bAutoDisable"))
 	float AutoDisableTime = 0.01f;
 
