Fix for spatialization plugins that rely on a corresponding reverb plugin to return spatialized audio not fnuctioning if "Enable Submix Sends" is true on the sound SoundWave.

diff --git a/Engine/Source/Runtime/AudioMixer/Private/AudioMixerSourceManager.cpp b/Engine/Source/Runtime/AudioMixer/Private/AudioMixerSourceManager.cpp
--- a/Engine/Source/Runtime/AudioMixer/Private/AudioMixerSourceManager.cpp
+++ b/Engine/Source/Runtime/AudioMixer/Private/AudioMixerSourceManager.cpp
@@ -2487,7 +2487,7 @@ namespace Audio
 				}
 			}
 
-			if (SourceInfo.IsRenderingToSubmixes())
+			if (SourceInfo.IsRenderingToSubmixes() || MixerDevice->bSpatializationIsExternalSend)
 			{
 				// Apply distance attenuation
 				ApplyDistanceAttenuation(SourceInfo, NumSamples);
