Changes to allow for synchronized "kicking" of certain async audio jobs created during a pass through FMixerDevice::OnProcessAudioStream.

Existing workloads should continue to work as they always have, but it is now possible for a procedural sound (SoundGenerator) to specify a "synchronized render queue" that its "audio decode tasks" should be added to. That queue's jobs are then held for the duration of the OnProcessAudioStream function and only kicked when some "owning" system kicks the queued tasks.

In the Harmonix Music specific case, UMusicEmitterComponent (a subclass of a USynthComponent) creates soundgenerators that specify a unique render queue id. Then, in the Harmonix Music system's "post render callback" which is called at the bottom of the FMixerDevice::OnProcessAudioStream, after doing the synchrnous "state updates" for all of the music renderers to keep them in sync, we kick all the queued up tasks so async rendering can proceed.

diff --git a/Engine/Source/Runtime/AudioMixer/Private/AudioMixerBuffer.cpp b/Engine/Source/Runtime/AudioMixer/Private/AudioMixerBuffer.cpp
--- a/Engine/Source/Runtime/AudioMixer/Private/AudioMixerBuffer.cpp
+++ b/Engine/Source/Runtime/AudioMixer/Private/AudioMixerBuffer.cpp
@@ -398,7 +398,7 @@ namespace Audio
 			NewTaskData.SoundWave = InWave;
 
 			check(Buffer->RealtimeAsyncHeaderParseTask == nullptr);
-			Buffer->RealtimeAsyncHeaderParseTask = CreateAudioTask(NewTaskData);
+			Buffer->RealtimeAsyncHeaderParseTask = CreateAudioTask(AudioDevice->DeviceID, NewTaskData);
 
 			Buffer->NumChannels = InWave->NumChannels;
 		}
diff --git a/Engine/Source/Runtime/AudioMixer/Private/AudioMixerDevice.cpp b/Engine/Source/Runtime/AudioMixer/Private/AudioMixerDevice.cpp
--- a/Engine/Source/Runtime/AudioMixer/Private/AudioMixerDevice.cpp
+++ b/Engine/Source/Runtime/AudioMixer/Private/AudioMixerDevice.cpp
@@ -4,6 +4,7 @@
 
 #include "AudioMixerSource.h"
 #include "AudioMixerSourceManager.h"
+#include "AudioMixerSourceDecode.h"
 #include "AudioMixerSubmix.h"
 #include "AudioMixerSourceVoice.h"
 #include "AudioPluginUtilities.h"
@@ -461,6 +462,9 @@ namespace Audio
 				AudioThreadTimingData.AudioThreadTime = 0.0;
 				AudioThreadTimingData.AudioRenderThreadTime = 0.0;
 
+				// Create synchronized Audio Task Queue for this device...
+				CreateSynchronizedAudioTaskQueue((Audio::AudioTaskQueueId)DeviceID);
+
 				// Start streaming audio
 				return AudioMixerPlatform->StartAudioStream();
 			}
@@ -499,6 +503,9 @@ namespace Audio
 			{
 				UnregisterSoundSubmix(*It);
 			}
+			
+			// Destroy the synchronized Audio Task Queue for this device
+			DestroySynchronizedAudioTaskQueue((Audio::AudioTaskQueueId)DeviceID);
 		}
 		
 		// reset all the sound effect presets loaded
@@ -819,6 +826,7 @@ namespace Audio
 		// notify interested parties
 		NotifyAudioDevicePostRender(RenderInfo);
 
+		KickQueuedTasks((Audio::AudioTaskQueueId)DeviceID);
 		return true;
 	}
 
@@ -2695,4 +2703,19 @@ namespace Audio
 		}
 	}
 
+	void FMixerDevice::CreateSynchronizedAudioTaskQueue(AudioTaskQueueId QueueId)
+	{
+		Audio::CreateSynchronizedAudioTaskQueue(QueueId);
+	}
+
+	void FMixerDevice::DestroySynchronizedAudioTaskQueue(AudioTaskQueueId QueueId, bool RunCurrentQueue)
+	{
+		Audio::DestroySynchronizedAudioTaskQueue(QueueId, RunCurrentQueue);
+	}
+
+	int FMixerDevice::KickQueuedTasks(AudioTaskQueueId QueueId)
+	{
+		return Audio::KickQueuedTasks(QueueId);
+	}
+
 }
diff --git a/Engine/Source/Runtime/AudioMixer/Private/AudioMixerSourceBuffer.cpp b/Engine/Source/Runtime/AudioMixer/Private/AudioMixerSourceBuffer.cpp
--- a/Engine/Source/Runtime/AudioMixer/Private/AudioMixerSourceBuffer.cpp
+++ b/Engine/Source/Runtime/AudioMixer/Private/AudioMixerSourceBuffer.cpp
@@ -75,6 +75,7 @@ namespace Audio
 		, NumChannels(InArgs.Buffer->NumChannels)
 		, BufferType(InArgs.Buffer->GetType())
 		, NumPrecacheFrames(InArgs.SoundWave->NumPrecacheFrames)
+		, AuioDeviceID(InArgs.AudioDeviceID)
 		, bInitialized(false)
 		, bBufferFinished(false)
 		, bPlayedCachedBuffer(false)
@@ -353,7 +354,7 @@ namespace Audio
 				NewTaskData.NumSamples = MaxSamples;
 				NewTaskData.NumChannels = NumChannels;
 				check(!AsyncRealtimeAudioTask);
-				AsyncRealtimeAudioTask = CreateAudioTask(NewTaskData);
+				AsyncRealtimeAudioTask = CreateAudioTask(AuioDeviceID, NewTaskData);
 			}
 
 			return false;
@@ -391,7 +392,7 @@ namespace Audio
 
 		FScopeLock Lock(&DecodeTaskCritSec);
 		check(!AsyncRealtimeAudioTask);
-		AsyncRealtimeAudioTask = CreateAudioTask(NewTaskData);
+		AsyncRealtimeAudioTask = CreateAudioTask(AuioDeviceID, NewTaskData);
 
 		return false;
 	}
diff --git a/Engine/Source/Runtime/AudioMixer/Private/AudioMixerSourceBuffer.h b/Engine/Source/Runtime/AudioMixer/Private/AudioMixerSourceBuffer.h
--- a/Engine/Source/Runtime/AudioMixer/Private/AudioMixerSourceBuffer.h
+++ b/Engine/Source/Runtime/AudioMixer/Private/AudioMixerSourceBuffer.h
@@ -139,6 +139,7 @@ namespace Audio
 		int32 NumChannels;
 		Audio::EBufferType::Type BufferType;
 		int32 NumPrecacheFrames;
+		Audio::FDeviceId AuioDeviceID;
 		TArray<uint8> CachedRealtimeFirstBuffer;
 
 		mutable FCriticalSection SoundWaveCritSec;
diff --git a/Engine/Source/Runtime/AudioMixer/Private/AudioMixerSourceDecode.cpp b/Engine/Source/Runtime/AudioMixer/Private/AudioMixerSourceDecode.cpp
--- a/Engine/Source/Runtime/AudioMixer/Private/AudioMixerSourceDecode.cpp
+++ b/Engine/Source/Runtime/AudioMixer/Private/AudioMixerSourceDecode.cpp
@@ -19,6 +19,14 @@ FAutoConsoleVariableRef CVarForceSyncAudioDecodes(
 	TEXT("0: Not Disabled, 1: Disabled"),
 	ECVF_Default);
 
+static int32 ForceSynchronizedAudioTaskKickCvar = 0;
+FAutoConsoleVariableRef CVarForceSynchronizedAudioTaskKick(
+	TEXT("au.ForceSynchronizedAudioTaskKick"),
+	ForceSynchronizedAudioTaskKickCvar,
+	TEXT("Force all Audio Tasks created in one \"audio render frame\" to be queued until they can all be \"kicked\" at once at the end of the frame.\n")
+	TEXT("0: Don't Force, 1: Force"),
+	ECVF_Default);
+
 namespace Audio
 {
 
@@ -287,6 +295,168 @@ public:
 	}
 };
 
+class FSynchronizedProceduralDecodeHandle : public FDecodeHandleBase
+{
+public:
+	FSynchronizedProceduralDecodeHandle(const FProceduralAudioTaskData& InJobData, AudioTaskQueueId InQueueId)
+	{
+		Task = new FAsyncTask<FAsyncDecodeWorker>(InJobData);
+		QueueId = InQueueId; 
+		{
+			FScopeLock Lock(&SynchronizationQuequesLockCs);
+			TArray<FAsyncTask<FAsyncDecodeWorker>*>* Queue = ProceduralRenderingSynchronizationQueues.Find(QueueId);
+			if (Queue)
+			{
+				Queue->Add(Task);
+				return;
+			}
+		}
+
+		// failed to queue it up, so do a normal start...
+		QueueId = 0;
+		if (ForceSyncAudioDecodesCvar || InJobData.bForceSyncDecode)
+		{
+			Task->StartSynchronousTask();
+			return;
+		}
+		Task->StartBackgroundTask();
+	}
+
+	virtual EAudioTaskType GetType() const override
+	{
+		return EAudioTaskType::Procedural;
+	}
+
+	virtual bool IsDone() const override
+	{
+		if (IsQueued())
+		{
+			return false;
+		}
+		return FDecodeHandleBase::IsDone();
+	}
+
+	bool IsQueued() const
+	{
+		if (!QueueId)
+		{
+			return false;
+		}
+		FScopeLock Lock(&SynchronizationQuequesLockCs);
+		TArray<FAsyncTask<FAsyncDecodeWorker>*>* Queue = ProceduralRenderingSynchronizationQueues.Find(QueueId);
+		return (Queue && Queue->Find(Task) != INDEX_NONE);
+	}
+
+	bool Dequeue(bool Run)
+	{
+		FScopeLock Lock(&SynchronizationQuequesLockCs);
+		TArray<FAsyncTask<FAsyncDecodeWorker>*>* Queue = ProceduralRenderingSynchronizationQueues.Find(QueueId);
+		if (!Queue)
+		{
+			return false;
+		}
+		int NumRemoved = Queue->Remove(Task);
+		if (NumRemoved > 0)
+		{
+			if (Run)
+			{
+				Task->StartBackgroundTask();
+			}
+			return true;
+		}
+		return false;
+	}
+
+	virtual void EnsureCompletion() override
+	{
+		{
+			FScopeLock Lock(&SynchronizationQuequesLockCs);
+			// For now, if this is in the queue still (not kicked)
+			// we're going to pull it out of the queue and run it now.
+			// This seems to only happen when first starting a sound as 
+			// the system tries to decode the first chunk of audio 
+			// asynchronously to the audio thread. 
+			Dequeue(/* Run */ true);
+			// Now we can wait on it to complete...
+		}
+		if (Task)
+		{
+			Task->EnsureCompletion(/*bIsLatencySensitive =*/ true);
+		}
+	}
+
+	virtual void CancelTask() override
+	{
+		{
+			FScopeLock Lock(&SynchronizationQuequesLockCs);
+			// If we dequeue it then it never ran and we are done. Otherwise...
+			if (!Dequeue(/* Run */ false))
+			{
+				FDecodeHandleBase::CancelTask();
+			}
+		}
+	}
+
+	virtual void GetResult(FProceduralAudioTaskResults& OutResult) override
+	{
+		Task->EnsureCompletion();
+		const FAsyncDecodeWorker& DecodeWorker = Task->GetTask();
+		OutResult = DecodeWorker.ProceduralResult;
+	}
+
+	static void CreateSynchronizedRenderQueued(AudioTaskQueueId QueueId)
+	{
+		FScopeLock Lock(&SynchronizationQuequesLockCs);
+		TArray<FAsyncTask<FAsyncDecodeWorker>*>* Queue = ProceduralRenderingSynchronizationQueues.Find(QueueId);
+		if (!Queue)
+		{
+			ProceduralRenderingSynchronizationQueues.Add(QueueId);
+			ProceduralRenderingSynchronizationQueues[QueueId].Reserve(128);
+		}
+	}
+
+	static void DestroySynchronizedRenderQueued(AudioTaskQueueId QueueId, bool RunCurrentQueue = false)
+	{
+		FScopeLock Lock(&SynchronizationQuequesLockCs);
+		
+		if (RunCurrentQueue)
+		{
+			KickQueuedTasks(QueueId);
+		}
+
+		TArray<FAsyncTask<FAsyncDecodeWorker>*>* Queue = ProceduralRenderingSynchronizationQueues.Find(QueueId);
+		if (Queue)
+		{
+			ProceduralRenderingSynchronizationQueues.Remove(QueueId);
+		}
+	}
+
+	static int KickQueuedTasks(AudioTaskQueueId QueueId)
+	{
+		FScopeLock Lock(&SynchronizationQuequesLockCs);
+		int NumStarted = 0;
+		TArray<FAsyncTask<FAsyncDecodeWorker>*>* Queue = ProceduralRenderingSynchronizationQueues.Find(QueueId);
+		if (Queue)
+		{
+			for (auto Task : *Queue)
+			{ 
+				Task->StartBackgroundTask();
+			}
+			NumStarted = Queue->Num();
+			Queue->Empty();
+		}
+		return NumStarted;
+	}
+private:
+	AudioTaskQueueId QueueId = 0;
+
+	static TMap <AudioTaskQueueId, TArray<FAsyncTask<FAsyncDecodeWorker>*>> ProceduralRenderingSynchronizationQueues;
+	static FCriticalSection SynchronizationQuequesLockCs;
+};
+
+TMap <AudioTaskQueueId, TArray<FAsyncTask<FAsyncDecodeWorker>*>> FSynchronizedProceduralDecodeHandle::ProceduralRenderingSynchronizationQueues;
+FCriticalSection FSynchronizedProceduralDecodeHandle::SynchronizationQuequesLockCs;
+
 class FDecodeHandle : public FDecodeHandleBase
 {
 public:
@@ -316,19 +486,47 @@ public:
 	}
 };
 
-IAudioTask* CreateAudioTask(const FProceduralAudioTaskData& InJobData)
+IAudioTask* CreateAudioTask(Audio::FDeviceId InDeviceId, const FProceduralAudioTaskData& InJobData)
 {
+	if (ForceSynchronizedAudioTaskKickCvar || (InJobData.SoundGenerator && InJobData.SoundGenerator->GetSynchronizedRenderQueueId()))
+	{
+		AudioTaskQueueId QueueId = InJobData.SoundGenerator->GetSynchronizedRenderQueueId();
+		if (!QueueId)
+		{
+			// Only use the audio device ID as the task queue id if both 
+			// ForceSynchronizedAudioTaskKickCvar is true AND the caller has
+			// not specified a specific task queue id in their SoundGenerator.
+			QueueId = (AudioTaskQueueId)InDeviceId;
+		}
+		return new FSynchronizedProceduralDecodeHandle(InJobData, QueueId);
+	}
+
 	return new FProceduralDecodeHandle(InJobData);
 }
 
-IAudioTask* CreateAudioTask(const FHeaderParseAudioTaskData& InJobData)
+IAudioTask* CreateAudioTask(Audio::FDeviceId InDeviceId, const FHeaderParseAudioTaskData& InJobData)
 {
 	return new FHeaderDecodeHandle(InJobData);
 }
 
-IAudioTask* CreateAudioTask(const FDecodeAudioTaskData& InJobData)
+IAudioTask* CreateAudioTask(Audio::FDeviceId InDeviceId, const FDecodeAudioTaskData& InJobData)
 {
 	return new FDecodeHandle(InJobData);
 }
 
+void CreateSynchronizedAudioTaskQueue(AudioTaskQueueId QueueId)
+{
+	FSynchronizedProceduralDecodeHandle::CreateSynchronizedRenderQueued(QueueId);
+}
+
+void DestroySynchronizedAudioTaskQueue(AudioTaskQueueId QueueId, bool RunCurrentQueue)
+{
+	FSynchronizedProceduralDecodeHandle::DestroySynchronizedRenderQueued(QueueId, RunCurrentQueue);
+}
+
+int KickQueuedTasks(AudioTaskQueueId QueueId)
+{
+	return FSynchronizedProceduralDecodeHandle::KickQueuedTasks(QueueId);
+}
+
 }
diff --git a/Engine/Source/Runtime/AudioMixer/Private/AudioMixerSourceDecode.h b/Engine/Source/Runtime/AudioMixer/Private/AudioMixerSourceDecode.h
--- a/Engine/Source/Runtime/AudioMixer/Private/AudioMixerSourceDecode.h
+++ b/Engine/Source/Runtime/AudioMixer/Private/AudioMixerSourceDecode.h
@@ -170,12 +170,26 @@ namespace Audio
 	};
 
 	// Creates a task to decode a decoded file header
-	IAudioTask* CreateAudioTask(const FHeaderParseAudioTaskData& InJobData);
+	IAudioTask* CreateAudioTask(Audio::FDeviceId InDeviceId, const FHeaderParseAudioTaskData& InJobData);
 
 	// Creates a task for a procedural sound wave generation
-	IAudioTask* CreateAudioTask(const FProceduralAudioTaskData& InJobData);
+	IAudioTask* CreateAudioTask(Audio::FDeviceId InDeviceId, const FProceduralAudioTaskData& InJobData);
 
 	// Creates a task to decode a chunk of audio
-	IAudioTask* CreateAudioTask(const FDecodeAudioTaskData& InJobData);
+	IAudioTask* CreateAudioTask(Audio::FDeviceId InDeviceId, const FDecodeAudioTaskData& InJobData);
+
+	// Creates a queue for audio decode requests with a specific Id. Tasks
+	// created with this Id will not be started immediately upon creation,
+	// but will instead be queued up to await a start "kick" later. NOTE:
+	// "kicking" the queue is the responsibility of the system that creates 
+	// the queue, typically someplace like in a FOnAudioDevicePostRender delegate! 
+	void CreateSynchronizedAudioTaskQueue(AudioTaskQueueId QueueId);
+
+	// Destroys an audio decode task queue. Tasks currently queued up are 
+	// optionally started.
+	void DestroySynchronizedAudioTaskQueue(AudioTaskQueueId QueueId, bool RunCurrentQueue = false);
+
+	// "Kicks" all of the audio decode tasks currentlyt in the specified queue.
+	int KickQueuedTasks(AudioTaskQueueId QueueId);
 
 }
diff --git a/Engine/Source/Runtime/AudioMixer/Public/AudioMixerDevice.h b/Engine/Source/Runtime/AudioMixer/Public/AudioMixerDevice.h
--- a/Engine/Source/Runtime/AudioMixer/Public/AudioMixerDevice.h
+++ b/Engine/Source/Runtime/AudioMixer/Public/AudioMixerDevice.h
@@ -6,6 +6,7 @@
 #include "AudioMixer.h"
 #include "AudioDevice.h"
 #include "Sound/SoundSubmix.h"
+#include "Sound/SoundGenerator.h"
 #include "DSP/BufferVectorOperations.h"
 #include "DSP/MultithreadedPatching.h"
 #include "Quartz/AudioMixerClockManager.h"
@@ -431,6 +432,20 @@ namespace Audio
 
 		/** Whether or not initialization of the submix system is underway and submixes can be registered */
 		bool bSubmixRegistrationDisabled;
+
+	public:
+
+		// Creates a queue for audio decode requests with a specific Id. Tasks
+		// created with this Id will not be started immediately upon creation,
+		// but will instead be queued up to await a start "kick" later.
+		static void CreateSynchronizedAudioTaskQueue(AudioTaskQueueId QueueId);
+
+		// Destroys an audio decode task queue. Tasks currently queued up are 
+		// optionally started.
+		static void DestroySynchronizedAudioTaskQueue(AudioTaskQueueId QueueId, bool RunCurrentQueue = false);
+
+		// "Kicks" all of the audio decode tasks currentlyt in the queue.
+		static int KickQueuedTasks(AudioTaskQueueId QueueId);
 	};
 }
 
diff --git a/Engine/Source/Runtime/Engine/Classes/Sound/SoundGenerator.h b/Engine/Source/Runtime/Engine/Classes/Sound/SoundGenerator.h
--- a/Engine/Source/Runtime/Engine/Classes/Sound/SoundGenerator.h
+++ b/Engine/Source/Runtime/Engine/Classes/Sound/SoundGenerator.h
@@ -7,6 +7,10 @@
 #include "Containers/Queue.h"
 #include "Templates/Function.h"
 
+namespace Audio
+{
+	using AudioTaskQueueId = uint32;
+}
 
 // Parameters used for constructing a new ISoundGenerator.
 struct FSoundGeneratorInitParams
@@ -44,6 +48,8 @@ public:
 	// Retrieves the next buffer of audio from the generator, called from the audio mixer
 	int32 GetNextBuffer(float* OutAudio, int32 NumSamples, bool bRequireNumberSamples = false);
 
+	virtual Audio::AudioTaskQueueId GetSynchronizedRenderQueueId() const { return 0; }
+
 protected:
 
 	// Protected method to execute lambda in audio render thread
