Expose one more member of the platform settings related to audio device setup. Used the same pattern as shown in GetBufferLength/GetNumOutputFrames in AudioDevice and AudioMixerDevice respectively.

diff --git a/Engine/Source/Runtime/AudioMixer/Public/AudioMixerDevice.h b/Engine/Source/Runtime/AudioMixer/Public/AudioMixerDevice.h
--- a/Engine/Source/Runtime/AudioMixer/Public/AudioMixerDevice.h
+++ b/Engine/Source/Runtime/AudioMixer/Public/AudioMixerDevice.h
@@ -222,7 +222,9 @@ namespace Audio
 		FORCEINLINE int32 GetNumDeviceChannels() const { return PlatformInfo.NumChannels; }
 
 		int32 GetNumOutputFrames() const { return PlatformSettings.CallbackBufferFrameSize; }
-		
+
+		int32 GetNumOutputBuffers() const { return PlatformSettings.NumBuffers; }
+
 		// Retrieve a pointer to the currently active platform. Only use this if you know what you are doing. The returned IAudioMixerPlatformInterface will only be alive as long as this FMixerDevice is alive.
 		IAudioMixerPlatformInterface* GetAudioMixerPlatform() const { return AudioMixerPlatform; }
 
diff --git a/Engine/Source/Runtime/Engine/Public/AudioDevice.h b/Engine/Source/Runtime/Engine/Public/AudioDevice.h
--- a/Engine/Source/Runtime/Engine/Public/AudioDevice.h
+++ b/Engine/Source/Runtime/Engine/Public/AudioDevice.h
@@ -1286,6 +1286,9 @@ public:
 	/** Returns the buffer length of the audio device. */
 	int32 GetBufferLength() const { return PlatformSettings.CallbackBufferFrameSize; }
 
+	/** Returns the number of buffers the audio device keeps queued. */
+	int32 GetNumBuffers() const { return PlatformSettings.NumBuffers; }
+
 	/** Whether or not there's a spatialization plugin enabled. */
 	bool IsSpatializationPluginEnabled() const
 	{
