Fix for submix effect buffer size mismatch
#rnx

diff --git a/Engine/Source/Runtime/AudioMixer/Private/AudioMixerSubmix.cpp b/Engine/Source/Runtime/AudioMixer/Private/AudioMixerSubmix.cpp
--- a/Engine/Source/Runtime/AudioMixer/Private/AudioMixerSubmix.cpp
+++ b/Engine/Source/Runtime/AudioMixer/Private/AudioMixerSubmix.cpp
@@ -1305,6 +1305,10 @@ namespace Audio
 
 				bool bProcessedAnEffect = false;
 
+				// Zero and resize buffer holding mix of current effect chains
+				SubmixChainMixBuffer.Reset(NumSamples);
+				SubmixChainMixBuffer.SetNumZeroed(NumSamples);
+
 				for (int32 EffectChainIndex = EffectChains.Num() - 1; EffectChainIndex >= 0; --EffectChainIndex)
 				{
 					FSubmixEffectFadeInfo& FadeInfo = EffectChains[EffectChainIndex];
@@ -1328,18 +1332,21 @@ namespace Audio
 					// Prepare the scratch buffer for effect chain processing
 					EffectChainOutputBuffer.SetNumUninitialized(NumSamples);
 
+					bProcessedAnEffect |= GenerateEffectChainAudio(InputData, InputBuffer, FadeInfo.EffectChain, EffectChainOutputBuffer);
 
+					// Mix effect chain output into SubmixChainMixBuffer
 					float StartFadeVolume = FadeInfo.FadeVolume.GetValue();
 					FadeInfo.FadeVolume.Update(DeltaTimeSec);
 					float EndFadeVolume = FadeInfo.FadeVolume.GetValue();
 
-					bProcessedAnEffect |= GenerateEffectChainAudio(InputData, InputBuffer, FadeInfo.EffectChain, StartFadeVolume, EndFadeVolume, EffectChainOutputBuffer);
+					// Mix this effect chain with other effect chains.
+					MixInBufferFast(EffectChainOutputBuffer, SubmixChainMixBuffer, StartFadeVolume, EndFadeVolume);
 				}
 
 				// If we processed any effects, write over the old input buffer vs mixing into it. This is basically the "wet channel" audio in a submix.
 				if (bProcessedAnEffect)
 				{
-					FMemory::Memcpy((void*)BufferPtr, (void*)EffectChainOutputBuffer.GetData(), sizeof(float)* NumSamples);
+					FMemory::Memcpy((void*)BufferPtr, (void*)SubmixChainMixBuffer.GetData(), sizeof(float)* NumSamples);
 				}
 
 				// Update Wet Level using modulator
@@ -1531,14 +1538,15 @@ namespace Audio
 		}
 	}
 
-	bool FMixerSubmix::GenerateEffectChainAudio(FSoundEffectSubmixInputData& InputData, AlignedFloatBuffer& InAudioBuffer, TArray<FSoundEffectSubmixPtr>& InEffectChain, const float InStartGain, const float InEndGain, AlignedFloatBuffer& OutBuffer)
+	bool FMixerSubmix::GenerateEffectChainAudio(FSoundEffectSubmixInputData& InputData, const FAlignedFloatBuffer& InAudioBuffer, TArray<FSoundEffectSubmixPtr>& InEffectChain, FAlignedFloatBuffer& OutBuffer)
 	{
-		// Reset the output scratch buffer
-		ScratchBuffer.Reset(NumSamples);
-		ScratchBuffer.AddZeroed(NumSamples);
+		checkf(InAudioBuffer.Num() == OutBuffer.Num(), TEXT("Processing effect chain audio must not alter the number of samples in a buffer. Buffer size mismatch InAudioBuffer[%d] != OutBuffer[%d]"), InAudioBuffer.Num(), OutBuffer.Num());
+
+		// Use the scratch buffer to hold the wet audio. For the first effect in the effect chain, this is the InAudioBuffer. 
+		ScratchBuffer = InAudioBuffer;
 
 		FSoundEffectSubmixOutputData OutputData;
-		OutputData.AudioBuffer = &ScratchBuffer;
+		OutputData.AudioBuffer = &OutBuffer;
 		OutputData.NumChannels = NumChannels;
 
 		const int32 NumOutputFrames = OutBuffer.Num() / NumChannels;
@@ -1554,44 +1562,45 @@ namespace Audio
 
 			// Check to see if we need to down-mix our audio before sending to the submix effect
 			const uint32 ChannelCountOverride = SubmixEffect->GetDesiredInputChannelCountOverride();
+			bool bCurrentEffectProcessedAudio = false;
 
 			if (ChannelCountOverride != INDEX_NONE && ChannelCountOverride != NumChannels)
 			{
 				// Perform the down-mix operation with the down-mixed scratch buffer
 				DownmixedBuffer.SetNumUninitialized(NumOutputFrames * ChannelCountOverride);
-				DownmixBuffer(NumChannels, InAudioBuffer, ChannelCountOverride, DownmixedBuffer);
+				DownmixBuffer(NumChannels, ScratchBuffer, ChannelCountOverride, DownmixedBuffer);
 
 				InputData.NumChannels = ChannelCountOverride;
 				InputData.AudioBuffer = &DownmixedBuffer;
-				SubmixEffect->ProcessAudio(InputData, OutputData);
+				bCurrentEffectProcessedAudio = SubmixEffect->ProcessAudio(InputData, OutputData);
 			}
 			else
 			{
 				// If we're not down-mixing, then just pass in the current wet buffer and our channel count is the same as the output channel count
 				InputData.NumChannels = NumChannels;
-				InputData.AudioBuffer = &InAudioBuffer;
-				SubmixEffect->ProcessAudio(InputData, OutputData);
+				InputData.AudioBuffer = &ScratchBuffer;
+				bCurrentEffectProcessedAudio = SubmixEffect->ProcessAudio(InputData, OutputData);
 			}
 
-			// Copy the output to the input
-			SubmixChainMixBuffer.Reset(NumSamples);
-			SubmixChainMixBuffer.AddZeroed(NumSamples);
-			BufferSubtractFast(*OutputData.AudioBuffer, InAudioBuffer, SubmixChainMixBuffer);
-			MixInBufferFast(SubmixChainMixBuffer, InAudioBuffer, InStartGain, InEndGain);
-
-			// Mix in the dry signal directly
-			const float DryLevel = SubmixEffect->GetDryLevel();
-			if (DryLevel > 0.0f)
+			if (bCurrentEffectProcessedAudio)
 			{
-				MixInBufferFast(InAudioBuffer, ScratchBuffer, DryLevel);
-			}
+				// Mix in the dry signal directly
+				const float DryLevel = SubmixEffect->GetDryLevel();
+				if (DryLevel > 0.0f)
+				{
+					MixInBufferFast(ScratchBuffer, *OutputData.AudioBuffer, DryLevel);
+				}
+				// Copy the output to the input
+				FMemory::Memcpy((void*)ScratchBuffer.GetData(), (void*)OutputData.AudioBuffer->GetData(), sizeof(float) * NumSamples);
 
-			bProcessedAnEffect = true;
+				bProcessedAnEffect = true;
+			}
 		}
 
-		if (bProcessedAnEffect)
+		if (!bProcessedAnEffect)
 		{
-			FMemory::Memcpy((void*)OutBuffer.GetData(), (void*)InAudioBuffer.GetData(), sizeof(float) * NumSamples);
+			// If no effects were processed pass through the input audio to the output buffer.
+			OutBuffer = InAudioBuffer;
 		}
 
 		return bProcessedAnEffect;
diff --git a/Engine/Source/Runtime/AudioMixer/Public/AudioMixerSubmix.h b/Engine/Source/Runtime/AudioMixer/Public/AudioMixerSubmix.h
--- a/Engine/Source/Runtime/AudioMixer/Public/AudioMixerSubmix.h
+++ b/Engine/Source/Runtime/AudioMixer/Public/AudioMixerSubmix.h
@@ -359,7 +359,7 @@ namespace Audio
 		void SubmixCommand(TFunction<void()> Command);
 
 		// Generates audio from the given effect chain into the given buffer
-		bool GenerateEffectChainAudio(FSoundEffectSubmixInputData& InputData, AlignedFloatBuffer& InAudioBuffer, TArray<FSoundEffectSubmixPtr>& InEffectChain, const float InStartGain, const float InEndGain, AlignedFloatBuffer& OutBuffer);
+		bool GenerateEffectChainAudio(FSoundEffectSubmixInputData& InputData, const FAlignedFloatBuffer& InAudioBuffer, TArray<FSoundEffectSubmixPtr>& InEffectChain, FAlignedFloatBuffer& OutBuffer);
 
 		// This mixer submix's Id
 		uint32 Id;
diff --git a/Engine/Source/Runtime/Engine/Classes/Sound/SoundEffectSubmix.h b/Engine/Source/Runtime/Engine/Classes/Sound/SoundEffectSubmix.h
--- a/Engine/Source/Runtime/Engine/Classes/Sound/SoundEffectSubmix.h
+++ b/Engine/Source/Runtime/Engine/Classes/Sound/SoundEffectSubmix.h
@@ -125,7 +125,9 @@ public:
 	virtual float GetDryLevel() const { return 0.0f; }
 
 	// Processes audio in the submix effect.
-	void ProcessAudio(FSoundEffectSubmixInputData& InData, FSoundEffectSubmixOutputData& OutData);
+	//
+	// If the audio cannot be processed, this function will return false and OutData will not be altered. 
+	bool ProcessAudio(FSoundEffectSubmixInputData& InData, FSoundEffectSubmixOutputData& OutData);
 
 	friend class USoundEffectPreset;
 
diff --git a/Engine/Source/Runtime/Engine/Private/SoundEffectSubmix.cpp b/Engine/Source/Runtime/Engine/Private/SoundEffectSubmix.cpp
--- a/Engine/Source/Runtime/Engine/Private/SoundEffectSubmix.cpp
+++ b/Engine/Source/Runtime/Engine/Private/SoundEffectSubmix.cpp
@@ -8,7 +8,7 @@ USoundEffectSubmixPreset::USoundEffectSubmixPreset(const FObjectInitializer& Obj
 {
 }
 
-void FSoundEffectSubmix::ProcessAudio(FSoundEffectSubmixInputData& InData, FSoundEffectSubmixOutputData& OutData)
+bool FSoundEffectSubmix::ProcessAudio(FSoundEffectSubmixInputData& InData, FSoundEffectSubmixOutputData& OutData)
 {
 	bIsRunning = true;
 	InData.PresetData = nullptr;
@@ -19,10 +19,8 @@ void FSoundEffectSubmix::ProcessAudio(FSoundEffectSubmixInputData& InData, FSoun
 	if (bIsActive && Preset.IsValid())
 	{
 		OnProcessAudio(InData, OutData);
+		return true;
 	}
-	else
-	{
-		// otherwise, bypass the effect and move inputs to outputs
-		OutData.AudioBuffer = MoveTemp(InData.AudioBuffer);
-	}
+
+	return false;
 }
