Fix for buzzing in effect chains

- Scratch buffer used for effects processing was using uninitialized memory.

diff --git a/Engine/Source/Runtime/AudioMixer/Private/AudioMixerSubmix.cpp b/Engine/Source/Runtime/AudioMixer/Private/AudioMixerSubmix.cpp
--- a/Engine/Source/Runtime/AudioMixer/Private/AudioMixerSubmix.cpp
+++ b/Engine/Source/Runtime/AudioMixer/Private/AudioMixerSubmix.cpp
@@ -1330,7 +1330,7 @@ namespace Audio
 					}
 
 					// Prepare the scratch buffer for effect chain processing
-					EffectChainOutputBuffer.SetNumUninitialized(NumSamples);
+					EffectChainOutputBuffer.SetNumZeroed(NumSamples);
 
 					bProcessedAnEffect |= GenerateEffectChainAudio(InputData, InputBuffer, FadeInfo.EffectChain, EffectChainOutputBuffer);
 
