Fix for infinite feedback.

- Buffer needs zeroing each time to prevent feedback

diff --git a/Engine/Source/Runtime/AudioMixer/Private/AudioMixerSubmix.cpp b/Engine/Source/Runtime/AudioMixer/Private/AudioMixerSubmix.cpp
--- a/Engine/Source/Runtime/AudioMixer/Private/AudioMixerSubmix.cpp
+++ b/Engine/Source/Runtime/AudioMixer/Private/AudioMixerSubmix.cpp
@@ -1330,7 +1330,8 @@ namespace Audio
 					}
 
 					// Prepare the scratch buffer for effect chain processing
-					EffectChainOutputBuffer.SetNumZeroed(NumSamples);
+					EffectChainOutputBuffer.Reset();
+					EffectChainOutputBuffer.AddZeroed(NumSamples);
 
 					bProcessedAnEffect |= GenerateEffectChainAudio(InputData, InputBuffer, FadeInfo.EffectChain, EffectChainOutputBuffer);
 
