Fixing diff between 5.0 & release staging (deleting comment)

diff --git a/Engine/Source/Runtime/AudioMixer/Public/Quartz/AudioMixerClockHandle.h b/Engine/Source/Runtime/AudioMixer/Public/Quartz/AudioMixerClockHandle.h
--- a/Engine/Source/Runtime/AudioMixer/Public/Quartz/AudioMixerClockHandle.h
+++ b/Engine/Source/Runtime/AudioMixer/Public/Quartz/AudioMixerClockHandle.h
@@ -117,7 +117,7 @@ public:
 
 	UFUNCTION(BlueprintCallable, Category = "Quartz Clock", meta = (WorldContext = "WorldContextObject"))
 	void ResumeClock(const UObject* WorldContextObject, UQuartzClockHandle*& ClockHandle);
-	UFUNCTION(BlueprintCallable, Category = "Quartz Clock", meta = (/*DeprecatedFunction, DeprecationMessage="Please use ResetTransportQuantized instead",*/ WorldContext = "WorldContextObject", AutoCreateRefTerm = "InDelegate", Keywords = "Transport, Counter"))
+	UFUNCTION(BlueprintCallable, Category = "Quartz Clock", meta = (WorldContext = "WorldContextObject", AutoCreateRefTerm = "InDelegate", Keywords = "Transport, Counter"))
 	void ResetTransport(const UObject* WorldContextObject, const FOnQuartzCommandEventBP& InDelegate);
 
 	UFUNCTION(BlueprintCallable, Category = "Quartz Clock", meta = (WorldContext = "WorldContextObject", AutoCreateRefTerm = "InDelegate", Keywords = "Transport, Counter"))
