Fix for submix buffer listener using premixed audio

diff --git a/Engine/Source/Runtime/AudioMixer/Private/AudioMixerSubmix.cpp b/Engine/Source/Runtime/AudioMixer/Private/AudioMixerSubmix.cpp
--- a/Engine/Source/Runtime/AudioMixer/Private/AudioMixerSubmix.cpp
+++ b/Engine/Source/Runtime/AudioMixer/Private/AudioMixerSubmix.cpp
@@ -1134,31 +1134,6 @@ namespace Audio
 			return;
 		}
 
-		// As an optimization, early out if we're set to auto-disable and we're not rendering audio
-		if (bAutoDisable && !IsRenderingAudio())
-		{
-			if (!bIsCurrentlyDisabled)
-			{
-				bIsCurrentlyDisabled = true;
-
-				UE_CLOG(LogSubmixEnablementCVar == 1, LogAudioMixer, Display,
-					TEXT("Submix Disabled. Num Sources: %d, Time Silent: %.2f, Disablement Threshold: %.2f, Submix Name: %s"), 
-					MixerSourceVoices.Num(),
-					(float)(MixerDevice->GetAudioClock() - SilenceTimeStartSeconds),
-					AutoDisableTime,
-					*SubmixName);
-			}
-
-			// Even though we're silent, broadcast the buffer to any listeners (will be a silent buffer)
-			SendAudioToSubmixBufferListeners(OutAudioBuffer);
-			return;
-		}
-
-		if (bIsCurrentlyDisabled)
-		{
-			bIsCurrentlyDisabled = false;
-			UE_CLOG(LogSubmixEnablementCVar == 1, LogAudioMixer, Display, TEXT("Submix Re-Enabled: %s"), *SubmixName);
-		}
 
 		// If this is a Soundfield Submix, process our soundfield and decode it to a OutAudioBuffer.
 		if (IsSoundfieldSubmix())
@@ -1212,6 +1187,32 @@ namespace Audio
 
 		float* BufferPtr = InputBuffer.GetData();
 
+		// As an optimization, early out if we're set to auto-disable and we're not rendering audio
+		if (bAutoDisable && !IsRenderingAudio())
+		{
+			if (!bIsCurrentlyDisabled)
+			{
+				bIsCurrentlyDisabled = true;
+
+				UE_CLOG(LogSubmixEnablementCVar == 1, LogAudioMixer, Display,
+					TEXT("Submix Disabled. Num Sources: %d, Time Silent: %.2f, Disablement Threshold: %.2f, Submix Name: %s"), 
+					MixerSourceVoices.Num(),
+					(float)(MixerDevice->GetAudioClock() - SilenceTimeStartSeconds),
+					AutoDisableTime,
+					*SubmixName);
+			}
+
+			// Even though we're silent, broadcast the buffer to any listeners (will be a silent buffer)
+			SendAudioToSubmixBufferListeners(InputBuffer);
+			return;
+		}
+
+		if (bIsCurrentlyDisabled)
+		{
+			bIsCurrentlyDisabled = false;
+			UE_CLOG(LogSubmixEnablementCVar == 1, LogAudioMixer, Display, TEXT("Submix Re-Enabled: %s"), *SubmixName);
+		}
+
 		// Mix all submix audio into this submix's input scratch buffer
 		{
 			CSV_SCOPED_TIMING_STAT(Audio, SubmixChildren);
@@ -1484,6 +1485,8 @@ namespace Audio
 			}
 		}
 
+		SendAudioToSubmixBufferListeners(InputBuffer);
+
 		// Mix the audio buffer of this submix with the audio buffer of the output buffer (i.e. with other submixes)
 		Audio::MixInBufferFast(InputBuffer, OutAudioBuffer);
 
@@ -1515,7 +1518,6 @@ namespace Audio
 			}
 		}
 
-		SendAudioToSubmixBufferListeners(OutAudioBuffer);
 	}
 
 	void FMixerSubmix::SendAudioToSubmixBufferListeners(FAlignedFloatBuffer& OutAudioBuffer)
