Move AudioClockManager shutdown from AudioStreamShutdown to object destructor.  This will let the clock manager state persist across audio device swaps.

diff --git a/Engine/Source/Runtime/AudioMixer/Private/Quartz/AudioMixerClockManager.cpp b/Engine/Source/Runtime/AudioMixer/Private/Quartz/AudioMixerClockManager.cpp
--- a/Engine/Source/Runtime/AudioMixer/Private/Quartz/AudioMixerClockManager.cpp
+++ b/Engine/Source/Runtime/AudioMixer/Private/Quartz/AudioMixerClockManager.cpp
@@ -289,11 +289,6 @@ namespace Audio
 
 	void FQuartzClockManager::Shutdown()
 	{
-		if (MixerDevice)
-		{
-			check(MixerDevice->IsAudioRenderingThread());
-		}
-
 		FScopeLock Lock(&ActiveClockCritSec);
 		ActiveClocks.Reset();
 	}
