Decoupling Clock Handles' clock access from QuartzSubsystem / FAudioDevice

diff --git a/Engine/Source/Runtime/AudioMixer/Private/Quartz/AudioMixerClock.cpp b/Engine/Source/Runtime/AudioMixer/Private/Quartz/AudioMixerClock.cpp
--- a/Engine/Source/Runtime/AudioMixer/Private/Quartz/AudioMixerClock.cpp
+++ b/Engine/Source/Runtime/AudioMixer/Private/Quartz/AudioMixerClock.cpp
@@ -15,14 +15,101 @@ FAutoConsoleVariableRef CVarHeadlessClockSampleRate(
 
 namespace Audio
 {
+	// FQuartzClockProxy Implementation
+	// ctor
+	FQuartzClockProxy::FQuartzClockProxy(TSharedPtr<FQuartzClock, ESPMode::ThreadSafe> InClock)
+		: ClockId(InClock->GetName())
+		, SharedQueue(InClock->GetCommandQueue())
+		, ClockWeakPtr(InClock)
+	{
+	}
+
+	bool FQuartzClockProxy::IsValid() const
+	{
+		return SharedQueue.Pin().IsValid();
+	}
+
+	bool FQuartzClockProxy::DoesClockExist() const
+	{
+		return IsValid();
+	}
+
+	bool FQuartzClockProxy::IsClockRunning() const
+	{
+		TSharedPtr<FQuartzClock, ESPMode::ThreadSafe> ClockPtr = ClockWeakPtr.Pin();
+		if (!ClockPtr)
+		{
+			return false;
+		}
+
+		return ClockPtr->IsRunning();
+	}
+
+	float FQuartzClockProxy::GetDurationOfQuantizationTypeInSeconds(const EQuartzCommandQuantization& QuantizationType, float Multiplier) const
+	{
+		TSharedPtr<FQuartzClock, ESPMode::ThreadSafe> ClockPtr = ClockWeakPtr.Pin();
+		if (!ClockPtr)
+		{
+			return 0.f;
+		}
+
+		return ClockPtr->GetDurationOfQuantizationTypeInSeconds(QuantizationType, Multiplier);
+	}
+
+	Audio::FQuartzClockTickRate FQuartzClockProxy::GetTickRate() const
+	{
+		TSharedPtr<FQuartzClock, ESPMode::ThreadSafe> ClockPtr = ClockWeakPtr.Pin();
+		if (!ClockPtr)
+		{
+			return {};
+		}
+
+		return ClockPtr->GetTickRate();
+	}
+
+	FQuartzTransportTimeStamp FQuartzClockProxy::GetCurrentClockTimestamp() const
+	{
+		TSharedPtr<FQuartzClock, ESPMode::ThreadSafe> ClockPtr = ClockWeakPtr.Pin();
+		if (!ClockPtr)
+		{
+			return {};
+		}
+
+		return ClockPtr->GetCurrentTimestamp();
+	}
+
+	float FQuartzClockProxy::GetEstimatedClockRunTimeSeconds() const
+	{
+		TSharedPtr<FQuartzClock, ESPMode::ThreadSafe> ClockPtr = ClockWeakPtr.Pin();
+		if (!ClockPtr)
+		{
+			return 0.f;
+		}
+
+		return ClockPtr->GetEstimatedRunTime();
+	}
+
+
+	bool FQuartzClockProxy::SendCommandToClock(TFunction<void(FQuartzClock*)> InCommand)
+	{
+		if (auto QueuePtr = SharedQueue.Pin())
+		{
+			QueuePtr->PushCommand(InCommand);
+			return true;
+		}
+
+		return false;
+	}
+
+
+	// FQuartzClock Implementation
 	FQuartzClock::FQuartzClock(const FName& InName, const FQuartzClockSettings& InClockSettings, FQuartzClockManager* InOwningClockManagerPtr)
-		: Metronome(InClockSettings.TimeSignature)
+		: Metronome(InClockSettings.TimeSignature, InName)
 		, OwningClockManagerPtr(InOwningClockManagerPtr)
 		, Name(InName)
 		, bIsRunning(false)
 		, bIgnoresFlush(InClockSettings.bIgnoreLevelChange)
 	{
-
 		FMixerDevice* MixerDevice = GetMixerDevice();
 
 		if (MixerDevice)
@@ -165,6 +252,7 @@ namespace Audio
 	void FQuartzClock::LowResolutionTick(float InDeltaTimeSeconds)
 	{
 		TRACE_CPUPROFILER_EVENT_SCOPE(QuartzClock::Tick_LowRes);
+		PreTickCommands->PumpCommandQueue(this);
 		Tick(static_cast<int32>(InDeltaTimeSeconds * Metronome.GetTickRate().GetSampleRate()));
 	}
 
@@ -172,6 +260,10 @@ namespace Audio
 	{
 		TRACE_CPUPROFILER_EVENT_SCOPE(QuartzClock::Tick);
 
+		TRACE_CPUPROFILER_EVENT_SCOPE(QuartzClock::GameThreadCommands);
+		PreTickCommands->PumpCommandQueue(this);
+
+		TRACE_CPUPROFILER_EVENT_SCOPE(QuartzClock::TickInternal);
 		if (!bIsRunning)
 		{
 			return;
@@ -199,8 +291,19 @@ namespace Audio
 			TickInternal(InNumFramesUntilNextTick - TickDelayLengthInFrames, PendingCommands, FramesOfLatency, TickDelayLengthInFrames);
 		}
 
-
 		Metronome.Tick(InNumFramesUntilNextTick, FramesOfLatency);
+
+		UpdateCachedState();
+	}
+
+	FQuartzClockCommandQueueWeakPtr FQuartzClock::GetCommandQueue() const
+	{
+		if (!PreTickCommands.IsValid())
+		{
+			PreTickCommands = TQuartzShareableCommandQueue<FQuartzClock>::Create();
+		}
+
+		return PreTickCommands;
 	}
 
 	void FQuartzClock::TickInternal(int32 InNumFramesUntilNextTick, TArray<PendingCommand>& CommandsToTick, int32 FramesOfLatency, int32 FramesOfDelay)
@@ -242,6 +345,15 @@ namespace Audio
 		}
 	}
 
+	void FQuartzClock::UpdateCachedState()
+	{
+		FScopeLock ScopeLock(&CachedClockStateCritSec);
+
+		CachedClockState.TickRate = Metronome.GetTickRate();
+		CachedClockState.TimeStamp = Metronome.GetTimeStamp();
+		CachedClockState.RunTimeInSeconds = (float)Metronome.GetTimeSinceStart();
+	}
+
 	void FQuartzClock::SetSampleRate(float InNewSampleRate)
 	{
 		if (FMath::IsNearlyEqual(InNewSampleRate, Metronome.GetTickRate().GetSampleRate()))
@@ -255,7 +367,7 @@ namespace Audio
 		// TODO: update the deadlines of all our new events
 	}
 
-	bool FQuartzClock::IgnoresFlush()
+	bool FQuartzClock::IgnoresFlush() const
 	{
 		return bIgnoresFlush;
 	}
@@ -362,26 +474,26 @@ namespace Audio
 		return PendingCommands.Num() + ClockAlteringPendingCommands.Num();
 	}
 
-	bool FQuartzClock::IsRunning()
+	bool FQuartzClock::IsRunning() const
 	{
 		return bIsRunning;
 	}
 
 	float FQuartzClock::GetDurationOfQuantizationTypeInSeconds(const EQuartzCommandQuantization& QuantizationType, float Multiplier)
 	{
+		FScopeLock ScopeLock(&CachedClockStateCritSec);
+
 		// if this is unquantized, return 0
 		if (QuantizationType == EQuartzCommandQuantization::None)
 		{
 			return 0;
 		}
 
-		FQuartzClockTickRate TickRate = Metronome.GetTickRate();
-
 		// get number of frames until the relevant quantization event
-		double FramesUntilExec = TickRate.GetFramesPerDuration(QuantizationType);
+		double FramesUntilExec = CachedClockState.TickRate.GetFramesPerDuration(QuantizationType);
 
 		//Translate frames to seconds
-		double SampleRate = TickRate.GetSampleRate();
+		double SampleRate = CachedClockState.TickRate.GetSampleRate();
 
 		if (!FMath::IsNearlyZero(SampleRate))
 		{
@@ -395,14 +507,14 @@ namespace Audio
 
 	FQuartzTransportTimeStamp FQuartzClock::GetCurrentTimestamp()
 	{
-		FQuartzTransportTimeStamp CurrentTimeStamp = Metronome.GetTimeStamp();
-
-		return CurrentTimeStamp;
+		FScopeLock ScopeLock(&CachedClockStateCritSec);
+		return CachedClockState.TimeStamp;
 	}
 
 	float FQuartzClock::GetEstimatedRunTime()
 	{
-		return (float)Metronome.GetTimeSinceStart();
+		FScopeLock ScopeLock(&CachedClockStateCritSec);
+		return CachedClockState.RunTimeInSeconds;
 	}
 
 	FMixerDevice* FQuartzClock::GetMixerDevice()
@@ -416,6 +528,34 @@ namespace Audio
 		return nullptr;
 	}
 
+	void FQuartzClock::AddQuantizedCommand(FQuartzQuantizedCommandInitInfo& InQuantizationCommandInitInfo)
+	{
+		if (!ensure(InQuantizationCommandInitInfo.QuantizedCommandPtr))
+		{
+			return;
+		}
+
+		// this method can't be utilized by play commands because the AudioMixerSource needs a handle in order to stop it.
+		// PlayCommands must be queued via the clock manager in AudioMixerSourceManager.
+		if (!ensure(EQuartzCommandType::PlaySound != InQuantizationCommandInitInfo.QuantizedCommandPtr->GetCommandType()))
+		{
+			return;
+		}
+
+		// Can this command run without an Audio Device?
+		FMixerDevice* MixerDevice = GetMixerDevice();
+		if (!MixerDevice && InQuantizationCommandInitInfo.QuantizedCommandPtr->RequiresAudioDevice())
+		{
+			InQuantizationCommandInitInfo.QuantizedCommandPtr->Cancel();
+		}
+
+		// this function is a friend of FQuartzClockManager, so we can use FindClock() directly
+		// to access the shared ptr to "this"
+		InQuantizationCommandInitInfo.SetOwningClockPtr(GetClockManager()->FindClock(GetName()));
+		InQuantizationCommandInitInfo.QuantizedCommandPtr->OnQueued(InQuantizationCommandInitInfo);
+		AddQuantizedCommand(InQuantizationCommandInitInfo.QuantizationBoundary, InQuantizationCommandInitInfo.QuantizedCommandPtr);
+	}
+
 	FMixerSourceManager* FQuartzClock::GetSourceManager()
 	{
 		FMixerDevice* MixerDevice = GetMixerDevice();
@@ -429,6 +569,17 @@ namespace Audio
 		return nullptr;
 	}
 
+	FQuartzClockTickRate FQuartzClock::GetTickRate()
+	{
+		FScopeLock ScopeLock(&CachedClockStateCritSec);
+		return CachedClockState.TickRate;
+	}
+
+	FName FQuartzClock::GetName() const
+	{
+		return Name;
+	}
+
 	FQuartzClockManager* FQuartzClock::GetClockManager()
 	{
 		checkSlow(OwningClockManagerPtr);
@@ -444,6 +595,16 @@ namespace Audio
 		Metronome.ResetTransport();
 	}
 
+	void FQuartzClock::AddToTickDelay(int32 NumFramesOfDelayToAdd)
+	{
+		TickDelayLengthInFrames += NumFramesOfDelayToAdd;
+	}
+
+	void FQuartzClock::SetTickDelay(int32 NumFramesOfDelay)
+	{
+		TickDelayLengthInFrames = NumFramesOfDelay;
+	}
+
 	bool FQuartzClock::CancelQuantizedCommandInternal(TSharedPtr<IQuartzQuantizedCommand> InCommandPtr, TArray<PendingCommand>& CommandsToTick)
 	{
 		for (int32 i = CommandsToTick.Num() - 1; i >= 0; --i)
diff --git a/Engine/Source/Runtime/AudioMixer/Private/Quartz/AudioMixerClockHandle.cpp b/Engine/Source/Runtime/AudioMixer/Private/Quartz/AudioMixerClockHandle.cpp
--- a/Engine/Source/Runtime/AudioMixer/Private/Quartz/AudioMixerClockHandle.cpp
+++ b/Engine/Source/Runtime/AudioMixer/Private/Quartz/AudioMixerClockHandle.cpp
@@ -6,119 +6,6 @@
 #include "AudioMixerDevice.h"
 #include "Engine/GameInstance.h"
 
-// FQuartzTickableObject implementation
-void FQuartzTickableObject::FQuartzTickableObjectGCObjectMembers::AddReferencedObjects(FReferenceCollector& Collector)
-{
-	Collector.AddReferencedObject(QuartzSubsystem);
-	Collector.AddReferencedObject(WorldPtr);
-}
-
-FQuartzTickableObject::~FQuartzTickableObject()
-{
-	Shutdown();
-
-	// shutdown the shared command queue
-	if (CommandQueuePtr.IsValid())
-	{
-		CommandQueuePtr->StopTakingCommands();
-		CommandQueuePtr.Reset();
-	}
-}
-
-FQuartzTickableObject* FQuartzTickableObject::Init(UWorld* InWorldPtr)
-{
-	if (!InWorldPtr)
-	{
-		// can't initialize if we don't have a valid world
-		return this;
-	}
-
-	GCObjectMembers.WorldPtr = InWorldPtr;
-	GCObjectMembers.QuartzSubsystem = UQuartzSubsystem::Get(GCObjectMembers.WorldPtr);
-	CommandQueuePtr = GCObjectMembers.QuartzSubsystem->CreateQuartzCommandQueue();
-	GCObjectMembers.QuartzSubsystem->SubscribeToQuartzTick(this);
-
-	bHasBeenInitialized = true;
-
-	return this;
-}
-
-int32 FQuartzTickableObject::AddCommandDelegate(const FOnQuartzCommandEventBP& InDelegate, TArray<TSharedPtr<Audio::FShareableQuartzCommandQueue, ESPMode::ThreadSafe>>& TargetSubscriberArray)
-{
-	TargetSubscriberArray.Add(CommandQueuePtr);
-
-	const int32 Num = QuantizedCommandDelegates.Num();
-	int32 SlotId = 0;
-
-	for (; SlotId < Num; ++SlotId)
-	{
-		if (!QuantizedCommandDelegates[SlotId].MulticastDelegate.IsBound())
-		{
-			QuantizedCommandDelegates[SlotId].MulticastDelegate.AddUnique(InDelegate);
-			return SlotId;
-		}
-	}
-
-	// need a new slot
-	QuantizedCommandDelegates.AddDefaulted_GetRef().MulticastDelegate.AddUnique(InDelegate);
-	return SlotId;
-}
-
-// called by UQuartzSubsystem::BeginDestroy()
-void FQuartzTickableObject::Shutdown()
-{
-	if (GCObjectMembers.WorldPtr && GCObjectMembers.QuartzSubsystem)
-	{
-		GCObjectMembers.QuartzSubsystem->UnsubscribeFromQuartzTick(this);
-		GCObjectMembers.WorldPtr = nullptr;
-		GCObjectMembers.QuartzSubsystem = nullptr;
-		bHasBeenInitialized = false;
-	}
-}
-
-TSharedPtr<Audio::FShareableQuartzCommandQueue, ESPMode::ThreadSafe> FQuartzTickableObject::GetCommandQueue()
-{
-	if (!CommandQueuePtr.IsValid())
-	{
-		CommandQueuePtr = MakeShared<Audio::FShareableQuartzCommandQueue, ESPMode::ThreadSafe>();
-	}
-
-	return CommandQueuePtr;
-}
-
-void FQuartzTickableObject::QuartzTick(float DeltaTime)
-{
-	PumpCommandQueue();
-}
-
-bool FQuartzTickableObject::QuartzIsTickable() const
-{
-	return (CommandQueuePtr && !CommandQueuePtr->IsQueueEmpty());
-}
-
-void FQuartzTickableObject::PumpCommandQueue()
-{
-	if (!CommandQueuePtr.IsValid())
-	{
-		return;
-	}
-
-	// gather move all the current commands into our temp container
-	TFunction<void(FQuartzTickableObject*)> Function;
-	while (CommandQueuePtr->EventDelegateQueue.Dequeue(Function))
-	{
-		TempCommandQueue.Emplace(Function);
-	}
-
-	// execute the commands (may alter the original container)
-	for (auto& Command : TempCommandQueue)
-	{
-		Command(this);
-	}
-
-	TempCommandQueue.Reset();
-}
-
 
 
 // Clock Handle implementation
@@ -128,32 +15,14 @@ UQuartzClockHandle::UQuartzClockHandle()
 
 UQuartzClockHandle::~UQuartzClockHandle()
 {
-
 }
 
 void UQuartzClockHandle::BeginDestroy()
 {
 	Super::BeginDestroy();
 
-	// un-subscribe from Subsystem tick and metronome events
-	if (UQuartzSubsystem* QuartzSubsystem = GetQuartzSubsystem())
-	{
-		QuartzSubsystem->UnsubscribeFromQuartzTick(this);
-
-		Audio::FQuartzClockManager* ClockManager = QuartzSubsystem->GetManagerForClock(QuartzGetWorld(), GetClockName());
-
-		if (ClockManager)
-		{
-			ClockManager->UnsubscribeFromAllTimeDivisions(CurrentClockId, GetCommandQueue());
-		}
-	}
-
-	// shutdown the shared command queue
-	if (CommandQueuePtr.IsValid())
-	{
-		CommandQueuePtr->StopTakingCommands();
-		CommandQueuePtr.Reset();
-	}
+	auto Queue = GetCommandQueue().Pin();
+	RawHandle.SendCommandToClock([Queue](Audio::FQuartzClock* InClock) { InClock->UnsubscribeFromAllTimeDivisions(Queue); });
 }
 
 void UQuartzClockHandle::StartClock(const UObject* WorldContextObject, UQuartzClockHandle*& ClockHandle)
@@ -162,46 +31,23 @@ void UQuartzClockHandle::StartClock(const UObject* WorldContextObject, UQuartzCl
 	ResumeClock(WorldContextObject, ClockHandle);
 }
 
-void UQuartzClockHandle::StopClock(const UObject* WorldContextObject, bool CancelPendingEvents, UQuartzClockHandle*& ClockHandle)
+void UQuartzClockHandle::StopClock(const UObject* WorldContextObject, bool bCancelPendingEvents, UQuartzClockHandle*& ClockHandle)
 {
 	ClockHandle = this;
-	if (UQuartzSubsystem* QuartzSubsystem = GetQuartzSubsystem())
-	{
-		Audio::FQuartzClockManager* ClockManager = QuartzSubsystem->GetManagerForClock(WorldContextObject, GetClockName());
-
-		if (ClockManager)
-		{
-			ClockManager->StopClock(CurrentClockId, CancelPendingEvents);
-		}
-	}
+	RawHandle.SendCommandToClock([bCancelPendingEvents](Audio::FQuartzClock* InClock) { InClock->Stop(bCancelPendingEvents); });
 }
 
 void UQuartzClockHandle::PauseClock(const UObject* WorldContextObject, UQuartzClockHandle*& ClockHandle)
 {
 	ClockHandle = this;
-	if (UQuartzSubsystem* QuartzSubsystem = GetQuartzSubsystem())
-	{
-		Audio::FQuartzClockManager* ClockManager = QuartzSubsystem->GetManagerForClock(WorldContextObject, GetClockName());
-		if (ClockManager)
-		{
-			ClockManager->PauseClock(CurrentClockId);
-		}
-	}
+	RawHandle.SendCommandToClock([](Audio::FQuartzClock* InClock) { InClock->Pause(); });
 }
 
 // Begin BP interface
 void UQuartzClockHandle::ResumeClock(const UObject* WorldContextObject, UQuartzClockHandle*& ClockHandle)
 {
 	ClockHandle = this;
-	if (UQuartzSubsystem* QuartzSubsystem = GetQuartzSubsystem())
-	{
-		Audio::FQuartzClockManager* ClockManager = QuartzSubsystem->GetManagerForClock(WorldContextObject, GetClockName());
-
-		if (ClockManager)
-		{
-			ClockManager->ResumeClock(CurrentClockId);
-		}
-	}
+	RawHandle.SendCommandToClock([](Audio::FQuartzClock* InClock) { InClock->Resume(); });
 }
 
 void UQuartzClockHandle::QueueQuantizedSound(const UObject* WorldContextObject, UQuartzClockHandle*& InClockHandle, const FAudioComponentCommandInfo& InAudioComponentData, const FOnQuartzCommandEventBP& InDelegate, const FQuartzQuantizationBoundary& InTargetBoundary)
@@ -209,89 +55,54 @@ void UQuartzClockHandle::QueueQuantizedSound(const UObject* WorldContextObject,
 	InClockHandle = this;
 	FName ClockName = GetClockName();
 
-	if (UQuartzSubsystem* QuartzSubsystem = GetQuartzSubsystem())
-	{
-		//Create a Queue Command, and give it the additional data that it needs
-		TSharedPtr<Audio::FQuantizedQueueCommand> QueueCommandPtr = MakeShared<Audio::FQuantizedQueueCommand>();
-		QueueCommandPtr->SetQueueCommand(InAudioComponentData);
+	//Create a Queue Command, and give it the additional data that it needs
+	TSharedPtr<Audio::FQuantizedQueueCommand> QueueCommandPtr = MakeShared<Audio::FQuantizedQueueCommand>();
+	QueueCommandPtr->SetQueueCommand(InAudioComponentData);
 
-		//Set up initial command info
-		Audio::FQuartzQuantizedCommandInitInfo CommandInitInfo = QuartzSubsystem->CreateDataDataForSchedulePlaySound(InClockHandle, InDelegate, InTargetBoundary);
+	//Set up initial command info
+	Audio::FQuartzQuantizedCommandInitInfo CommandInitInfo = UQuartzSubsystem::CreateRequestDataForSchedulePlaySound(InClockHandle, InDelegate, InTargetBoundary);
 
-		//(Queue's setup is identical to PlaySound except for the command ptr, so fix that here)
-		CommandInitInfo.QuantizedCommandPtr.Reset();
-		CommandInitInfo.QuantizedCommandPtr = QueueCommandPtr;
+	//(Queue's setup is identical to PlaySound except for the command ptr, so fix that here)
+	CommandInitInfo.QuantizedCommandPtr.Reset();
+	CommandInitInfo.QuantizedCommandPtr = QueueCommandPtr;
 
-		QuartzSubsystem->AddCommandToClock(WorldContextObject, CommandInitInfo, ClockName);	
-	}
+	RawHandle.SendCommandToClock([CommandInitInfo](Audio::FQuartzClock* InClock) mutable { InClock->AddQuantizedCommand(CommandInitInfo); });
 }
 
 // deprecated: use ResetTransportQuantized
 void UQuartzClockHandle::ResetTransport(const UObject* WorldContextObject, const FOnQuartzCommandEventBP& InDelegate)
 {
-	if (UQuartzSubsystem* QuartzSubsystem = GetQuartzSubsystem())
-	{
-		Audio::FQuartzQuantizedCommandInitInfo Data(QuartzSubsystem->CreateDataForTransportReset(this, FQuartzQuantizationBoundary(EQuartzCommandQuantization::Bar), InDelegate));
-		QuartzSubsystem->AddCommandToClock(WorldContextObject, Data, GetClockName());
-	}
+	Audio::FQuartzQuantizedCommandInitInfo Data(UQuartzSubsystem::CreateRequestDataForTransportReset(this, FQuartzQuantizationBoundary(EQuartzCommandQuantization::Bar), InDelegate));
+	RawHandle.SendCommandToClock([Data](Audio::FQuartzClock* InClock) mutable { InClock->AddQuantizedCommand(Data); });
 }
 
 void UQuartzClockHandle::ResetTransportQuantized(const UObject* WorldContextObject, FQuartzQuantizationBoundary InQuantizationBoundary, const FOnQuartzCommandEventBP& InDelegate, UQuartzClockHandle*& ClockHandle)
 {
 	ClockHandle = this;
-	if (UQuartzSubsystem* QuartzSubsystem = GetQuartzSubsystem())
-	{
-		Audio::FQuartzQuantizedCommandInitInfo Data(QuartzSubsystem->CreateDataForTransportReset(this, InQuantizationBoundary, InDelegate));
-		QuartzSubsystem->AddCommandToClock(WorldContextObject, Data, GetClockName());
-	}
+	Audio::FQuartzQuantizedCommandInitInfo Data(UQuartzSubsystem::CreateRequestDataForTransportReset(this, InQuantizationBoundary, InDelegate));
+	RawHandle.SendCommandToClock([Data](Audio::FQuartzClock* InClock) mutable { InClock->AddQuantizedCommand(Data); });
 }
 
 
 
 bool UQuartzClockHandle::IsClockRunning(const UObject* WorldContextObject)
 {
-	if (UQuartzSubsystem* QuartzSubsystem = GetQuartzSubsystem())
-	{
-		return QuartzSubsystem->IsClockRunning(WorldContextObject, CurrentClockId);
-	}
-
-	return false;
+	return RawHandle.IsClockRunning();
 }
 
 float UQuartzClockHandle::GetDurationOfQuantizationTypeInSeconds(const UObject* WorldContextObject, const EQuartzCommandQuantization& QuantizationType, float Multiplier)
 {
-	if (UQuartzSubsystem* QuartzSubsystem = GetQuartzSubsystem())
-	{
-		return QuartzSubsystem->GetDurationOfQuantizationTypeInSeconds(WorldContextObject, CurrentClockId, QuantizationType, Multiplier);
-	}
-	else
-	{
-		return INDEX_NONE;
-	}
+	return RawHandle.GetDurationOfQuantizationTypeInSeconds(QuantizationType, Multiplier);
 }
 
 FQuartzTransportTimeStamp UQuartzClockHandle::GetCurrentTimestamp(const UObject* WorldContextObject)
 {
-	if (UQuartzSubsystem* QuartzSubsystem = GetQuartzSubsystem())
-	{
-		return QuartzSubsystem->GetCurrentClockTimestamp(WorldContextObject, CurrentClockId);
-	}
-	else
-	{
-		return FQuartzTransportTimeStamp();
-	}
+	return RawHandle.GetCurrentClockTimestamp();
 }
 
 float UQuartzClockHandle::GetEstimatedRunTime(const UObject* WorldContextObject)
 {
-	if (UQuartzSubsystem* QuartzSubsystem = GetQuartzSubsystem())
-	{
-		return QuartzSubsystem->GetEstimatedClockRunTime(WorldContextObject, CurrentClockId);
-	}
-	else
-	{
-		return INDEX_NONE;
-	}
+	return RawHandle.GetEstimatedClockRunTimeSeconds();
 }
 
 void UQuartzClockHandle::StartOtherClock(const UObject* WorldContextObject, FName OtherClockName, FQuartzQuantizationBoundary InQuantizationBoundary, const FOnQuartzCommandEventBP& InDelegate)
@@ -302,13 +113,11 @@ void UQuartzClockHandle::StartOtherClock(const UObject* WorldContextObject, FNam
 		return;
 	}
 
-	if (UQuartzSubsystem* QuartzSubsystem = GetQuartzSubsystem())
-	{
-		Audio::FQuartzQuantizedCommandInitInfo Data(QuartzSubsystem->CreateDataForStartOtherClock(this, OtherClockName, InQuantizationBoundary, InDelegate));
-		QuartzSubsystem->AddCommandToClock(WorldContextObject, Data, GetClockName());
-	}
+	Audio::FQuartzQuantizedCommandInitInfo Data(UQuartzSubsystem::CreateRequestDataForStartOtherClock(this, OtherClockName, InQuantizationBoundary, InDelegate));
+	RawHandle.SendCommandToClock([Data](Audio::FQuartzClock* InClock) mutable { InClock->AddQuantizedCommand(Data); });
 }
 
+// todo: Move the bulk of these functions to FQuartzTcikableObject once lightweight clock handles are spun up.
 void UQuartzClockHandle::SubscribeToQuantizationEvent(const UObject* WorldContextObject, EQuartzCommandQuantization InQuantizationBoundary, const FOnQuartzMetronomeEventBP& OnQuantizationEvent, UQuartzClockHandle*& ClockHandle)
 {
 	ClockHandle = this;
@@ -319,59 +128,39 @@ void UQuartzClockHandle::SubscribeToQuantizationEvent(const UObject* WorldContex
 		return;
 	}
 
-	if (UQuartzSubsystem* QuartzSubsystem = GetQuartzSubsystem())
-	{
-		Audio::FQuartzClockManager* ClockManager = QuartzSubsystem->GetManagerForClock(WorldContextObject, GetClockName());
-		if (ClockManager && ClockManager->DoesClockExist(CurrentClockId) && OnQuantizationEvent.IsBound())
-		{
-			MetronomeDelegates[static_cast<int32>(InQuantizationBoundary)].MulticastDelegate.AddUnique(OnQuantizationEvent);
-			ClockManager->SubscribeToTimeDivision(CurrentClockId, GetCommandQueue(), InQuantizationBoundary);
-		}
-	}
+	AddMetronomeBpDelegate(InQuantizationBoundary, OnQuantizationEvent);
+
+	auto Queue = GetCommandQueue().Pin();
+	RawHandle.SendCommandToClock([Queue, InQuantizationBoundary](Audio::FQuartzClock* InClock) { InClock->SubscribeToTimeDivision(Queue, InQuantizationBoundary); });
 }
 
 void UQuartzClockHandle::SubscribeToAllQuantizationEvents(const UObject* WorldContextObject, const FOnQuartzMetronomeEventBP& OnQuantizationEvent, UQuartzClockHandle*& ClockHandle)
 {
 	ClockHandle = this;
-	if (UQuartzSubsystem* QuartzSubsystem = GetQuartzSubsystem())
-	{
-		Audio::FQuartzClockManager* ClockManager = QuartzSubsystem->GetManagerForClock(WorldContextObject, GetClockName());
-		if (ClockManager && ClockManager->DoesClockExist(CurrentClockId) && OnQuantizationEvent.IsBound())
-		{
-			for (int32 i = 0; i < static_cast<int32>(EQuartzCommandQuantization::Count) - 1; ++i)
-			{
-				MetronomeDelegates[i].MulticastDelegate.AddUnique(OnQuantizationEvent);
-			}
 
-			ClockManager->SubscribeToAllTimeDivisions(CurrentClockId, GetCommandQueue());
-		}
+	for (int32 i = 0; i < static_cast<int32>(EQuartzCommandQuantization::Count) - 1; ++i)
+	{
+		AddMetronomeBpDelegate(static_cast<EQuartzCommandQuantization>(i), OnQuantizationEvent);
 	}
+
+	auto Queue = GetCommandQueue().Pin();
+	RawHandle.SendCommandToClock([Queue](Audio::FQuartzClock* InClock) { InClock->SubscribeToAllTimeDivisions(Queue); });
 }
 
 void UQuartzClockHandle::UnsubscribeFromTimeDivision(const UObject* WorldContextObject, EQuartzCommandQuantization InQuantizationBoundary, UQuartzClockHandle*& ClockHandle)
 {
 	ClockHandle = this;
-	if (UQuartzSubsystem* QuartzSubsystem = GetQuartzSubsystem())
-	{
-		Audio::FQuartzClockManager* ClockManager = QuartzSubsystem->GetManagerForClock(WorldContextObject, GetClockName());
-		if (ClockManager && ClockManager->DoesClockExist(CurrentClockId))
-		{
-			ClockManager->UnsubscribeFromTimeDivision(CurrentClockId, GetCommandQueue(), InQuantizationBoundary);
-		}
-	}
+
+	auto Queue = GetCommandQueue().Pin();
+	RawHandle.SendCommandToClock([Queue, InQuantizationBoundary](Audio::FQuartzClock* InClock) { InClock->UnsubscribeFromTimeDivision(Queue, InQuantizationBoundary); });
 }
 
 void UQuartzClockHandle::UnsubscribeFromAllTimeDivisions(const UObject* WorldContextObject, UQuartzClockHandle*& ClockHandle)
 {
 	ClockHandle = this;
-	if (UQuartzSubsystem* QuartzSubsystem = GetQuartzSubsystem())
-	{
-		Audio::FQuartzClockManager* ClockManager = QuartzSubsystem->GetManagerForClock(WorldContextObject, GetClockName());
-		if (ClockManager && ClockManager->DoesClockExist(CurrentClockId))
-		{
-			ClockManager->UnsubscribeFromAllTimeDivisions(CurrentClockId, GetCommandQueue());
-		}
-	}
+
+	auto Queue = GetCommandQueue().Pin();
+	RawHandle.SendCommandToClock([Queue](Audio::FQuartzClock* InClock) { InClock->UnsubscribeFromAllTimeDivisions(Queue); });
 }
 
 // Metronome Alteration (setters)
@@ -384,14 +173,9 @@ void UQuartzClockHandle::SetMillisecondsPerTick(const UObject* WorldContextObjec
 		return;
 	}
 
-	if (UQuartzSubsystem* QuartzSubsystem = GetQuartzSubsystem())
-	{
-		Audio::FQuartzClockTickRate TickRate;
-		TickRate.SetMillisecondsPerTick(MillisecondsPerTick);
-
-		Audio::FQuartzQuantizedCommandInitInfo Data(QuartzSubsystem->CreateDataForTickRateChange(this, InDelegate, TickRate, InQuantizationBoundary));
-		QuartzSubsystem->AddCommandToClock(WorldContextObject, Data, GetClockName());
-	}
+	Audio::FQuartzClockTickRate TickRate;
+	TickRate.SetMillisecondsPerTick(MillisecondsPerTick);
+	SetTickRateInternal(InQuantizationBoundary, InDelegate, TickRate);
 }
 
 void UQuartzClockHandle::SetTicksPerSecond(const UObject* WorldContextObject, const FQuartzQuantizationBoundary& InQuantizationBoundary, const FOnQuartzCommandEventBP& InDelegate, UQuartzClockHandle*& ClockHandle, float TicksPerSecond)
@@ -403,14 +187,9 @@ void UQuartzClockHandle::SetTicksPerSecond(const UObject* WorldContextObject, co
 		return;
 	}
 
-	if (UQuartzSubsystem* QuartzSubsystem = GetQuartzSubsystem())
-	{
-		Audio::FQuartzClockTickRate TickRate;
-		TickRate.SetSecondsPerTick(1.f / TicksPerSecond);
-
-		Audio::FQuartzQuantizedCommandInitInfo Data(QuartzSubsystem->CreateDataForTickRateChange(this, InDelegate, TickRate, InQuantizationBoundary));
-		QuartzSubsystem->AddCommandToClock(WorldContextObject, Data, GetClockName());
-	}
+	Audio::FQuartzClockTickRate TickRate;
+	TickRate.SetSecondsPerTick(1.f / TicksPerSecond);
+	SetTickRateInternal(InQuantizationBoundary, InDelegate, TickRate);
 }
 
 void UQuartzClockHandle::SetSecondsPerTick(const UObject* WorldContextObject, const FQuartzQuantizationBoundary& InQuantizationBoundary, const FOnQuartzCommandEventBP& InDelegate, UQuartzClockHandle*& ClockHandle, float SecondsPerTick)
@@ -422,14 +201,9 @@ void UQuartzClockHandle::SetSecondsPerTick(const UObject* WorldContextObject, co
 		return;
 	}
 
-	if (UQuartzSubsystem* QuartzSubsystem = GetQuartzSubsystem())
-	{
-		Audio::FQuartzClockTickRate TickRate;
-		TickRate.SetSecondsPerTick(SecondsPerTick);
-
-		Audio::FQuartzQuantizedCommandInitInfo Data(QuartzSubsystem->CreateDataForTickRateChange(this, InDelegate, TickRate, InQuantizationBoundary));
-		QuartzSubsystem->AddCommandToClock(WorldContextObject, Data, GetClockName());
-	}
+	Audio::FQuartzClockTickRate TickRate;
+	TickRate.SetSecondsPerTick(SecondsPerTick);
+	SetTickRateInternal(InQuantizationBoundary, InDelegate, TickRate);
 }
 
 void UQuartzClockHandle::SetThirtySecondNotesPerMinute(const UObject* WorldContextObject, const FQuartzQuantizationBoundary& InQuantizationBoundary, const FOnQuartzCommandEventBP& InDelegate, UQuartzClockHandle*& ClockHandle, float ThirtySecondsNotesPerMinute)
@@ -441,14 +215,9 @@ void UQuartzClockHandle::SetThirtySecondNotesPerMinute(const UObject* WorldConte
 		return;
 	}
 
-	if (UQuartzSubsystem* QuartzSubsystem = GetQuartzSubsystem())
-	{
-		Audio::FQuartzClockTickRate TickRate;
-		TickRate.SetThirtySecondNotesPerMinute(ThirtySecondsNotesPerMinute);
-
-		Audio::FQuartzQuantizedCommandInitInfo Data(QuartzSubsystem->CreateDataForTickRateChange(this, InDelegate, TickRate, InQuantizationBoundary));
-		QuartzSubsystem->AddCommandToClock(WorldContextObject, Data, GetClockName());
-	}
+	Audio::FQuartzClockTickRate TickRate;
+	TickRate.SetThirtySecondNotesPerMinute(ThirtySecondsNotesPerMinute);
+	SetTickRateInternal(InQuantizationBoundary, InDelegate, TickRate);
 }
 
 void UQuartzClockHandle::SetBeatsPerMinute(const UObject* WorldContextObject, const FQuartzQuantizationBoundary& InQuantizationBoundary, const FOnQuartzCommandEventBP& InDelegate, UQuartzClockHandle*& ClockHandle, float BeatsPerMinute)
@@ -460,14 +229,15 @@ void UQuartzClockHandle::SetBeatsPerMinute(const UObject* WorldContextObject, co
 		return;
 	}
 
-	if (UQuartzSubsystem* QuartzSubsystem = GetQuartzSubsystem())
-	{
-		Audio::FQuartzClockTickRate TickRate;
-		TickRate.SetBeatsPerMinute(BeatsPerMinute);
+	Audio::FQuartzClockTickRate TickRate;
+	TickRate.SetBeatsPerMinute(BeatsPerMinute);
+	SetTickRateInternal(InQuantizationBoundary, InDelegate, TickRate);
+}
 
-		Audio::FQuartzQuantizedCommandInitInfo Data(QuartzSubsystem->CreateDataForTickRateChange(this, InDelegate, TickRate, InQuantizationBoundary));
-		QuartzSubsystem->AddCommandToClock(WorldContextObject, Data, GetClockName());
-	}
+void UQuartzClockHandle::SetTickRateInternal(const FQuartzQuantizationBoundary& InQuantizationBoundary, const FOnQuartzCommandEventBP& InDelegate, const Audio::FQuartzClockTickRate& NewTickRate)
+{
+	Audio::FQuartzQuantizedCommandInitInfo Data(UQuartzSubsystem::CreateRequestDataForTickRateChange(this, InDelegate, NewTickRate, InQuantizationBoundary));
+	RawHandle.SendCommandToClock([Data](Audio::FQuartzClock* InClock) mutable { InClock->AddQuantizedCommand(Data); });
 }
 
 // Metronome getters
@@ -490,7 +260,7 @@ float UQuartzClockHandle::GetTicksPerSecond(const UObject* WorldContextObject) c
 	if (GetCurrentTickRate(WorldContextObject, OutTickRate))
 	{
 		const float SecondsPerTick = OutTickRate.GetSecondsPerTick();
-		
+
 		if (!FMath::IsNearlyZero(SecondsPerTick))
 		{
 			return 1.f / SecondsPerTick;
@@ -538,7 +308,7 @@ float UQuartzClockHandle::GetBeatsPerMinute(const UObject* WorldContextObject) c
 // End BP interface
 
 
-UQuartzClockHandle* UQuartzClockHandle::SubscribeToClock(const UObject* WorldContextObject, FName ClockName)
+UQuartzClockHandle* UQuartzClockHandle::SubscribeToClock(const UObject* WorldContextObject, FName ClockName, Audio::FQuartzClockProxy* InHandlePtr)
 {
 	// create ID
 	CurrentClockId = ClockName;
@@ -547,9 +317,13 @@ UQuartzClockHandle* UQuartzClockHandle::SubscribeToClock(const UObject* WorldCon
 	TempId.Append(CurrentClockId.ToString());
 	ClockHandleId = FName(*TempId);
 
-	// TODO: subscribe to clock w/ ClockHandleId
 	bConnectedToClock = true;
 
+	if (InHandlePtr)
+	{
+		RawHandle = *InHandlePtr;
+	}
+
 	return this;
 }
 
@@ -557,57 +331,12 @@ UQuartzClockHandle* UQuartzClockHandle::SubscribeToClock(const UObject* WorldCon
 // returns true if OutTickRate is valid and was updated
 bool UQuartzClockHandle::GetCurrentTickRate(const UObject* WorldContextObject, Audio::FQuartzClockTickRate& OutTickRate) const
 {
-	if (UQuartzSubsystem* QuartzSubsystem = GetQuartzSubsystem())
+	if (RawHandle.IsValid())
 	{
-		Audio::FQuartzClockManager* ClockManager = QuartzSubsystem->GetManagerForClock(WorldContextObject, GetClockName());
-
-		if (ClockManager)
-		{
-			OutTickRate = ClockManager->GetTickRateForClock(CurrentClockId);
-			return true;
-		}
+		OutTickRate = RawHandle.GetTickRate();
+		return true;
 	}
 
 	OutTickRate = {};
 	return false;
 }
-
-void UQuartzClockHandle::ProcessCommand(const Audio::FQuartzQuantizedCommandDelegateData& Data)
-{
-	checkSlow(Data.DelegateSubType < EQuartzCommandDelegateSubType::Count && (Data.DelegateID < QuantizedCommandDelegates.Num()));
-	if (UQuartzSubsystem* QuartzSubsystem = GetQuartzSubsystem())
-	{
-		QuartzSubsystem->PushLatencyTrackerResult(Data.RequestRecieved());
-	}
-
-	CommandDelegateGameThreadData& GameThreadEntry = QuantizedCommandDelegates[Data.DelegateID];
-
-	GameThreadEntry.MulticastDelegate.Broadcast(Data.DelegateSubType, "Sample Payload");
-
-	// track the number of active QuantizedCommands that may be sending info back to us.
-	// (new command)
-	if (Data.DelegateSubType == EQuartzCommandDelegateSubType::CommandOnQueued)
-	{
-		GameThreadEntry.RefCount.Increment();
-	}
-
-	// (end of a command)
-	if (Data.DelegateSubType == EQuartzCommandDelegateSubType::CommandOnCanceled)
-	{
-		// are all the commands done?
-		if (GameThreadEntry.RefCount.Decrement() == 0)
-		{
-			GameThreadEntry.MulticastDelegate.Clear();
-		}
-	}
-}
-
-void UQuartzClockHandle::ProcessCommand(const Audio::FQuartzMetronomeDelegateData& Data)
-{
-	if (UQuartzSubsystem* QuartzSubsystem = GetQuartzSubsystem())
-	{
-		QuartzSubsystem->PushLatencyTrackerResult(Data.RequestRecieved());
-	}
-
-	MetronomeDelegates[static_cast<int32>(Data.Quantization)].MulticastDelegate.Broadcast(CurrentClockId, Data.Quantization, Data.Bar, Data.Beat, Data.BeatFraction);
-}
\ No newline at end of file
diff --git a/Engine/Source/Runtime/AudioMixer/Private/Quartz/AudioMixerClockManager.cpp b/Engine/Source/Runtime/AudioMixer/Private/Quartz/AudioMixerClockManager.cpp
--- a/Engine/Source/Runtime/AudioMixer/Private/Quartz/AudioMixerClockManager.cpp
+++ b/Engine/Source/Runtime/AudioMixer/Private/Quartz/AudioMixerClockManager.cpp
@@ -52,7 +52,7 @@ namespace Audio
 		}
 	}
 
-	TSharedPtr<FQuartzClock> FQuartzClockManager::GetOrCreateClock(const FName& InClockName, const FQuartzClockSettings& InClockSettings, bool bOverrideTickRateIfClockExists)
+	FQuartzClockProxy FQuartzClockManager::GetOrCreateClock(const FName& InClockName, const FQuartzClockSettings& InClockSettings, bool bOverrideTickRateIfClockExists)
 	{
 		FScopeLock Lock(&ActiveClockCritSec);
 
@@ -70,11 +70,25 @@ namespace Audio
 				Clock->ChangeTimeSignature(NewSettings.TimeSignature);
 			}
 
-			return Clock;
+			return FQuartzClockProxy(Clock);
 		}
 
 		// doesn't exist, create new clock
-		return ActiveClocks.Emplace_GetRef(MakeShared<FQuartzClock>(InClockName, NewSettings, this));
+		return FQuartzClockProxy(ActiveClocks.Emplace_GetRef(MakeShared<FQuartzClock>(InClockName, NewSettings, this)));
+	}
+
+	FQuartzClockProxy FQuartzClockManager::GetClock(const FName& InClockName)
+	{
+		FScopeLock Lock(&ActiveClockCritSec);
+		TSharedPtr<FQuartzClock> ClockPtr = FindClock(InClockName);
+
+		if (ClockPtr)
+		{
+			return FQuartzClockProxy(ClockPtr);
+		}
+
+		UE_LOG(LogAudioQuartz, Warning, TEXT("Could not find Clock: %s (returning empty handle)"), *ClockPtr->GetName().ToString());
+		return {};
 	}
 
 	bool FQuartzClockManager::DoesClockExist(const FName& InClockName)
diff --git a/Engine/Source/Runtime/AudioMixer/Private/Quartz/QuartzMetronome.cpp b/Engine/Source/Runtime/AudioMixer/Private/Quartz/QuartzMetronome.cpp
--- a/Engine/Source/Runtime/AudioMixer/Private/Quartz/QuartzMetronome.cpp
+++ b/Engine/Source/Runtime/AudioMixer/Private/Quartz/QuartzMetronome.cpp
@@ -5,13 +5,14 @@
 
 namespace Audio
 {
-	FQuartzMetronome::FQuartzMetronome() : TimeSinceStart(0)
+	FQuartzMetronome::FQuartzMetronome(FName InClockName)
+		: TimeSinceStart(0), ClockName(InClockName)
 	{
 		SetTickRate(CurrentTickRate);
 	}
 
-	FQuartzMetronome::FQuartzMetronome(const FQuartzTimeSignature& InTimeSignature)
-		: CurrentTimeSignature(InTimeSignature), TimeSinceStart(0)
+	FQuartzMetronome::FQuartzMetronome(const FQuartzTimeSignature& InTimeSignature, FName InClockName)
+		: CurrentTimeSignature(InTimeSignature), TimeSinceStart(0), ClockName(InClockName)
 	{
 		SetTickRate(CurrentTickRate);
 	}
@@ -417,17 +418,18 @@ namespace Audio
 
 	void FQuartzMetronome::FireEvents(int32 EventFlags)
 	{
-		FQuartzMetronomeDelegateData Data;
-		Data.Bar = (CurrentTimeStamp.Bars);
-		Data.Beat = (CurrentTimeStamp.Beat);
-		Data.BeatFraction = (CurrentTimeStamp.BeatFraction);
-
 		if (!(EventFlags &= ListenerFlags))
 		{
 			// no events occurred that we have listeners for
 			return;
 		}
 
+		FQuartzMetronomeDelegateData Data;
+		Data.Bar = (CurrentTimeStamp.Bars);
+		Data.Beat = (CurrentTimeStamp.Beat);
+		Data.BeatFraction = (CurrentTimeStamp.BeatFraction);
+		Data.ClockName = ClockName;
+
 		// loop through quantization boundaries
 		int32 i = -1;
 		for (TArray<MetronomeCommandQueuePtr>& QuantizationBoundarySubscribers : MetronomeSubscriptionMatrix)
diff --git a/Engine/Source/Runtime/AudioMixer/Private/Quartz/QuartzSubsystem.cpp b/Engine/Source/Runtime/AudioMixer/Private/Quartz/QuartzSubsystem.cpp
--- a/Engine/Source/Runtime/AudioMixer/Private/Quartz/QuartzSubsystem.cpp
+++ b/Engine/Source/Runtime/AudioMixer/Private/Quartz/QuartzSubsystem.cpp
@@ -5,6 +5,7 @@
 #include "Quartz/QuartzMetronome.h"
 #include "Quartz/AudioMixerClockManager.h"
 #include "Sound/QuartzQuantizationUtilities.h"
+#include "Sound/QuartzSubscription.h"
 #include "ProfilingDebugging/CountersTrace.h"
 #include "Stats/Stats.h"
 
@@ -75,13 +76,7 @@ void UQuartzSubsystem::BeginDestroy()
 	Super::BeginDestroy();
 
 	// force un-subscribe all Quartz tickable objects
-	TArray<FQuartzTickableObject*> SubscribersCopy = QuartzTickSubscribers;
-	for (FQuartzTickableObject* Entry : SubscribersCopy)
-	{
-		Entry->Shutdown();
-	}
 	QuartzTickSubscribers.Reset();
-
 	SubsystemClockManager.Shutdown();
 	SubsystemClockManager.Flush();
 }
@@ -213,14 +208,14 @@ UQuartzSubsystem* UQuartzSubsystem::Get(UWorld* World)
 }
 
 
-TSharedPtr<Audio::FShareableQuartzCommandQueue, ESPMode::ThreadSafe> UQuartzSubsystem::CreateQuartzCommandQueue()
+TSharedPtr<Audio::TQuartzShareableCommandQueue<FQuartzTickableObject>, ESPMode::ThreadSafe> UQuartzSubsystem::CreateQuartzCommandQueue()
 {
-	return MakeShared<Audio::FShareableQuartzCommandQueue, ESPMode::ThreadSafe>();
+	return MakeShared<Audio::TQuartzShareableCommandQueue<FQuartzTickableObject>, ESPMode::ThreadSafe>();
 }
 
 
 
-Audio::FQuartzQuantizedRequestData UQuartzSubsystem::CreateDataDataForSchedulePlaySound(UQuartzClockHandle* InClockHandle, const FOnQuartzCommandEventBP& InDelegate, const FQuartzQuantizationBoundary& InQuantizationBoundary)
+Audio::FQuartzQuantizedRequestData UQuartzSubsystem::CreateRequestDataForSchedulePlaySound(UQuartzClockHandle* InClockHandle, const FOnQuartzCommandEventBP& InDelegate, const FQuartzQuantizationBoundary& InQuantizationBoundary)
 {
 	Audio::FQuartzQuantizedRequestData CommandInitInfo;
 
@@ -235,15 +230,6 @@ Audio::FQuartzQuantizedRequestData UQuartzSubsystem::CreateDataDataForSchedulePl
 	CommandInitInfo.QuantizedCommandPtr = MakeShared<Audio::FQuantizedPlayCommand>();
 	CommandInitInfo.GameThreadSubscribers.Append(InQuantizationBoundary.GameThreadSubscribers);
 
-	const bool bRequiresAudioDevice = CommandInitInfo.QuantizedCommandPtr->RequiresAudioDevice();
-	const bool bClockManagedByAudioDevice = ClockManagerTypeMap.Contains(CommandInitInfo.ClockName)
-		&& ClockManagerTypeMap[CommandInitInfo.ClockName] == EQuarztClockManagerType::AudioEngine;
-
-	if (!(bRequiresAudioDevice && bClockManagedByAudioDevice))
-	{
-		return {};
-	}
-
 	if (InDelegate.IsBound())
 	{
 		CommandInitInfo.GameThreadDelegateID = InClockHandle->AddCommandDelegate(InDelegate, CommandInitInfo.GameThreadSubscribers);
@@ -253,13 +239,41 @@ Audio::FQuartzQuantizedRequestData UQuartzSubsystem::CreateDataDataForSchedulePl
 }
 
 
+Audio::FQuartzQuantizedRequestData UQuartzSubsystem::CreateDataForTickRateChange(UQuartzClockHandle* InClockHandle, const FOnQuartzCommandEventBP& InDelegate, const Audio::FQuartzClockTickRate& InNewTickRate, const FQuartzQuantizationBoundary& InQuantizationBoundary)
+{
+	// new static function
+	return UQuartzSubsystem::CreateRequestDataForTickRateChange(InClockHandle, InDelegate, InNewTickRate, InQuantizationBoundary);
+}
+
+
+Audio::FQuartzQuantizedRequestData UQuartzSubsystem::CreateDataForTransportReset(UQuartzClockHandle* InClockHandle, const FQuartzQuantizationBoundary& InQuantizationBoundary, const FOnQuartzCommandEventBP& InDelegate)
+{
+	// new static function
+	return UQuartzSubsystem::CreateRequestDataForTransportReset(InClockHandle, InQuantizationBoundary, InDelegate);
+}
+
+
+Audio::FQuartzQuantizedRequestData CreateDataForStartOtherClock(UQuartzClockHandle* InClockHandle, FName InClockToStart, const FQuartzQuantizationBoundary& InQuantizationBoundary, const FOnQuartzCommandEventBP& InDelegate)
+{
+	// new static function
+	return UQuartzSubsystem::CreateRequestDataForStartOtherClock(InClockHandle, InClockToStart, InQuantizationBoundary, InDelegate);
+}
+
+
+Audio::FQuartzQuantizedRequestData CreateDataDataForSchedulePlaySound(UQuartzClockHandle* InClockHandle, const FOnQuartzCommandEventBP& InDelegate, const FQuartzQuantizationBoundary& InQuantizationBoundary)
+{
+	// new static function
+	return UQuartzSubsystem::CreateRequestDataForSchedulePlaySound(InClockHandle, InDelegate, InQuantizationBoundary);
+}
+
+
 bool UQuartzSubsystem::IsQuartzEnabled()
 {
 	return true;
 }
 
 
-Audio::FQuartzQuantizedRequestData UQuartzSubsystem::CreateDataForTickRateChange(UQuartzClockHandle* InClockHandle, const FOnQuartzCommandEventBP& InDelegate, const Audio::FQuartzClockTickRate& InNewTickRate, const FQuartzQuantizationBoundary& InQuantizationBoundary)
+Audio::FQuartzQuantizedRequestData UQuartzSubsystem::CreateRequestDataForTickRateChange(UQuartzClockHandle* InClockHandle, const FOnQuartzCommandEventBP& InDelegate, const Audio::FQuartzClockTickRate& InNewTickRate, const FQuartzQuantizationBoundary& InQuantizationBoundary)
 {
 	if (!ensure(InClockHandle))
 	{
@@ -285,7 +299,7 @@ Audio::FQuartzQuantizedRequestData UQuartzSubsystem::CreateDataForTickRateChange
 	return CommandInitInfo;
 }
 
-Audio::FQuartzQuantizedRequestData UQuartzSubsystem::CreateDataForTransportReset(UQuartzClockHandle* InClockHandle, const FQuartzQuantizationBoundary& InQuantizationBoundary, const FOnQuartzCommandEventBP& InDelegate)
+Audio::FQuartzQuantizedRequestData UQuartzSubsystem::CreateRequestDataForTransportReset(UQuartzClockHandle* InClockHandle, const FQuartzQuantizationBoundary& InQuantizationBoundary, const FOnQuartzCommandEventBP& InDelegate)
 {
 	if (!ensure(InClockHandle))
 	{
@@ -310,7 +324,7 @@ Audio::FQuartzQuantizedRequestData UQuartzSubsystem::CreateDataForTransportReset
 	return CommandInitInfo;
 }
 
-Audio::FQuartzQuantizedRequestData UQuartzSubsystem::CreateDataForStartOtherClock(UQuartzClockHandle* InClockHandle, FName InClockToStart, const FQuartzQuantizationBoundary& InQuantizationBoundary, const FOnQuartzCommandEventBP& InDelegate)
+Audio::FQuartzQuantizedRequestData UQuartzSubsystem::CreateRequestDataForStartOtherClock(UQuartzClockHandle* InClockHandle, FName InClockToStart, const FQuartzQuantizationBoundary& InQuantizationBoundary, const FOnQuartzCommandEventBP& InDelegate)
 {
 	if (!ensure(InClockHandle))
 	{
@@ -362,10 +376,10 @@ UQuartzClockHandle* UQuartzSubsystem::CreateNewClock(const UObject* WorldContext
 		InSettings.TimeSignature.NumBeats = 1;
 	}
 
-	ClockManager->GetOrCreateClock(ClockName, InSettings, bOverrideSettingsIfClockExists);
+	Audio::FQuartzClockProxy ClockHandle = ClockManager->GetOrCreateClock(ClockName, InSettings, bOverrideSettingsIfClockExists);
 
 	UQuartzClockHandle* ClockHandlePtr = (UQuartzClockHandle *)(NewObject<UQuartzClockHandle>()->Init(WorldContextObject->GetWorld()));
-	ClockHandlePtr->SubscribeToClock(WorldContextObject, ClockName);
+	ClockHandlePtr->SubscribeToClock(WorldContextObject, ClockName, &ClockHandle);
 
 	return ClockHandlePtr;
 }
@@ -394,18 +408,15 @@ void UQuartzSubsystem::DeleteClockByHandle(const UObject* WorldContextObject, UQ
 UQuartzClockHandle* UQuartzSubsystem::GetHandleForClock(const UObject* WorldContextObject, FName ClockName)
 {
 	Audio::FQuartzClockManager* ClockManager = GetManagerForClock(WorldContextObject, ClockName);
-	if (!ClockManager)
+	if (!ClockManager || !ClockManager->DoesClockExist(ClockName))
 	{
 		return nullptr;
 	}
 
-	if (ClockManager->DoesClockExist(ClockName))
-	{
-		UQuartzClockHandle* ClockHandlePtr = (UQuartzClockHandle *)(NewObject<UQuartzClockHandle>()->Init(WorldContextObject->GetWorld()));
-		return ClockHandlePtr->SubscribeToClock(WorldContextObject, ClockName);
-	}
+	Audio::FQuartzClockProxy ClockHandle = ClockManager->GetClock(ClockName);
 
-	return nullptr;
+	UQuartzClockHandle* ClockHandlePtr = (UQuartzClockHandle *)(NewObject<UQuartzClockHandle>()->Init(WorldContextObject->GetWorld()));
+	return ClockHandlePtr->SubscribeToClock(WorldContextObject, ClockName, &ClockHandle);
 }
 
 
diff --git a/Engine/Source/Runtime/AudioMixer/Public/Quartz/AudioMixerClock.h b/Engine/Source/Runtime/AudioMixer/Public/Quartz/AudioMixerClock.h
--- a/Engine/Source/Runtime/AudioMixer/Public/Quartz/AudioMixerClock.h
+++ b/Engine/Source/Runtime/AudioMixer/Public/Quartz/AudioMixerClock.h
@@ -10,27 +10,90 @@ namespace Audio
 {
 	// forwards
 	class FMixerDevice;
+	class FQuartzClock;
 	class FMixerSourceManager;
 	class FQuartzClockManager;
 
-	// Contains the pending command and the number of frames it has to wait to fire
-	struct PendingCommand
+	template<class ListenerType>
+	class TQuartzShareableCommandQueue;
+
+	using FQuartzClockCommandQueuePtr = TSharedPtr<TQuartzShareableCommandQueue<FQuartzClock>, ESPMode::ThreadSafe>;
+	using FQuartzClockCommandQueueWeakPtr = TWeakPtr<TQuartzShareableCommandQueue<FQuartzClock>, ESPMode::ThreadSafe>;
+
+	
+	/**
+	 *	FQuartzClockProxy:
+	 *
+	 *		This class is a C++ handle to the underlying clock.
+	 *		
+	 *		It is mostly a wrapper around a TWeakPtr<FQuartzClock> and
+	 *		TSharedPtr<TQuartzShareableCommandQueue<FQuartzClock>.
+	 *		
+	 *		The getters query the underlying FQuartzClock directly,
+	 *		which returns values updated during the last audio-engine tick
+	 *
+	 *		If you need to add more getters, add copies of the members in question to
+	 *		FQuartzClock::FQuartzClockState and update FQuartzClock::UpdateCachedState()
+	 *		for thread-safe access (or manually protect access w/ CachedClockStateCritSec)
+	 *
+	 *		SendCommandToClock() can be used to execute lambdas at the beginning
+	 *		of the next clock tick.  These lambdas can call FQuartzClock's public methods safely.
+	 *
+	 *		Your lambda will take an FQuartzClock* as an argument, which will be passed in by the
+	 *		FQuartzClock itself when it pumps the command queue.
+	 */
+	class AUDIOMIXER_API FQuartzClockProxy
 	{
+	public:
 		// ctor
-		PendingCommand(TSharedPtr<IQuartzQuantizedCommand> InCommand, int32 InNumFramesUntilExec)
-		: Command(InCommand)
-		, NumFramesUntilExec(InNumFramesUntilExec)
-		{
-		}
+		FQuartzClockProxy() {}
+		FQuartzClockProxy(TSharedPtr<FQuartzClock, ESPMode::ThreadSafe> InClock);
+
+		FName GetClockName() const { return ClockId; }
+
+		bool IsValid() const;
+		operator bool() const { return IsValid(); }
+
+		bool DoesClockExist() const;
+
+		bool IsClockRunning() const;
+
+		Audio::FQuartzClockTickRate GetTickRate() const;
+
+		float GetEstimatedClockRunTimeSeconds() const;
 
-		// Quantized Command Object
-		TSharedPtr<IQuartzQuantizedCommand> Command;
+		FQuartzTransportTimeStamp GetCurrentClockTimestamp() const;
 
-		// Countdown to execution
-		int32 NumFramesUntilExec{ 0 };
-	}; // struct PendingCommand
+		float GetDurationOfQuantizationTypeInSeconds(const EQuartzCommandQuantization& QuantizationType, float Multiplier) const;
 
-	// Class that encapsulates sample-accurate timing logic, as well as firing QuantizedAudioCommands
+		// returns false if the clock is not valid or has shut down
+		bool SendCommandToClock(TFunction<void(FQuartzClock*)> InCommand);
+
+	private:
+		FName ClockId;
+
+		FQuartzClockCommandQueueWeakPtr SharedQueue;
+
+	protected:
+		TWeakPtr<FQuartzClock, ESPMode::ThreadSafe> ClockWeakPtr;
+
+	}; // class FQuartzClockProxy
+
+
+	
+	/**
+	 *	FQuartzClock:
+	 *
+	 *		This class receives, schedules, and fires quantized commands. 
+	 *		The underlying FQuartzMetronome handles all counting / timing logic.
+	 *
+	 *		This class gets ticked externally (i.e. by some Clock Manager)
+	 *		and counts down the time-to-fire the commands in audio frames.
+	 *
+	 *
+	 *		UpdateCachedState() updates a game-thread copy of data accessed via FQuartzClockProxy
+	 *		(see FQuartzClockState)
+	 */
 	class AUDIOMIXER_API FQuartzClock
 	{
 	public:
@@ -41,79 +104,51 @@ namespace Audio
 		// dtor
 		~FQuartzClock();
 
+		// Transport Control:
 		// alter the tick rate (take by-value to make sample-rate adjustments in-place)
 		void ChangeTickRate(FQuartzClockTickRate InNewTickRate, int32 NumFramesLeft = 0);
 
-		// alter the time signature
 		void ChangeTimeSignature(const FQuartzTimeSignature& InNewTimeSignature);
 
-		// start ticking the clock
 		void Resume();
 
-		// stop ticking and reset the clock
-		void Stop(bool CancelPendingEvents);
-
-		// stop ticking the clock
 		void Pause();
 
-		// reset the metronome
 		void Restart(bool bPause = true);
 
-		// shutdown
-		void Shutdown();
+		void Stop(bool CancelPendingEvents); // Pause + Restart
 
-		// low-resolution clock update
-		// (not sample-accurate!, useful when running without an Audio Device)
-		void LowResolutionTick(float InDeltaTimeSeconds);
-
-		// sample accurate clock update
-		void Tick(int32 InNumFramesUntilNextTick);
-
-		// Set the sample rate of the clock
 		void SetSampleRate(float InNewSampleRate);
 
-		// get the tick rate
-		FQuartzClockTickRate GetTickRate() const { return Metronome.GetTickRate(); }
-
-		// get the identifier of the clock
-		FName GetName() const { return Name; }
-
-		// clock will persist across level changes
-		bool IgnoresFlush();
+		void ResetTransport();
 
-		// Does this clock match the provided settings
-		bool DoesMatchSettings(const FQuartzClockSettings& InClockSettings) const;
+		// (used for StartOtherClock command to handle the sub-tick as the target clock)
+		void AddToTickDelay(int32 NumFramesOfDelayToAdd);
 
-		void SubscribeToTimeDivision(MetronomeCommandQueuePtr InListenerQueue, EQuartzCommandQuantization InQuantizationBoundary);
+		// (used for StartOtherClock command to handle the sub-tick as the target clock)
+		void SetTickDelay(int32 NumFramesOfDelay);
 
-		void SubscribeToAllTimeDivisions(MetronomeCommandQueuePtr InListenerQueue);
+		void Shutdown();
 
-		void UnsubscribeFromTimeDivision(MetronomeCommandQueuePtr InListenerQueue, EQuartzCommandQuantization InQuantizationBoundary);
+		// Getters:
+		FQuartzClockTickRate GetTickRate();
 
-		void UnsubscribeFromAllTimeDivisions(MetronomeCommandQueuePtr InListenerQueue);
+		FName GetName() const;
 
-		// Add a new event to be triggered by this clock
-		// TODO: return a handle to this event so "looping" events can be canceled
-		void AddQuantizedCommand(FQuartzQuantizationBoundary InQuantizationBondary, TSharedPtr<IQuartzQuantizedCommand> InNewEvent);
+		bool IgnoresFlush() const;
 
-		// Cancel pending command
-		bool CancelQuantizedCommand(TSharedPtr<IQuartzQuantizedCommand> InCommandPtr);
+		bool DoesMatchSettings(const FQuartzClockSettings& InClockSettings) const;
 
-		// does the clock have any pending events
 		bool HasPendingEvents() const;
 
 		int32 NumPendingEvents() const;
 
-		// is the clock currently ticking?
-		bool IsRunning();
+		bool IsRunning() const;
 
-		// Returns the duration in seconds of the given Quantization Type
 		float GetDurationOfQuantizationTypeInSeconds(const EQuartzCommandQuantization& QuantizationType, float Multiplier);
 
-		// Returns the current location of the clock in the transport
 		FQuartzTransportTimeStamp GetCurrentTimestamp();
 
-		// Returns the amount of time, in seconds, the clock has been running. Caution: due to latency, this will not be perfectly accurate
 		float GetEstimatedRunTime();
 
 		FMixerDevice* GetMixerDevice();
@@ -122,19 +157,63 @@ namespace Audio
 
 		FQuartzClockManager* GetClockManager();
 
-		void ResetTransport();
+		FQuartzClockCommandQueueWeakPtr GetCommandQueue() const;
 
-		void AddToTickDelay(int32 NumFramesOfDelayToAdd)
-		{
-			TickDelayLengthInFrames += NumFramesOfDelayToAdd;
-		}
+		// Metronome Event Subscription:
+		void SubscribeToTimeDivision(MetronomeCommandQueuePtr InListenerQueue, EQuartzCommandQuantization InQuantizationBoundary);
 
-		void SetTickDelay(int32 NumFramesOfDelay)
-		{
-			TickDelayLengthInFrames = NumFramesOfDelay;
-		}
+		void SubscribeToAllTimeDivisions(MetronomeCommandQueuePtr InListenerQueue);
+
+		void UnsubscribeFromTimeDivision(MetronomeCommandQueuePtr InListenerQueue, EQuartzCommandQuantization InQuantizationBoundary);
+
+		void UnsubscribeFromAllTimeDivisions(MetronomeCommandQueuePtr InListenerQueue);
+
+		// Quantized Command Management:
+		void AddQuantizedCommand(FQuartzQuantizedCommandInitInfo& InQuantizationCommandInitInfo);
+
+		void AddQuantizedCommand(FQuartzQuantizationBoundary InQuantizationBondary, TSharedPtr<IQuartzQuantizedCommand> InNewEvent);
+
+		bool CancelQuantizedCommand(TSharedPtr<IQuartzQuantizedCommand> InCommandPtr);
+
+// 	protected: todo - expose this part of the interface through an authoritative handle
+
+		// low-resolution clock update
+		// (not sample-accurate!, useful when running without an Audio Device)
+		void LowResolutionTick(float InDeltaTimeSeconds);
+
+		// sample accurate clock update
+		void Tick(int32 InNumFramesUntilNextTick);
 
 	private:
+		// Contains the pending command and the number of frames it has to wait to fire
+		struct PendingCommand
+		{
+			// ctor
+			PendingCommand(TSharedPtr<IQuartzQuantizedCommand> InCommand, int32 InNumFramesUntilExec)
+				: Command(InCommand)
+				, NumFramesUntilExec(InNumFramesUntilExec)
+			{
+			}
+
+			// Quantized Command Object
+			TSharedPtr<IQuartzQuantizedCommand> Command;
+
+			// Countdown to execution
+			int32 NumFramesUntilExec{ 0 };
+		}; // struct PendingCommand
+
+		// mutex-protected update at the end of Tick()
+		FCriticalSection CachedClockStateCritSec;
+		void UpdateCachedState();
+
+		// data is cached when an FQuartzClock is ticked
+		struct FQuartzClockState
+		{
+			FQuartzClockTickRate TickRate;
+			FQuartzTransportTimeStamp TimeStamp;
+			float RunTimeInSeconds;
+		} CachedClockState;
+
 		void TickInternal(int32 InNumFramesUntilNextTick, TArray<PendingCommand>& CommandsToTick, int32 FramesOfLatency = 0, int32 FramesOfDelay = 0);
 
 		bool CancelQuantizedCommandInternal(TSharedPtr<IQuartzQuantizedCommand> InCommandPtr, TArray<PendingCommand>& CommandsToTick);
@@ -149,9 +228,11 @@ namespace Audio
 
 		FName Name;
 
-		// TODO: Make this configurable
 		float ThreadLatencyInMilliseconds{ 40.f };
 
+		// Command queue handed out to GameThread objects to queue commands. These get executed at the top of Tick()
+		mutable FQuartzClockCommandQueuePtr PreTickCommands; // (mutable for lazy init in GetCommandQueue())
+
 		// Container of external commands to be executed (TUniquePointer<QuantizedAudioCommand>)
 		TArray<PendingCommand> ClockAlteringPendingCommands;
 		TArray<PendingCommand> PendingCommands;
diff --git a/Engine/Source/Runtime/AudioMixer/Public/Quartz/AudioMixerClockHandle.h b/Engine/Source/Runtime/AudioMixer/Public/Quartz/AudioMixerClockHandle.h
--- a/Engine/Source/Runtime/AudioMixer/Public/Quartz/AudioMixerClockHandle.h
+++ b/Engine/Source/Runtime/AudioMixer/Public/Quartz/AudioMixerClockHandle.h
@@ -8,84 +8,22 @@
 #include "Sound/QuartzSubscription.h"
 #include "Quartz/QuartzSubsystem.h"
 #include "Quartz/QuartzMetronome.h"
-#include "UObject/GCObject.h"
 
 #include "AudioMixerClockHandle.generated.h"
 
 
-class UAudioComponent;
 
 
-class AUDIOMIXER_API FQuartzTickableObject
-{
-	struct AUDIOMIXER_API FQuartzTickableObjectGCObjectMembers : public FGCObject
-	{
-	public:
-		UQuartzSubsystem* QuartzSubsystem;
-		UWorld* WorldPtr{ nullptr };
-
-		void AddReferencedObjects(FReferenceCollector& Collector) override;
-		virtual FString GetReferencerName() const override
-		{
-			return TEXT("FQuartzTickableObject::FQuartzTickableObjectGCObjectMembers");
-		}
-	};
-
-	public:
-		FQuartzTickableObject() {}
-		virtual ~FQuartzTickableObject();
-
-		FQuartzTickableObject* Init(UWorld* InWorldPtr);
-
-		// called by the associated QuartzSubsystem
-		void QuartzTick(float DeltaTime);
-		bool QuartzIsTickable() const;
-		UWorld* QuartzGetWorld() const { return GCObjectMembers.WorldPtr; }
-
-		bool IsInitialized() const { return bHasBeenInitialized; }
-
-		// access to the associated QuartzSubsystem
-		UQuartzSubsystem* GetQuartzSubsystem() const { return GCObjectMembers.QuartzSubsystem; }
-
-		TSharedPtr<Audio::FShareableQuartzCommandQueue, ESPMode::ThreadSafe> GetCommandQueue();
-
-		int32 AddCommandDelegate(const FOnQuartzCommandEventBP& InDelegate, TArray<TSharedPtr<Audio::FShareableQuartzCommandQueue, ESPMode::ThreadSafe>>& TargetSubscriberArray);
-
-		// virtual interface
-		virtual void ProcessCommand(const Audio::FQuartzQuantizedCommandDelegateData& Data) {};
-
-		virtual void ProcessCommand(const Audio::FQuartzMetronomeDelegateData& Data) {};
-
-		virtual void ProcessCommand(const Audio::FQuartzQueueCommandData& Data) {};
-
-		void Shutdown();
-
-
-	protected:
-		struct CommandDelegateGameThreadData
-		{
-			FOnQuartzCommandEvent MulticastDelegate;
-			FThreadSafeCounter RefCount;
-		};
-
-		struct MetronomeDelegateGameThreadData
-		{
-			FOnQuartzMetronomeEvent MulticastDelegate;
-		};
-	
-		void PumpCommandQueue();
-
-		TSharedPtr<Audio::FShareableQuartzCommandQueue, ESPMode::ThreadSafe> CommandQueuePtr;
-		TArray<TFunction<void(FQuartzTickableObject*)>> TempCommandQueue;
-
-		TArray<CommandDelegateGameThreadData> QuantizedCommandDelegates;
-		MetronomeDelegateGameThreadData MetronomeDelegates[static_cast<int32>(EQuartzCommandQuantization::Count)];
-
-	private:
-		FQuartzTickableObjectGCObjectMembers GCObjectMembers;
-		bool bHasBeenInitialized = false;
-};
-
+/**
+ *  This class is a BP / Game thread wrapper around FQuartzClockProxy
+ *	(to talk to the underlying clock)
+ 
+ *  ...and inherits from FQuartzTickableObject
+ *	(to listen to the underlying clock)
+ *  
+ *  It can subscribe to Quantized Event & Metronome delegates to synchronize
+ *  gameplay & VFX to Quartz events fired from the Audio Engine
+ */
 UCLASS(BlueprintType, Blueprintable, Transient, ClassGroup = Quartz, meta = (BlueprintSpawnableComponent))
 class AUDIOMIXER_API UQuartzClockHandle : public UObject, public FQuartzTickableObject
 {
@@ -102,7 +40,6 @@ public:
 	void BeginDestroy() override;
 	// end UObject interface
 
-
 // Begin Blueprint Interface
 
 	// Clock manipulation
@@ -177,7 +114,6 @@ public:
 	UFUNCTION(BlueprintCallable, Category = "Quantization", meta = (WorldContext = "WorldContextObject", AdvancedDisplay = "QuantizationBoundary, Delegate", AutoCreateRefTerm = "QuantizationBoundary, Delegate", Keywords = "BPM, Tempo"))
 	void SetBeatsPerMinute(const UObject* WorldContextObject, UPARAM(ref) const FQuartzQuantizationBoundary& QuantizationBoundary, const FOnQuartzCommandEventBP& Delegate, UQuartzClockHandle*& ClockHandle, float BeatsPerMinute = 60.f);
 
-
 	// Metronome getters
 	UFUNCTION(BlueprintCallable, Category = "Quantization", meta = (WorldContext = "WorldContextObject", AutoCreateRefTerm = "InDelegate", Keywords = "BPM, Tempo"))
 	float GetMillisecondsPerTick(const UObject* WorldContextObject) const;
@@ -197,30 +133,30 @@ public:
 // End Blueprint Interface
 	void QueueQuantizedSound(const UObject* WorldContextObject, UQuartzClockHandle*& ClockHandle, const FAudioComponentCommandInfo& AudioComponentData, const FOnQuartzCommandEventBP& InDelegate, const FQuartzQuantizationBoundary& InTargetBoundary);
 
-	UQuartzClockHandle* SubscribeToClock(const UObject* WorldContextObject, FName ClockName);
+	UQuartzClockHandle* SubscribeToClock(const UObject* WorldContextObject, FName ClockName, Audio::FQuartzClockProxy* InHandlePtr = nullptr);
 
 	FName GetClockName() const { return CurrentClockId; }
 
 	FName GetHandleName() const { return ClockHandleId; }
-	
+
 	bool DoesClockExist(const UObject* WorldContextObject) const
 	{
-		if (UQuartzSubsystem* QuartzSubsystem = GetQuartzSubsystem())
-		{
-			return QuartzSubsystem->DoesClockExist(WorldContextObject, CurrentClockId);
-		}
-
-		return false;
+		check(RawHandle);
+		return RawHandle.DoesClockExist();
 	}
 
-	virtual void ProcessCommand(const Audio::FQuartzQuantizedCommandDelegateData& Data) override;
-
-	virtual void ProcessCommand(const Audio::FQuartzMetronomeDelegateData& Data) override;
+	UE_DEPRECATED(5.1, "This function should not be called directly, and the original functionality has been moved into FQuartzTickable")
+	virtual void ProcessCommand(const Audio::FQuartzQuantizedCommandDelegateData& Data) override {}
+	UE_DEPRECATED(5.1, "This function should not be called directly, and the original functionality has been moved into FQuartzTickable")
+	virtual void ProcessCommand(const Audio::FQuartzMetronomeDelegateData& Data) override {};
 
 	bool GetCurrentTickRate(const UObject* WorldContextObject, Audio::FQuartzClockTickRate& OutTickRate) const;
 
 private:
-	
+	void SetTickRateInternal(const FQuartzQuantizationBoundary& InQuantizationBoundary, const FOnQuartzCommandEventBP& InDelegate, const Audio::FQuartzClockTickRate& NewTickRate);
+
+	Audio::FQuartzClockProxy RawHandle;
+
 	FName CurrentClockId;
 
 	FName ClockHandleId;
diff --git a/Engine/Source/Runtime/AudioMixer/Public/Quartz/AudioMixerClockManager.h b/Engine/Source/Runtime/AudioMixer/Public/Quartz/AudioMixerClockManager.h
--- a/Engine/Source/Runtime/AudioMixer/Public/Quartz/AudioMixerClockManager.h
+++ b/Engine/Source/Runtime/AudioMixer/Public/Quartz/AudioMixerClockManager.h
@@ -9,6 +9,7 @@ namespace Audio
 {
 	// forwards
 	class FMixerDevice;
+	class FQuartzClock;
 
 	// Class that owns, updates, and provides access to all active clocks
 	// All methods are thread-safe. The method locks if it returns a value, and stages a command if it returns void
@@ -34,14 +35,14 @@ namespace Audio
 
 		// add (and take ownership of) a new clock
 		// safe to call from AudioThread (uses critical section)
-		TSharedPtr<FQuartzClock> GetOrCreateClock(const FName& InClockName, const FQuartzClockSettings& InClockSettings, bool bOverrideTickRateIfClockExists = false);
-
+		FQuartzClockProxy GetOrCreateClock(const FName& InClockName, const FQuartzClockSettings& InClockSettings, bool bOverrideTickRateIfClockExists = false);
+		FQuartzClockProxy GetClock(const FName& InClockName);
 
 		// returns true if a clock with the given name already exists.
-		bool DoesClockExist(const FName& InClockName);
+		bool DoesClockExist(const FName& InClockName); // maxtodo: deprecate when clock is on the subsystem
 
 		// returns true if the name is running
-		bool IsClockRunning(const FName& InClockName);
+		bool IsClockRunning(const FName& InClockName);  // maxtodo: deprecate when clock is on the subsystem?
 
 		// Returns the duration in seconds of the given Quantization Type, or -1 if the Clock is invalid or nonexistent
 		float GetDurationOfQuantizationTypeInSeconds(const FName& InClockName, const EQuartzCommandQuantization& QuantizationType, float Multiplier);
@@ -101,7 +102,7 @@ namespace Audio
 
 		bool HasClockBeenTickedThisUpdate(FName InClockName);
 
-		int32 GetLastUpdateSizeInFrames() { return LastUpdateSizeInFrames; }
+		int32 GetLastUpdateSizeInFrames() const { return LastUpdateSizeInFrames; }
 
 		// get access to the owning FMixerDevice
 		FMixerDevice* GetMixerDevice() const;
@@ -124,5 +125,8 @@ namespace Audio
 
 		FThreadSafeCounter LastClockTickedIndex{ 0 };
 		int32 LastUpdateSizeInFrames{ 0 };
+
+		// allow a clock that is queuing a command directly use FindClock() to retrieve the TSharedPtr<FQuartzClock>
+		friend void FQuartzClock::AddQuantizedCommand(FQuartzQuantizedCommandInitInfo& InQuantizationCommandInitInfo);
 	};
 } // namespace Audio
diff --git a/Engine/Source/Runtime/AudioMixer/Public/Quartz/QuartzMetronome.h b/Engine/Source/Runtime/AudioMixer/Public/Quartz/QuartzMetronome.h
--- a/Engine/Source/Runtime/AudioMixer/Public/Quartz/QuartzMetronome.h
+++ b/Engine/Source/Runtime/AudioMixer/Public/Quartz/QuartzMetronome.h
@@ -6,32 +6,33 @@
 
 namespace Audio
 {
-	using MetronomeCommandQueuePtr = TSharedPtr<FShareableQuartzCommandQueue, ESPMode::ThreadSafe>;
+	using MetronomeCommandQueuePtr = TSharedPtr<TQuartzShareableCommandQueue<FQuartzTickableObject>, ESPMode::ThreadSafe>;
 
 	// Class to track the passage of musical time, and allow subscribers to be notified when these musical events take place
 	class FQuartzMetronome
 	{
 	public:
 		// ctor
-		FQuartzMetronome();
-		FQuartzMetronome(const FQuartzTimeSignature& InTimeSignature);
+		FQuartzMetronome(FName InClockName = {});
+		FQuartzMetronome(const FQuartzTimeSignature& InTimeSignature, FName InClockName = {});
 
 		// dtor
 		~FQuartzMetronome();
 
-		// called by owning FQuartzClock
+		// Transport Control:
 		void Tick(int32 InNumSamples, int32 FramesOfLatency = 0);
 
-		// called by owning FQuartzClock
 		void SetTickRate(FQuartzClockTickRate InNewTickRate, int32 NumFramesLeft = 0);
 
-		FQuartzClockTickRate GetTickRate() const { return CurrentTickRate; }
-
 		void SetSampleRate(float InNewSampleRate);
 
-		// affects bars/beats values we send back to the game thread
 		void SetTimeSignature(const FQuartzTimeSignature& InNewTimeSignature);
 
+		void ResetTransport();
+
+		// Getters
+		FQuartzClockTickRate GetTickRate() const { return CurrentTickRate; }
+
 		double GetFramesUntilBoundary(FQuartzQuantizationBoundary InQuantizationBoundary) const;
 
 		const FQuartzTimeSignature & GetTimeSignature() const { return CurrentTimeSignature; }
@@ -40,6 +41,7 @@ namespace Audio
 
 		double GetTimeSinceStart() const { return TimeSinceStart; }
 
+		// Event Subscription
 		void SubscribeToTimeDivision(MetronomeCommandQueuePtr InListenerQueue, EQuartzCommandQuantization InQuantizationBoundary);
 
 		void SubscribeToAllTimeDivisions(MetronomeCommandQueuePtr InListenerQueue);
@@ -48,9 +50,10 @@ namespace Audio
 
 		void UnsubscribeFromAllTimeDivisions(MetronomeCommandQueuePtr InListenerQueue);
 
-		void ResetTransport();
 
 	private:
+
+		// Helpers:
 		void RecalculateDurations();
 
 		void FireEvents(int32 EventFlags);
@@ -69,7 +72,6 @@ namespace Audio
 
 		FQuartzClockTickRate CurrentTickRate;
 
-		// (Subtract one because Tick and 1/32 are the same.)
 		TArray<MetronomeCommandQueuePtr> MetronomeSubscriptionMatrix[static_cast<int32>(EQuartzCommandQuantization::Count)];
 		
 		// wrapper around our array so it can be indexed into by different Enums that represent musical time
@@ -118,5 +120,7 @@ namespace Audio
 		//Keeps track of time in seconds since the Clock was last reset
 		double TimeSinceStart;
 
+		FName ClockName;
+
 	}; // class QuartzMetronome
 } // namespace Audio
\ No newline at end of file
diff --git a/Engine/Source/Runtime/AudioMixer/Public/Quartz/QuartzSubsystem.h b/Engine/Source/Runtime/AudioMixer/Public/Quartz/QuartzSubsystem.h
--- a/Engine/Source/Runtime/AudioMixer/Public/Quartz/QuartzSubsystem.h
+++ b/Engine/Source/Runtime/AudioMixer/Public/Quartz/QuartzSubsystem.h
@@ -17,11 +17,13 @@ namespace Audio
 {
 	class FMixerDevice;
 	class FQuartzClockManager;
-	class FShareableQuartzCommandQueue;
+
+	template<class ListenerType>
+	class TQuartzShareableCommandQueue;
 }
-class FQuartzTickableObject;
 
-using MetronomeCommandQueuePtr = TSharedPtr<Audio::FShareableQuartzCommandQueue, ESPMode::ThreadSafe>;
+class FQuartzTickableObject;
+using MetronomeCommandQueuePtr = TSharedPtr<Audio::TQuartzShareableCommandQueue<FQuartzTickableObject>, ESPMode::ThreadSafe>;
 
 // GetManagerForClock() logic will need to be updated if more entries are present
 UENUM(BlueprintType)
@@ -52,26 +54,37 @@ public:
 	virtual void Tick(float DeltaTime) override;
 	virtual bool IsTickable() const override;
 	virtual TStatId GetStatId() const override;
-	//~ End FTickableGameObject Interface
 
-	// if we have another stakeholder later we can use polymorphism here.
-	// in the mean time, we are avoiding the virtual overhead
 	void SubscribeToQuartzTick(FQuartzTickableObject* InObjectToTick);
 	void UnsubscribeFromQuartzTick(FQuartzTickableObject* InObjectToTick);
+	//~ End FTickableGameObject Interface
 
 	// static methods
 	static UQuartzSubsystem* Get(UWorld* World);
 
 	// create a new command queue to be shared between ClockHandles and other threads
-	static TSharedPtr<Audio::FShareableQuartzCommandQueue, ESPMode::ThreadSafe> CreateQuartzCommandQueue();
+	static TSharedPtr<Audio::TQuartzShareableCommandQueue<FQuartzTickableObject>, ESPMode::ThreadSafe> CreateQuartzCommandQueue();
 
 	// Helper functions for initializing quantized command initialization struct (to consolidate eyesore)
+	static Audio::FQuartzQuantizedRequestData CreateRequestDataForTickRateChange(UQuartzClockHandle* InClockHandle, const FOnQuartzCommandEventBP& InDelegate, const Audio::FQuartzClockTickRate& InNewTickRate, const FQuartzQuantizationBoundary& InQuantizationBoundary);
+	static Audio::FQuartzQuantizedRequestData CreateRequestDataForTransportReset(UQuartzClockHandle* InClockHandle, const FQuartzQuantizationBoundary& InQuantizationBoundary, const FOnQuartzCommandEventBP& InDelegate);
+	static Audio::FQuartzQuantizedRequestData CreateRequestDataForStartOtherClock(UQuartzClockHandle* InClockHandle, FName InClockToStart, const FQuartzQuantizationBoundary& InQuantizationBoundary, const FOnQuartzCommandEventBP& InDelegate);
+	static Audio::FQuartzQuantizedRequestData CreateRequestDataForSchedulePlaySound(UQuartzClockHandle* InClockHandle, const FOnQuartzCommandEventBP& InDelegate, const FQuartzQuantizationBoundary& InQuantizationBoundary);
+
+	// DEPRECATED HELPERS: non-static versions of the above CreateDataFor...() functions
+	UE_DEPRECATED(5.1, "Use the static (CreateRequestDataFor) version of this function instead")
 	Audio::FQuartzQuantizedRequestData CreateDataForTickRateChange(UQuartzClockHandle* InClockHandle, const FOnQuartzCommandEventBP& InDelegate, const Audio::FQuartzClockTickRate& InNewTickRate, const FQuartzQuantizationBoundary& InQuantizationBoundary);
+
+	UE_DEPRECATED(5.1, "Use the static (CreateRequestDataFor) version of this function instead")
 	Audio::FQuartzQuantizedRequestData CreateDataForTransportReset(UQuartzClockHandle* InClockHandle, const FQuartzQuantizationBoundary& InQuantizationBoundary, const FOnQuartzCommandEventBP& InDelegate);
+	
+	UE_DEPRECATED(5.1, "Use the static (CreateRequestDataFor) version of this function instead")
 	Audio::FQuartzQuantizedRequestData CreateDataForStartOtherClock(UQuartzClockHandle* InClockHandle, FName InClockToStart, const FQuartzQuantizationBoundary& InQuantizationBoundary, const FOnQuartzCommandEventBP& InDelegate);
+
+	UE_DEPRECATED(5.1, "Use the static (CreateRequestDataFor) version of this function instead")
 	Audio::FQuartzQuantizedRequestData CreateDataDataForSchedulePlaySound(UQuartzClockHandle* InClockHandle, const FOnQuartzCommandEventBP& InDelegate, const FQuartzQuantizationBoundary& InQuantizationBoundary);
 
-	UFUNCTION(BlueprintCallable, Category = "Quartz Clock Handle")
+	UFUNCTION(BlueprintCallable, Category = "Quartz Clock Handle", meta = (DeprecatedFunction, DeprecationMessage = "Quartz is always enabled. This function will always return true"))
 	bool IsQuartzEnabled();
 
 	// Clock Creation
@@ -96,7 +109,7 @@ public:
 	bool DoesClockExist(const UObject* WorldContextObject, FName ClockName);
 
 	// returns true if the clock is running
-	UFUNCTION(BlueprintCallable, Category = "Quartz Clock Handle", meta = (WorldContext = "WorldContextObject"))
+	UFUNCTION(BlueprintCallable, Category = "Quartz Clock Handle", meta = (WorldContext = "WorldContextObject", DeprecatedFunction, DeprecationMessage = "Obtain and Query a UQuartzClockHandle instead"))
 	bool IsClockRunning(const UObject* WorldContextObject, FName ClockName);
 
 	/** Returns the duration in seconds of the given Quantization Type
@@ -105,15 +118,15 @@ public:
 	 * @param The quantity of the Quantization Type to calculate the time of
 	 * @return The duration, in seconds, of a multiplier amount of the Quantization Type, or -1 in the case the clock is invalid
 	 */
-	UFUNCTION(BlueprintCallable, Category = "Quartz Clock Handle", meta = (WorldContext = "WorldContextObject"))
+	UFUNCTION(BlueprintCallable, Category = "Quartz Clock Handle", meta = (WorldContext = "WorldContextObject", DeprecatedFunction, DeprecationMessage = "Obtain and Query a UQuartzClockHandle instead"))
 	float GetDurationOfQuantizationTypeInSeconds(const UObject* WorldContextObject, FName ClockName, const EQuartzCommandQuantization& QuantizationType, float Multiplier = 1.0f);
 
 	// Retrieves a timestamp for the clock
-	UFUNCTION(BlueprintCallable, Category = "Quartz Clock Handle", meta = (WorldContext = "WorldContextObject"))
+	UFUNCTION(BlueprintCallable, Category = "Quartz Clock Handle", meta = (WorldContext = "WorldContextObject", DeprecatedFunction, DeprecationMessage = "Obtain and Query a UQuartzClockHandle instead"))
 	FQuartzTransportTimeStamp GetCurrentClockTimestamp(const UObject* WorldContextObject, const FName& InClockName);
 
 	// Returns the amount of time, in seconds, the clock has been running. Caution: due to latency, this will not be perfectly accurate
-	UFUNCTION(BlueprintCallable, Category = "Quartz Clock Handle", meta = (WorldContext = "WorldContextObject"))
+	UFUNCTION(BlueprintCallable, Category = "Quartz Clock Handle", meta = (WorldContext = "WorldContextObject", DeprecatedFunction, DeprecationMessage = "Obtain and Query a UQuartzClockHandle instead"))
 	float GetEstimatedClockRunTime(const UObject* WorldContextObject, const FName& InClockName);
 
 	// latency data (Game thread -> Audio Render Thread)
@@ -146,8 +159,10 @@ public:
 	UFUNCTION(BlueprintCallable, Category = "Quartz Subsystem", meta = (WorldContext = "WorldContextObject"))
 	float GetRoundTripMaxLatency(const UObject* WorldContextObject);
 
+	UE_DEPRECATED(5.1, "Obtain and use a UQuartzClockHandle / FQuartzClockProxy instead")
 	void AddCommandToClock(const UObject* WorldContextObject, Audio::FQuartzQuantizedCommandInitInfo& InQuantizationCommandInitInfo, FName ClockName);
 
+	// maxtodo: deprecate this? or make private?
 	Audio::FQuartzClockManager* GetManagerForClock(const UObject* WorldContextObject, FName ExistingClockName = FName());
 
 private:
diff --git a/Engine/Source/Runtime/Engine/Classes/Sound/QuartzQuantizationUtilities.cpp b/Engine/Source/Runtime/Engine/Classes/Sound/QuartzQuantizationUtilities.cpp
--- a/Engine/Source/Runtime/Engine/Classes/Sound/QuartzQuantizationUtilities.cpp
+++ b/Engine/Source/Runtime/Engine/Classes/Sound/QuartzQuantizationUtilities.cpp
@@ -5,7 +5,9 @@
 #include "Core/Public/CoreGlobals.h"
 #include "Core/Public/Math/NumericLimits.h"
 #include "Core/Public/Math/UnrealMathUtility.h"
+
 #include "AudioMixerDevice.h"
+#include "Sound/QuartzSubscription.h"
 
 #define INVALID_DURATION -1
 
@@ -18,23 +20,18 @@ EQuartzCommandQuantization TimeSignatureQuantizationToCommandQuantization(const
 	{
 		case EQuartzTimeSignatureQuantization::HalfNote :
 			return EQuartzCommandQuantization::HalfNote;
-			break;
 
 		case EQuartzTimeSignatureQuantization::QuarterNote :
 			return EQuartzCommandQuantization::QuarterNote;
-			break;
 
 		case EQuartzTimeSignatureQuantization::EighthNote :
 			return EQuartzCommandQuantization::EighthNote;
-			break;
 
 		case EQuartzTimeSignatureQuantization::SixteenthNote :
 			return EQuartzCommandQuantization::SixteenthNote;
-			break;
 
 		case EQuartzTimeSignatureQuantization::ThirtySecondNote :
 			return EQuartzCommandQuantization::ThirtySecondNote;
-			break;
 
 		default:
 			return EQuartzCommandQuantization::Count;
@@ -73,7 +70,7 @@ bool FQuartzTimeSignature::operator==(const FQuartzTimeSignature& Other) const
 			const bool NumPulsesMatch = (OptionalPulseOverride[i].NumberOfPulses == Other.OptionalPulseOverride[i].NumberOfPulses);
 			const bool DurationsMatch = (OptionalPulseOverride[i].PulseDuration == Other.OptionalPulseOverride[i].PulseDuration);
 
-			if (!(NumPulseEntries && DurationsMatch))
+			if (!(NumPulsesMatch && DurationsMatch))
 			{
 				Result = false;
 				break;
@@ -404,7 +401,7 @@ namespace Audio
 		return nullptr;
 	}
 
-	void IQuartzQuantizedCommand::AddSubscriber(TSharedPtr<FShareableQuartzCommandQueue, ESPMode::ThreadSafe> InSubscriber)
+	void IQuartzQuantizedCommand::AddSubscriber(TSharedPtr<TQuartzShareableCommandQueue<FQuartzTickableObject>, ESPMode::ThreadSafe> InSubscriber)
 	{
 		GameThreadSubscribers.AddUnique(InSubscriber);
 	}
@@ -412,8 +409,8 @@ namespace Audio
 	void IQuartzQuantizedCommand::OnQueued(const FQuartzQuantizedCommandInitInfo& InCommandInitInfo)
 	{
 		TRACE_CPUPROFILER_EVENT_SCOPE(QuartzQuantizedCommand::OnQueued);
-		Audio::FMixerDevice* MixerDevice = InCommandInitInfo.OwningClockPointer->GetMixerDevice();
-		if (MixerDevice)
+
+		if (Audio::FMixerDevice* MixerDevice = InCommandInitInfo.OwningClockPointer->GetMixerDevice())
 		{
 			MixerDevice->QuantizedEventClockManager.PushLatencyTrackerResult(FQuartzCrossThreadMessage::RequestRecieved());
 		}
@@ -457,7 +454,6 @@ namespace Audio
 			Data.DelegateSubType = EQuartzCommandDelegateSubType::CommandOnFailedToQueue;
 			Data.DelegateID = GameThreadDelegateID;
 
-
 			for (auto& SubscriberQueue : GameThreadSubscribers)
 			{
 				if (SubscriberQueue.IsValid())
diff --git a/Engine/Source/Runtime/Engine/Classes/Sound/QuartzQuantizationUtilities.h b/Engine/Source/Runtime/Engine/Classes/Sound/QuartzQuantizationUtilities.h
--- a/Engine/Source/Runtime/Engine/Classes/Sound/QuartzQuantizationUtilities.h
+++ b/Engine/Source/Runtime/Engine/Classes/Sound/QuartzQuantizationUtilities.h
@@ -18,17 +18,17 @@ ENGINE_API DECLARE_LOG_CATEGORY_EXTERN(LogAudioQuartz, Log, All);
 struct FQuartzClockTickRate;
 struct FQuartzQuantizationBoundary;
 struct FQuartzTimeSignature;
-
+class FQuartzTickableObject;
 
 
 namespace Audio
 {
 	// forwards (Audio::)
-	class IAudioMixerQuantizedEventListener;
+	template <class ListenerType>
+	class TQuartzShareableCommandQueue;
+
 	class IQuartzQuantizedCommand;
 	class FQuartzClock;
-	class FShareableQuartzCommandQueue;
-
 	class FMixerDevice;
 
 	struct FQuartzQuantizedCommandDelegateData;
@@ -36,7 +36,6 @@ namespace Audio
 	struct FQuartzQuantizedCommandInitInfo;
 } // namespace Audio
 
-// UOBJECT LAYER:
 
 // An enumeration for specifying quantization for Quartz commands
 UENUM(BlueprintType)
@@ -252,7 +251,7 @@ struct ENGINE_API FQuartzQuantizationBoundary
 	bool bResumeClockOnQueued{ false };
 
 	// Game thread subscribers that will be passed to command init data (for C++ implementations)
-	TArray<TSharedPtr<Audio::FShareableQuartzCommandQueue, ESPMode::ThreadSafe>> GameThreadSubscribers;
+	TArray<TSharedPtr<Audio::TQuartzShareableCommandQueue<FQuartzTickableObject>, ESPMode::ThreadSafe>> GameThreadSubscribers;
 
 	// ctor
 	FQuartzQuantizationBoundary(
@@ -451,7 +450,7 @@ namespace Audio
 		FName OtherClockName;
 		TSharedPtr<IQuartzQuantizedCommand> QuantizedCommandPtr;
 		FQuartzQuantizationBoundary QuantizationBoundary{ EQuartzCommandQuantization::Tick, 1.f, EQuarztQuantizationReference::BarRelative, true };
-		TArray<TSharedPtr<FShareableQuartzCommandQueue, ESPMode::ThreadSafe>> GameThreadSubscribers;
+		TArray<TSharedPtr<TQuartzShareableCommandQueue<FQuartzTickableObject>, ESPMode::ThreadSafe>> GameThreadSubscribers;
 		int32 GameThreadDelegateID{ -1 };
 	};
 
@@ -468,6 +467,7 @@ namespace Audio
 		void SetOwningClockPtr(TSharedPtr<Audio::FQuartzClock> InClockPointer)
 		{
 			OwningClockPointer = InClockPointer;
+			ensure(OwningClockPointer);
 		}
 
 		// shared with FQuartzQuantizedRequestData
@@ -476,7 +476,7 @@ namespace Audio
 		FName OtherClockName;
 		TSharedPtr<IQuartzQuantizedCommand> QuantizedCommandPtr{ nullptr };
 		FQuartzQuantizationBoundary QuantizationBoundary;
-		TArray<TSharedPtr<FShareableQuartzCommandQueue, ESPMode::ThreadSafe>> GameThreadSubscribers;
+		TArray<TSharedPtr<TQuartzShareableCommandQueue<FQuartzTickableObject>, ESPMode::ThreadSafe>> GameThreadSubscribers;
 		int32 GameThreadDelegateID{ -1 };
 
 		// Audio Render thread-specific data:
@@ -501,7 +501,7 @@ namespace Audio
 		// allocate a copy of the derived class
 		virtual TSharedPtr<IQuartzQuantizedCommand> GetDeepCopyOfDerivedObject() const;
 
-		void AddSubscriber(TSharedPtr<FShareableQuartzCommandQueue, ESPMode::ThreadSafe> InSubscriber);
+		void AddSubscriber(TSharedPtr<TQuartzShareableCommandQueue<FQuartzTickableObject>, ESPMode::ThreadSafe> InSubscriber);
 
 		// Command has reached the AudioRenderThread
 		void OnQueued(const FQuartzQuantizedCommandInitInfo& InCommandInitInfo);
@@ -538,8 +538,8 @@ namespace Audio
 		virtual EQuartzCommandType GetCommandType() const = 0;
 
 
-	protected:
-		// base classes can override these to add extra functionality
+	private:
+		// derived classes can override these to add extra functionality
 		virtual void OnQueuedCustom(const FQuartzQuantizedCommandInitInfo& InCommandInitInfo) {}
 		virtual void FailedToQueueCustom() {}
 		virtual void AboutToStartCustom() {}
@@ -548,9 +548,8 @@ namespace Audio
 		virtual void OnClockStartedCustom() {}
 		virtual void CancelCustom() {}
 
-		TArray<TSharedPtr<FShareableQuartzCommandQueue, ESPMode::ThreadSafe>> GameThreadSubscribers;
+		TArray<TSharedPtr<TQuartzShareableCommandQueue<FQuartzTickableObject>, ESPMode::ThreadSafe>> GameThreadSubscribers;
 
-	private:
 		int32 GameThreadDelegateID{ -1 };
 		bool bAboutToStartHasBeenCalled{ false };
 	}; // class IAudioMixerQuantizedCommandBase
@@ -581,7 +580,7 @@ struct ENGINE_API FAudioComponentCommandInfo
 		, CommandID(Other.CommandID)
 	{}
 
- 	FAudioComponentCommandInfo(TSharedPtr<Audio::FShareableQuartzCommandQueue, ESPMode::ThreadSafe> InComponentCommandPtr, FQuartzQuantizationBoundary InAnticaptoryBoundary)
+ 	FAudioComponentCommandInfo(TSharedPtr<Audio::TQuartzShareableCommandQueue<FQuartzTickableObject>, ESPMode::ThreadSafe> InComponentCommandPtr, FQuartzQuantizationBoundary InAnticaptoryBoundary)
 		: ComponentCommandPtr(InComponentCommandPtr)
 		, AnticapatoryBoundary(InAnticaptoryBoundary)
 	{
@@ -589,7 +588,7 @@ struct ENGINE_API FAudioComponentCommandInfo
 		CommandID = CommandIDs++;
 	}
 
-	TSharedPtr<Audio::FShareableQuartzCommandQueue, ESPMode::ThreadSafe> ComponentCommandPtr;
+	TSharedPtr<Audio::TQuartzShareableCommandQueue<FQuartzTickableObject>, ESPMode::ThreadSafe> ComponentCommandPtr;
 	FQuartzQuantizationBoundary AnticapatoryBoundary;
 	uint32 CommandID{ (uint32)INDEX_NONE };
 };
\ No newline at end of file
diff --git a/Engine/Source/Runtime/Engine/Classes/Sound/QuartzSubscription.cpp b/Engine/Source/Runtime/Engine/Classes/Sound/QuartzSubscription.cpp
--- a/Engine/Source/Runtime/Engine/Classes/Sound/QuartzSubscription.cpp
+++ b/Engine/Source/Runtime/Engine/Classes/Sound/QuartzSubscription.cpp
@@ -1,8 +1,7 @@
 // Copyright Epic Games, Inc. All Rights Reserved.
 
 #include "Sound/QuartzSubscription.h"
-
-#include "Quartz/AudioMixerClockHandle.h"
+#include "Quartz/QuartzSubsystem.h"
 
 namespace Audio
 {
@@ -11,35 +10,158 @@ namespace Audio
 		ClockName(InClockName)
 	{}
 
-	void FShareableQuartzCommandQueue::PushEvent(const FQuartzQuantizedCommandDelegateData& Data)
+} // namespace Audio
+
+// FQuartzTickableObject implementation
+void FQuartzTickableObject::FQuartzTickableObjectGCObjectMembers::AddReferencedObjects(FReferenceCollector& Collector)
+{
+	Collector.AddReferencedObject(QuartzSubsystem);
+	Collector.AddReferencedObject(WorldPtr);
+}
+
+FString FQuartzTickableObject::FQuartzTickableObjectGCObjectMembers::GetReferencerName() const
+{
+	return TEXT("FQuartzTickableObject::FQuartzTickableObjectGCObjectMembers");
+}
+
+FQuartzTickableObject::~FQuartzTickableObject()
+{
+	if (GCObjectMembers.WorldPtr && GCObjectMembers.QuartzSubsystem)
 	{
-		if (bIsAcceptingCommands)
-		{
-			EventDelegateQueue.Enqueue([InData = Data](FQuartzTickableObject* Handle) { Handle->ProcessCommand(InData); });
-		}
+		GCObjectMembers.QuartzSubsystem->UnsubscribeFromQuartzTick(this);
+		GCObjectMembers.WorldPtr = nullptr;
+		GCObjectMembers.QuartzSubsystem = nullptr;
+	}
+}
+
+FQuartzTickableObject* FQuartzTickableObject::Init(UWorld* InWorldPtr)
+{
+	if (!ensure(InWorldPtr))
+	{
+		// can't initialize if we don't have a valid world
+		return this;
 	}
 
-	void FShareableQuartzCommandQueue::PushEvent(const FQuartzMetronomeDelegateData& Data)
+	GCObjectMembers.WorldPtr = InWorldPtr;
+	GCObjectMembers.QuartzSubsystem = UQuartzSubsystem::Get(GCObjectMembers.WorldPtr);
+	CommandQueue = Audio::TQuartzShareableCommandQueue<FQuartzTickableObject>::Create();
+	GCObjectMembers.QuartzSubsystem->SubscribeToQuartzTick(this);
+
+	return this;
+}
+
+int32 FQuartzTickableObject::AddCommandDelegate(const FOnQuartzCommandEventBP& InDelegate, TArray<TSharedPtr<Audio::TQuartzShareableCommandQueue<FQuartzTickableObject>, ESPMode::ThreadSafe>>& TargetSubscriberArray)
+{
+	TargetSubscriberArray.Add(CommandQueue);
+
+	const int32 Num = QuantizedCommandDelegates.Num();
+	int32 SlotId = 0;
+
+	for (; SlotId < Num; ++SlotId)
 	{
-		if (bIsAcceptingCommands)
+		if (!QuantizedCommandDelegates[SlotId].MulticastDelegate.IsBound())
 		{
-			EventDelegateQueue.Enqueue([InData = Data](FQuartzTickableObject* Handle) { Handle->ProcessCommand(InData); });
+			QuantizedCommandDelegates[SlotId].MulticastDelegate.AddUnique(InDelegate);
+			return SlotId;
 		}
 	}
 
-	void FShareableQuartzCommandQueue::PushEvent(const FQuartzQueueCommandData& Data)
+	// need a new slot
+	QuantizedCommandDelegates.AddDefaulted_GetRef().MulticastDelegate.AddUnique(InDelegate);
+	return SlotId;
+}
+
+UQuartzSubsystem* FQuartzTickableObject::GetQuartzSubsystem() const
+{
+	if (ensureAlwaysMsgf(GCObjectMembers.IsValid(), TEXT("FQUartzTickableObject::GetQuartzSubsystem called when it has not been initialized")))
+	{
+		return GCObjectMembers.QuartzSubsystem;
+	}
+	
+	return nullptr;
+}
+
+void FQuartzTickableObject::ExecCommand(const Audio::FQuartzQuantizedCommandDelegateData& Data)
+{
+	checkSlow(Data.DelegateSubType < EQuartzCommandDelegateSubType::Count);
+	if (GCObjectMembers.QuartzSubsystem)
 	{
-		if (bIsAcceptingCommands)
+		GCObjectMembers.QuartzSubsystem->PushLatencyTrackerResult(Data.RequestRecieved());
+	}
+
+	// Broadcast to the BP delegate if we have one bound
+	if (Data.DelegateID >= 0 && Data.DelegateID < QuantizedCommandDelegates.Num()
+		&& ensure(QuantizedCommandDelegates[Data.DelegateID].MulticastDelegate.IsBound()))
+	{
+		FCommandDelegateGameThreadData& GameThreadEntry = QuantizedCommandDelegates[Data.DelegateID];
+
+		GameThreadEntry.MulticastDelegate.Broadcast(Data.DelegateSubType, "Sample Payload");
+
+		// track the number of active QuantizedCommands that may be sending info back to us.
+		// this is a bit of a hack because sound cues can play multiple wave instances
+		// and each of those wave instances is sending a delegate back to us.
+		// todo: clean this up at the wave-instance level to avoid ref counting here.
+		// (new command)
+		if (Data.DelegateSubType == EQuartzCommandDelegateSubType::CommandOnQueued)
+		{
+			GameThreadEntry.RefCount.Increment();
+		}
+
+		// (end of a command)
+		if (Data.DelegateSubType == EQuartzCommandDelegateSubType::CommandOnCanceled)
 		{
-			EventDelegateQueue.Enqueue([InData = Data](FQuartzTickableObject* Handle) { Handle->ProcessCommand(InData); });
+			// are all the commands done?
+			if (GameThreadEntry.RefCount.Decrement() == 0)
+			{
+				GameThreadEntry.MulticastDelegate.Clear();
+			}
 		}
 	}
 
+	// call base-class method
+	ProcessCommand(Data);
+}
 
-	void FShareableQuartzCommandQueue::StopTakingCommands()
+void FQuartzTickableObject::ExecCommand(const Audio::FQuartzMetronomeDelegateData& Data)
+{
+	if (GCObjectMembers.QuartzSubsystem)
 	{
-		bIsAcceptingCommands = false;
-		EventDelegateQueue.Empty();
+		GCObjectMembers.QuartzSubsystem->PushLatencyTrackerResult(Data.RequestRecieved());
 	}
 
-} // namespace Audio
\ No newline at end of file
+	MetronomeDelegates[static_cast<int32>(Data.Quantization)].MulticastDelegate.Broadcast(Data.ClockName, Data.Quantization, Data.Bar, Data.Beat, Data.BeatFraction);
+
+	// call base-class method
+	ProcessCommand(Data);
+}
+
+void FQuartzTickableObject::ExecCommand(const Audio::FQuartzQueueCommandData& Data)
+{
+	// call base-class method
+	ProcessCommand(Data);
+}
+
+TWeakPtr<Audio::TQuartzShareableCommandQueue<FQuartzTickableObject>, ESPMode::ThreadSafe> FQuartzTickableObject::GetCommandQueue()
+{
+	if (!CommandQueue.IsValid())
+	{
+		CommandQueue = MakeShared<Audio::TQuartzShareableCommandQueue<FQuartzTickableObject>, ESPMode::ThreadSafe>();
+	}
+
+	return CommandQueue;
+}
+
+void FQuartzTickableObject::QuartzTick(float DeltaTime)
+{
+	CommandQueue->PumpCommandQueue(this);
+}
+
+bool FQuartzTickableObject::QuartzIsTickable() const
+{
+	return CommandQueue.IsValid();
+}
+
+void FQuartzTickableObject::AddMetronomeBpDelegate(EQuartzCommandQuantization InQuantizationBoundary, const FOnQuartzMetronomeEventBP& OnQuantizationEvent)
+{
+	MetronomeDelegates[static_cast<int32>(InQuantizationBoundary)].MulticastDelegate.AddUnique(OnQuantizationEvent);
+}
diff --git a/Engine/Source/Runtime/Engine/Classes/Sound/QuartzSubscription.h b/Engine/Source/Runtime/Engine/Classes/Sound/QuartzSubscription.h
--- a/Engine/Source/Runtime/Engine/Classes/Sound/QuartzSubscription.h
+++ b/Engine/Source/Runtime/Engine/Classes/Sound/QuartzSubscription.h
@@ -2,17 +2,22 @@
 
 #pragma once
 
-#include "Sound/QuartzQuantizationUtilities.h"
+
+#include "Containers/MpscQueue.h"
 #include "UObject/ObjectMacros.h"
-#include "Containers/Queue.h"
-#include "DSP/VolumeFader.h"
+#include "UObject/GCObject.h"
+#include "Quartz/AudioMixerClock.h"
+#include "Sound/QuartzQuantizationUtilities.h"
+
 
 // forwards
-class FQuartzTickableObject;
 class UQuartzClockHandle;
+class UQuartzSubsystem;
 
 namespace Audio
 {
+	class FQuartzClockProxy;
+
 	// Struct used to communicate command state back to the game play thread
 	struct ENGINE_API FQuartzQuantizedCommandDelegateData : public FQuartzCrossThreadMessage
 	{
@@ -22,8 +27,7 @@ namespace Audio
 		// ID so the clock handle knows which delegate to fire
 		int32 DelegateID{ -1 };
 
-		// todo: more payload
-	};
+	}; // struct FQuartzQuantizedCommandDelegateData
 
 	// Struct used to communicate metronome events back to the game play thread
 	struct ENGINE_API FQuartzMetronomeDelegateData : public FQuartzCrossThreadMessage
@@ -32,7 +36,8 @@ namespace Audio
 		int32 Beat;
 		float BeatFraction;
 		EQuartzCommandQuantization Quantization;
-	};
+		FName ClockName;
+	}; // struct FQuartzMetronomeDelegateData
 
 	//Struct used to queue events to be sent to the Audio Render thread closer to their start time
 	struct ENGINE_API FQuartzQueueCommandData : public FQuartzCrossThreadMessage
@@ -41,32 +46,177 @@ namespace Audio
 		FName ClockName;
 
 		FQuartzQueueCommandData(const FAudioComponentCommandInfo& InAudioComponentCommandInfo, FName InClockName);
-	};
+	}; // struct FQuartzQueueCommandData
 
 
-	// Class that can be shared between Game Thread and any other thread to queue commands
-	// (ONLY the Game Thread may EXECUTE the commands in the queue, enforced by UQuartzClockHandle* argument in the TFunctions)
+	// old non-template\ version of TQuartzShareableCommandQueue
+	class UE_DEPRECATED(5.1, "Message") FShareableQuartzCommandQueue;
 	class ENGINE_API FShareableQuartzCommandQueue
 	{
+	};
+
+	/**
+	*	Template class for mono-directional MPSC command queues 
+	* 
+	*   in order to enforce thread-safe access to the object executing the commands,
+	*	"listener type" is the type of the object that is being accessed in the commands
+	*	that object will have to provide a 'this' ptr (of type ListenerType) in order to 
+	*   invoke the commands on itself. (The lambdas do NOT and should NOT cache a ptr or
+	*	reference to the target).
+	* 
+	*	User-provided lambdas can take any (single) argument type T in PushEvent()
+	*	but there must exist a ListenerType::ExecCommand(T) overload for any PushEvent(T) instantiated. 
+	* 
+	*	(see FQuartzTickableObject and FQuartzClock as examples)
+	* 
+	**/
+	template <class ListenerType>
+	class ENGINE_API TQuartzShareableCommandQueue
+	{
 	public:
-		// (add more PushEvent overloads for other event types)
-		void PushEvent(const FQuartzQuantizedCommandDelegateData& Data);
-		void PushEvent(const FQuartzMetronomeDelegateData& Data);
-		void PushEvent(const FQuartzQueueCommandData& Data);
 
-		bool IsQueueEmpty()
+		// ctor
+		TQuartzShareableCommandQueue() {}
+
+		// dtor
+		~TQuartzShareableCommandQueue() {}
+
+		// static helper to create a new sharable queue
+		static TSharedPtr<TQuartzShareableCommandQueue<ListenerType>, ESPMode::ThreadSafe> Create()
 		{
-			return EventDelegateQueue.IsEmpty();
+			return MakeShared<TQuartzShareableCommandQueue<ListenerType>, ESPMode::ThreadSafe>();
 		}
 
+		// note: ListenerType must have a ExecCommand() overload for each instantiation of this method
+		template <typename T>
+		void PushEvent(const T& Data)
+		{
+			CommandQueue.Enqueue([InData = Data](ListenerType* Listener) { Listener->ExecCommand(InData); });
+		}
 
-		// called when the game object owner is shutting down
-		void StopTakingCommands();
+		void PushCommand(TFunction<void(ListenerType*)> InCommand)
+		{
+			CommandQueue.Enqueue([=](ListenerType* Listener) { InCommand(Listener); });
+		}
 
-	private:
-		TQueue<TFunction<void(FQuartzTickableObject*)>, EQueueMode::Mpsc> EventDelegateQueue;
-		FThreadSafeBool bIsAcceptingCommands{ true };
+		void PumpCommandQueue(ListenerType* InListener)
+		{
+			// gather move all the current commands into our temp container
+			// (in case the commands themselves alter the container)
+			TOptional<TFunction<void(ListenerType*)>> Function = CommandQueue.Dequeue();
+			while (Function.IsSet())
+			{
+				TempCommandQueue.Emplace(Function.GetValue());
+				Function = CommandQueue.Dequeue();
+			}
+
+			for (auto& Command : TempCommandQueue)
+			{
+				Command(InListener);
+			}
+
+			TempCommandQueue.Reset();
+		}
 
-		friend class ::FQuartzTickableObject; // UQuartzTickableObject, and those that inherit from it, are the only classes allowed to pump the command queue
+	private:
+		TMpscQueue<TFunction<void(ListenerType*)>> CommandQueue;
+		TArray<TFunction<void(ListenerType*)>> TempCommandQueue;
 	};
+
 } // namespace Audio
+
+
+	/**
+	 *	FQuartzTickableObject
+	 *
+	 *		This is the base class for non-Audio Render Thread objects that want to receive
+	 *		callbacks for Quartz events.
+	 *
+	 *		It is a wrapper around TQuartzShareableCommandQueue.
+	 *		(see UQuartzClockHandle or UAudioComponent as implementation examples)
+	 */
+	class ENGINE_API FQuartzTickableObject
+	{
+		struct ENGINE_API FQuartzTickableObjectGCObjectMembers : public FGCObject
+		{
+			UQuartzSubsystem* QuartzSubsystem{ nullptr };
+
+			UWorld* WorldPtr{ nullptr };
+
+			// Begin FGCObject Interface
+			virtual void AddReferencedObjects(FReferenceCollector& Collector) override;
+			virtual FString GetReferencerName() const override;
+			// End FGCObject Interface
+
+			bool IsValid() const { return QuartzSubsystem && WorldPtr; }
+		}; // struct FQuartzTickableObjectGCObjectMembers
+
+	public:
+		// ctor
+		FQuartzTickableObject() {}
+
+		// dtor
+		virtual ~FQuartzTickableObject();
+
+		FQuartzTickableObject* Init(UWorld* InWorldPtr);
+
+		// called by the associated QuartzSubsystem
+		void QuartzTick(float DeltaTime);
+		
+		bool QuartzIsTickable() const;
+		
+		// maxtodo: deprecate, derived classes should have their own world
+		UWorld* QuartzGetWorld() const { return GCObjectMembers.WorldPtr; }
+
+		void AddMetronomeBpDelegate(EQuartzCommandQuantization InQuantizationBoundary, const FOnQuartzMetronomeEventBP& OnQuantizationEvent);
+
+		bool IsInitialized() const { return GCObjectMembers.IsValid(); }
+
+		TWeakPtr<Audio::TQuartzShareableCommandQueue<FQuartzTickableObject>, ESPMode::ThreadSafe> GetCommandQueue();
+
+		int32 AddCommandDelegate(const FOnQuartzCommandEventBP& InDelegate, TArray<TSharedPtr<Audio::TQuartzShareableCommandQueue<FQuartzTickableObject>, ESPMode::ThreadSafe>>& TargetSubscriberArray);
+
+		// maxtodo: deprecate (use the static getter on the subsystem class)
+		UQuartzSubsystem* GetQuartzSubsystem() const;
+
+		// required by TQuartzShareableCommandQueue template
+		void ExecCommand(const Audio::FQuartzQuantizedCommandDelegateData& Data);
+		void ExecCommand(const Audio::FQuartzMetronomeDelegateData& Data);
+		void ExecCommand(const Audio::FQuartzQueueCommandData& Data);
+
+		// virtual interface (ExecCommand will forward the data to derived classes' ProcessCommand() call)
+		virtual void ProcessCommand(const Audio::FQuartzQuantizedCommandDelegateData& Data) {}
+		virtual void ProcessCommand(const Audio::FQuartzMetronomeDelegateData& Data) {}
+		virtual void ProcessCommand(const Audio::FQuartzQueueCommandData& Data) {}
+
+	private:
+		struct FMetronomeDelegateGameThreadData
+		{
+			FOnQuartzMetronomeEvent MulticastDelegate;
+		};
+
+		struct FCommandDelegateGameThreadData
+		{
+			FOnQuartzCommandEvent MulticastDelegate;
+			FThreadSafeCounter RefCount;
+		};
+
+		// delegate containers
+		FMetronomeDelegateGameThreadData MetronomeDelegates[static_cast<int32>(EQuartzCommandQuantization::Count)];
+		TArray<FCommandDelegateGameThreadData> QuantizedCommandDelegates;
+
+		// command queue
+		TSharedPtr<Audio::TQuartzShareableCommandQueue<FQuartzTickableObject>, ESPMode::ThreadSafe> CommandQueue;
+		TArray<TFunction<void(FQuartzTickableObject*)>> TempCommandQueue;
+
+	public:
+		// deprecate public access
+		UE_DEPRECATED(5.1, "FQuartzTickableObject::CommandQueuePtr should no longer be accessed directly. (This TSharedPtr will always be invalid).  Use Use GetCommandQueue().Pin() instead")
+		TSharedPtr<Audio::TQuartzShareableCommandQueue<FQuartzTickableObject>, ESPMode::ThreadSafe> CommandQueuePtr;
+
+	private:
+		// gc object members struct (USubsystem / world )
+		FQuartzTickableObjectGCObjectMembers GCObjectMembers;
+
+	}; // class FQuartzTickableObject
+
diff --git a/Engine/Source/Runtime/Engine/Private/Components/AudioComponent.cpp b/Engine/Source/Runtime/Engine/Private/Components/AudioComponent.cpp
--- a/Engine/Source/Runtime/Engine/Private/Components/AudioComponent.cpp
+++ b/Engine/Source/Runtime/Engine/Private/Components/AudioComponent.cpp
@@ -284,13 +284,11 @@ void UAudioComponent::OnUnregister()
 	{
 		Stop();
 	}
-	Shutdown();
 }
 
 void UAudioComponent::EndPlay(const EEndPlayReason::Type EndPlayReason)
 {
 	Super::EndPlay(EndPlayReason);
-	Shutdown();
 }
 
 const UObject* UAudioComponent::AdditionalStatObject() const
@@ -448,13 +446,7 @@ void UAudioComponent::ProcessCommand(const Audio::FQuartzQuantizedCommandDelegat
 
 void UAudioComponent::ProcessCommand(const Audio::FQuartzQueueCommandData& InQueueCommandData)
 {
-	if (UQuartzSubsystem* QuartzSubsystem = GetQuartzSubsystem())
-	{
-		QuartzSubsystem->PushLatencyTrackerResult(InQueueCommandData.RequestRecieved());
-
-		//Queue the sound
-		PlayQueuedQuantizedInternal(GetWorld(), InQueueCommandData.AudioComponentCommandInfo);
-	}
+	PlayQueuedQuantizedInternal(GetWorld(), InQueueCommandData.AudioComponentCommandInfo);
 }
 
 void UAudioComponent::PlayQuantized(
@@ -472,7 +464,7 @@ void UAudioComponent::PlayQuantized(
 	{
 		Init(GetWorld());
 	}
-	check(FQuartzTickableObject::IsInitialized() && CommandQueuePtr.IsValid());
+	check(FQuartzTickableObject::IsInitialized());
 
 	int32 TimeCVarVal = FMath::Clamp(CVarTimeToTakeUpVoiceSlot->GetInt(), 0, (int32)EQuartzCommandQuantization::Count - 1);
 	EQuartzCommandQuantization MinimumQuantization = EQuartzCommandQuantization(TimeCVarVal);
@@ -480,7 +472,7 @@ void UAudioComponent::PlayQuantized(
 	// Make an anticapatory quantization boundary to try to avoid taking up a whole voice slot while waiting for a queued event
 	FQuartzQuantizationBoundary AnticipationQuantizationBoundary = FQuartzQuantizationBoundary(MinimumQuantization, 1.0f, EQuarztQuantizationReference::CurrentTimeRelative, true);
 
-	FAudioComponentCommandInfo NewComponentCommandInfo(CommandQueuePtr, AnticipationQuantizationBoundary);
+	FAudioComponentCommandInfo NewComponentCommandInfo(FQuartzTickableObject::GetCommandQueue().Pin(), AnticipationQuantizationBoundary);
 
 	// And a new pending quartz command data
 	FAudioComponentPendingQuartzCommandData AudioComponentQuartzCommandData
@@ -534,12 +526,8 @@ void UAudioComponent::PlayQuantized(
 
 void UAudioComponent::PlayQueuedQuantizedInternal(const UObject* WorldContextObject, FAudioComponentCommandInfo InCommandInfo)
 {
-	//Initialize the tickable object portion of the Audio Component, if it hasn't been initialized already
-	if (!FQuartzTickableObject::IsInitialized())
-	{
-		Init(GetWorld());
-	}
-	check(FQuartzTickableObject::IsInitialized() && CommandQueuePtr.IsValid());
+	// confirm the FQuartzTickableObject has been initialized.
+	ensure(FQuartzTickableObject::IsInitialized());
 
 	bool bFoundQuantizedCommand = false;
 	bool bIsValidCommand = true;
@@ -569,7 +557,7 @@ void UAudioComponent::PlayQueuedQuantizedInternal(const UObject* WorldContextObj
 				// confirm a valid handle
 				if (Handle != nullptr)
 				{
-					InternalRequestData.QuantizedRequestData = Handle->GetQuartzSubsystem()->CreateDataDataForSchedulePlaySound(Handle, PendingData.Delegate, PendingData.AnticapatoryBoundary);
+					InternalRequestData.QuantizedRequestData = UQuartzSubsystem::CreateRequestDataForSchedulePlaySound(Handle, PendingData.Delegate, PendingData.AnticapatoryBoundary);
 					UGameplayStatics::PrimeSound(Sound);
 				}
 
@@ -596,7 +584,7 @@ void UAudioComponent::PlayQueuedQuantizedInternal(const UObject* WorldContextObj
 
 				if (bIsValidCommand)
 				{
-					InternalRequestData.QuantizedRequestData.GameThreadSubscribers.Add(CommandQueuePtr);
+					InternalRequestData.QuantizedRequestData.GameThreadSubscribers.Add(GetCommandQueue().Pin());
 
 					// Now play the quartz command
 					PlayInternal(InternalRequestData);
@@ -606,7 +594,6 @@ void UAudioComponent::PlayQueuedQuantizedInternal(const UObject* WorldContextObj
 			{
 				UE_LOG(LogAudioQuartz, Warning, TEXT("Attempting to play Quantized Sound without supplying a valid Sound to play"));
 			}
-			
 
 
 
