Engine fixes for removing pendingkill.

FQuartzTickableObject -> World, QuartzSubsystem
FNiagaraWorldManager -> World (via Outer of collection instance)
SoundLibraryContext -> Actor
SoundLibraryAnimContext -> SkeletalMeshComponent

diff --git a/Engine/Plugins/FX/Niagara/Source/Niagara/Private/NiagaraWorldManager.cpp b/Engine/Plugins/FX/Niagara/Source/Niagara/Private/NiagaraWorldManager.cpp
--- a/Engine/Plugins/FX/Niagara/Source/Niagara/Private/NiagaraWorldManager.cpp
+++ b/Engine/Plugins/FX/Niagara/Source/Niagara/Private/NiagaraWorldManager.cpp
@@ -473,7 +473,7 @@ UNiagaraCullProxyComponent* FNiagaraWorldManager::GetCullProxy(UNiagaraComponent
 
 UNiagaraParameterCollectionInstance* FNiagaraWorldManager::GetParameterCollection(UNiagaraParameterCollection* Collection)
 {
-	if (!Collection)
+	if (!Collection || bIsTearingDown)
 	{
 		return nullptr;
 	}
@@ -727,6 +727,11 @@ void FNiagaraWorldManager::OnPreWorldFinishDestroy(UWorld* World)
 
 void FNiagaraWorldManager::OnWorldBeginTearDown(UWorld* World)
 {
+	FNiagaraWorldManager** Manager = WorldManagers.Find(World);
+	if (Manager)
+	{
+		(*Manager)->bIsTearingDown = true;
+	}
 // 	FNiagaraWorldManager** Manager = WorldManagers.Find(World);
 // 	if (Manager)
 // 	{
diff --git a/Engine/Plugins/FX/Niagara/Source/Niagara/Public/NiagaraWorldManager.h b/Engine/Plugins/FX/Niagara/Source/Niagara/Public/NiagaraWorldManager.h
--- a/Engine/Plugins/FX/Niagara/Source/Niagara/Public/NiagaraWorldManager.h
+++ b/Engine/Plugins/FX/Niagara/Source/Niagara/Public/NiagaraWorldManager.h
@@ -313,6 +313,8 @@ private:
 	/** True if the app has focus. We prevent some culling if the app doesn't have focus as it can interfere. */
 	bool bAppHasFocus;
 
+	bool bIsTearingDown = false;
+
 	float WorldLoopTime = 0.0f;
 	
 	ENiagaraDebugPlaybackMode RequestedDebugPlaybackMode = ENiagaraDebugPlaybackMode::Play;
diff --git a/Engine/Source/Runtime/AudioMixer/Public/Quartz/AudioMixerClockHandle.h b/Engine/Source/Runtime/AudioMixer/Public/Quartz/AudioMixerClockHandle.h
--- a/Engine/Source/Runtime/AudioMixer/Public/Quartz/AudioMixerClockHandle.h
+++ b/Engine/Source/Runtime/AudioMixer/Public/Quartz/AudioMixerClockHandle.h
@@ -15,20 +15,17 @@
 
 class UAudioComponent;
 
-
+// Children of this class MUST be UObjects and they must call QuartzAddReferencedObjects from their own override of AddReferencedObjects.
+// TODO: Ideally this should use CRTP to put the fields on the child or something to avoid the possibility of not calling it.
 class AUDIOMIXER_API FQuartzTickableObject
 {
-	struct AUDIOMIXER_API FQuartzTickableObjectGCObjectMembers : public FGCObject
+	struct AUDIOMIXER_API FQuartzTickableObjectGCObjectMembers
 	{
 	public:
 		UQuartzSubsystem* QuartzSubsystem;
 		UWorld* WorldPtr{ nullptr };
 
-		void AddReferencedObjects(FReferenceCollector& Collector) override;
-		virtual FString GetReferencerName() const override
-		{
-			return TEXT("FQuartzTickableObject::FQuartzTickableObjectGCObjectMembers");
-		}
+		void AddReferencedObjects(FReferenceCollector& Collector);
 	};
 
 	public:
@@ -60,6 +57,10 @@ class AUDIOMIXER_API FQuartzTickableObject
 
 		void Shutdown();
 
+		void QuartzAddReferencedObjects(FReferenceCollector& Collector)
+		{
+			GCObjectMembers.AddReferencedObjects(Collector);
+		}
 
 	protected:
 		struct CommandDelegateGameThreadData
@@ -219,6 +220,11 @@ public:
 
 	bool GetCurrentTickRate(const UObject* WorldContextObject, Audio::FQuartzClockTickRate& OutTickRate) const;
 
+	static void AddReferencedObjects(UObject* InThis, FReferenceCollector& Collector) 
+	{
+		CastChecked<UQuartzClockHandle>(InThis)->QuartzAddReferencedObjects(Collector);
+	}
+
 private:
 	
 	FName CurrentClockId;
diff --git a/Engine/Source/Runtime/Engine/Classes/Components/AudioComponent.h b/Engine/Source/Runtime/Engine/Classes/Components/AudioComponent.h
--- a/Engine/Source/Runtime/Engine/Classes/Components/AudioComponent.h
+++ b/Engine/Source/Runtime/Engine/Classes/Components/AudioComponent.h
@@ -748,6 +748,11 @@ public:
 	virtual void BeginDestroy() override;
 	//~ End UObject Interface.
 
+	static void AddReferencedObjects(UObject* InThis, FReferenceCollector& Collector)
+	{
+		CastChecked<UAudioComponent>(InThis)->QuartzAddReferencedObjects(Collector);
+	}
+
 	//~ Begin USceneComponent Interface
 	virtual void Activate(bool bReset=false) override;
 	virtual void Deactivate() override;
