Quartz work to reduce uobject dependency, clean up implicit knowledge of FQuartzTickableObject base class, and the addition of FQuartzClockHandles to access clocks directly.

#p4v-preflight-copy 20084336

diff --git a/Engine/Source/Runtime/AudioMixer/Private/Quartz/AudioMixerClock.cpp b/Engine/Source/Runtime/AudioMixer/Private/Quartz/AudioMixerClock.cpp
--- a/Engine/Source/Runtime/AudioMixer/Private/Quartz/AudioMixerClock.cpp
+++ b/Engine/Source/Runtime/AudioMixer/Private/Quartz/AudioMixerClock.cpp
@@ -3,6 +3,7 @@
 #include "Quartz/AudioMixerClock.h"
 #include "Quartz/AudioMixerClockManager.h"
 #include "AudioMixerSourceManager.h"
+#include "Sound/QuartzSubscription.h"
 
 
 static float HeadlessClockSampleRateCvar = 100000.f;
@@ -15,14 +16,100 @@ FAutoConsoleVariableRef CVarHeadlessClockSampleRate(
 
 namespace Audio
 {
+	// FQuartzClockProxy Implementation
+	// ctor
+	FQuartzClockProxy::FQuartzClockProxy(TSharedPtr<FQuartzClock, ESPMode::ThreadSafe> InClock)
+		: ClockId(InClock->GetName())
+		, SharedQueue(InClock->GetCommandQueue())
+		, ClockWeakPtr(InClock)
+	{
+	}
+
+	bool FQuartzClockProxy::IsValid() const
+	{
+		return SharedQueue.Pin().IsValid();
+	}
+
+	bool FQuartzClockProxy::DoesClockExist() const
+	{
+		return IsValid();
+	}
+
+	bool FQuartzClockProxy::IsClockRunning() const
+	{
+		TSharedPtr<FQuartzClock, ESPMode::ThreadSafe> ClockPtr = ClockWeakPtr.Pin();
+		if (!ClockPtr)
+		{
+			return false;
+		}
+
+		return ClockPtr->IsRunning();
+	}
+
+	float FQuartzClockProxy::GetDurationOfQuantizationTypeInSeconds(const EQuartzCommandQuantization& QuantizationType, float Multiplier) const
+	{
+		TSharedPtr<FQuartzClock, ESPMode::ThreadSafe> ClockPtr = ClockWeakPtr.Pin();
+		if (!ClockPtr)
+		{
+			return 0.f;
+		}
+
+		return ClockPtr->GetDurationOfQuantizationTypeInSeconds(QuantizationType, Multiplier);
+	}
+
+	Audio::FQuartzClockTickRate FQuartzClockProxy::GetTickRate() const
+	{
+		TSharedPtr<FQuartzClock, ESPMode::ThreadSafe> ClockPtr = ClockWeakPtr.Pin();
+		if (!ClockPtr)
+		{
+			return {};
+		}
+
+		return ClockPtr->GetTickRate();
+	}
+
+	FQuartzTransportTimeStamp FQuartzClockProxy::GetCurrentClockTimestamp() const
+	{
+		TSharedPtr<FQuartzClock, ESPMode::ThreadSafe> ClockPtr = ClockWeakPtr.Pin();
+		if (!ClockPtr)
+		{
+			return {};
+		}
+
+		return ClockPtr->GetCurrentTimestamp();
+	}
+
+	float FQuartzClockProxy::GetEstimatedClockRunTimeSeconds() const
+	{
+		TSharedPtr<FQuartzClock, ESPMode::ThreadSafe> ClockPtr = ClockWeakPtr.Pin();
+		if (!ClockPtr)
+		{
+			return 0.f;
+		}
+
+		return ClockPtr->GetEstimatedRunTime();
+	}
+
+
+	bool FQuartzClockProxy::SendCommandToClock(TFunction<void(FQuartzClock*)> InCommand)
+	{
+		if (auto QueuePtr = SharedQueue.Pin())
+		{
+			QueuePtr->PushCommand(InCommand);
+			return true;
+		}
+
+		return false;
+	}
+
+	// FQuartzClock Implementation
 	FQuartzClock::FQuartzClock(const FName& InName, const FQuartzClockSettings& InClockSettings, FQuartzClockManager* InOwningClockManagerPtr)
-		: Metronome(InClockSettings.TimeSignature)
+		: Metronome(InClockSettings.TimeSignature, InName)
 		, OwningClockManagerPtr(InOwningClockManagerPtr)
 		, Name(InName)
 		, bIsRunning(false)
 		, bIgnoresFlush(InClockSettings.bIgnoreLevelChange)
 	{
-
 		FMixerDevice* MixerDevice = GetMixerDevice();
 
 		if (MixerDevice)
@@ -62,18 +149,23 @@ namespace Audio
 		// adjust time-till-fire for existing commands
 		for (auto& Command : PendingCommands)
 		{
-			Command.NumFramesUntilExec = NumFramesLeft + Ratio * (Command.NumFramesUntilExec - NumFramesLeft);
+			if(!Command.Command->ShouldDeadlineIgnoresBpmChanges())
+			{
+				Command.NumFramesUntilExec = NumFramesLeft + Ratio * (Command.NumFramesUntilExec - NumFramesLeft);
+			}
 		}
 
 		for (auto& Command : ClockAlteringPendingCommands)
 		{
-			Command.NumFramesUntilExec = NumFramesLeft + Ratio * (Command.NumFramesUntilExec - NumFramesLeft);
+			if(!Command.Command->ShouldDeadlineIgnoresBpmChanges())
+			{
+				Command.NumFramesUntilExec = NumFramesLeft + Ratio * (Command.NumFramesUntilExec - NumFramesLeft);
+			}
 		}
 	}
 
 	void FQuartzClock::ChangeTimeSignature(const FQuartzTimeSignature& InNewTimeSignature)
 	{
-		// TODO: what does this do to pending events waiting for the beat? (maybe nothing is reasonable?)
 		Metronome.SetTimeSignature(InNewTimeSignature);
 	}
 
@@ -165,6 +257,7 @@ namespace Audio
 	void FQuartzClock::LowResolutionTick(float InDeltaTimeSeconds)
 	{
 		TRACE_CPUPROFILER_EVENT_SCOPE(QuartzClock::Tick_LowRes);
+		PreTickCommands->PumpCommandQueue(this);
 		Tick(static_cast<int32>(InDeltaTimeSeconds * Metronome.GetTickRate().GetSampleRate()));
 	}
 
@@ -172,6 +265,10 @@ namespace Audio
 	{
 		TRACE_CPUPROFILER_EVENT_SCOPE(QuartzClock::Tick);
 
+		TRACE_CPUPROFILER_EVENT_SCOPE(QuartzClock::GameThreadCommands);
+		PreTickCommands->PumpCommandQueue(this);
+
+		TRACE_CPUPROFILER_EVENT_SCOPE(QuartzClock::TickInternal);
 		if (!bIsRunning)
 		{
 			return;
@@ -199,8 +296,19 @@ namespace Audio
 			TickInternal(InNumFramesUntilNextTick - TickDelayLengthInFrames, PendingCommands, FramesOfLatency, TickDelayLengthInFrames);
 		}
 
-
 		Metronome.Tick(InNumFramesUntilNextTick, FramesOfLatency);
+
+		UpdateCachedState();
+	}
+
+	FQuartzClockCommandQueueWeakPtr FQuartzClock::GetCommandQueue() const
+	{
+		if (!PreTickCommands.IsValid())
+		{
+			PreTickCommands = TQuartzShareableCommandQueue<FQuartzClock>::Create();
+		}
+
+		return PreTickCommands;
 	}
 
 	void FQuartzClock::TickInternal(int32 InNumFramesUntilNextTick, TArray<PendingCommand>& CommandsToTick, int32 FramesOfLatency, int32 FramesOfDelay)
@@ -222,10 +330,12 @@ namespace Audio
 				PendingCommand.Command->OnFinalCallback(PendingCommand.NumFramesUntilExec + FramesOfDelay);
 				PendingCommand.Command.Reset();
 				bHaveCommandsToRemove = true;
+
 			}
 			else // not yet executing
 			{
 				PendingCommand.NumFramesUntilExec -= InNumFramesUntilNextTick;
+				PendingCommand.Command->Update(PendingCommand.NumFramesUntilExec);
 			}
 		}
 
@@ -242,6 +352,15 @@ namespace Audio
 		}
 	}
 
+	void FQuartzClock::UpdateCachedState()
+	{
+		FScopeLock ScopeLock(&CachedClockStateCritSec);
+
+		CachedClockState.TickRate = Metronome.GetTickRate();
+		CachedClockState.TimeStamp = Metronome.GetTimeStamp();
+		CachedClockState.RunTimeInSeconds = (float)Metronome.GetTimeSinceStart();
+	}
+
 	void FQuartzClock::SetSampleRate(float InNewSampleRate)
 	{
 		if (FMath::IsNearlyEqual(InNewSampleRate, Metronome.GetTickRate().GetSampleRate()))
@@ -251,11 +370,9 @@ namespace Audio
 
 		// update Tick Rate
 		Metronome.SetSampleRate(InNewSampleRate);
-
-		// TODO: update the deadlines of all our new events
 	}
 
-	bool FQuartzClock::IgnoresFlush()
+	bool FQuartzClock::IgnoresFlush() const
 	{
 		return bIgnoresFlush;
 	}
@@ -265,24 +382,24 @@ namespace Audio
 		return Metronome.GetTimeSignature() == InClockSettings.TimeSignature;
 	}
 
-	void FQuartzClock::SubscribeToTimeDivision(MetronomeCommandQueuePtr InListenerQueue, EQuartzCommandQuantization InQuantizationBoundary)
+	void FQuartzClock::SubscribeToTimeDivision(FQuartzGameThreadSubscriber InSubscriber, EQuartzCommandQuantization InQuantizationBoundary)
 	{
-		Metronome.SubscribeToTimeDivision(InListenerQueue, InQuantizationBoundary);
+		Metronome.SubscribeToTimeDivision(InSubscriber, InQuantizationBoundary);
 	}
 
-	void FQuartzClock::SubscribeToAllTimeDivisions(MetronomeCommandQueuePtr InListenerQueue)
+	void FQuartzClock::SubscribeToAllTimeDivisions(FQuartzGameThreadSubscriber InSubscriber)
 	{
-		Metronome.SubscribeToAllTimeDivisions(InListenerQueue);
+		Metronome.SubscribeToAllTimeDivisions(InSubscriber);
 	}
 
-	void FQuartzClock::UnsubscribeFromTimeDivision(MetronomeCommandQueuePtr InListenerQueue, EQuartzCommandQuantization InQuantizationBoundary)
+	void FQuartzClock::UnsubscribeFromTimeDivision(FQuartzGameThreadSubscriber InSubscriber, EQuartzCommandQuantization InQuantizationBoundary)
 	{
-		Metronome.UnsubscribeFromTimeDivision(InListenerQueue, InQuantizationBoundary);
+		Metronome.UnsubscribeFromTimeDivision(InSubscriber, InQuantizationBoundary);
 	}
 
-	void FQuartzClock::UnsubscribeFromAllTimeDivisions(MetronomeCommandQueuePtr InListenerQueue)
+	void FQuartzClock::UnsubscribeFromAllTimeDivisions(FQuartzGameThreadSubscriber InSubscriber)
 	{
-		Metronome.UnsubscribeFromAllTimeDivisions(InListenerQueue);
+		Metronome.UnsubscribeFromAllTimeDivisions(InSubscriber);
 	}
 
 	void FQuartzClock::AddQuantizedCommand(FQuartzQuantizationBoundary InQuantizationBondary, TSharedPtr<IQuartzQuantizedCommand> InNewEvent)
@@ -309,10 +426,9 @@ namespace Audio
 			Resume();
 		}
 
-
 		int32 FramesUntilExec = 0;
 
-		// if this is unquantized, execute immediately (even if the clock is paused)
+		// if this is un-quantized, execute immediately (even if the clock is paused)
 		if (InQuantizationBondary.Quantization == EQuartzCommandQuantization::None)
 		{
 			InNewEvent->AboutToStart();
@@ -324,6 +440,10 @@ namespace Audio
 		FramesUntilExec = FMath::RoundToInt(Metronome.GetFramesUntilBoundary(InQuantizationBondary)); // query metronome (round result to int)
 		FramesUntilExec = FMath::Max(0, InNewEvent->OverrideFramesUntilExec(FramesUntilExec)); // allow command to override the deadline (clamp result)
 
+		// finalize the requested subscriber offsets and notify the command of their deadline
+		InNewEvent->OnScheduled(Metronome.GetTickRate());
+		InNewEvent->Update(FramesUntilExec);
+
 		// if this is going to execute on the next tick, warn Game Thread Subscribers as soon as possible
 		if (FramesUntilExec == 0)
 		{
@@ -362,26 +482,26 @@ namespace Audio
 		return PendingCommands.Num() + ClockAlteringPendingCommands.Num();
 	}
 
-	bool FQuartzClock::IsRunning()
+	bool FQuartzClock::IsRunning() const
 	{
 		return bIsRunning;
 	}
 
 	float FQuartzClock::GetDurationOfQuantizationTypeInSeconds(const EQuartzCommandQuantization& QuantizationType, float Multiplier)
 	{
+		FScopeLock ScopeLock(&CachedClockStateCritSec);
+
 		// if this is unquantized, return 0
 		if (QuantizationType == EQuartzCommandQuantization::None)
 		{
 			return 0;
 		}
 
-		FQuartzClockTickRate TickRate = Metronome.GetTickRate();
-
 		// get number of frames until the relevant quantization event
-		double FramesUntilExec = TickRate.GetFramesPerDuration(QuantizationType);
+		double FramesUntilExec = CachedClockState.TickRate.GetFramesPerDuration(QuantizationType);
 
 		//Translate frames to seconds
-		double SampleRate = TickRate.GetSampleRate();
+		double SampleRate = CachedClockState.TickRate.GetSampleRate();
 
 		if (!FMath::IsNearlyZero(SampleRate))
 		{
@@ -395,14 +515,14 @@ namespace Audio
 
 	FQuartzTransportTimeStamp FQuartzClock::GetCurrentTimestamp()
 	{
-		FQuartzTransportTimeStamp CurrentTimeStamp = Metronome.GetTimeStamp();
-
-		return CurrentTimeStamp;
+		FScopeLock ScopeLock(&CachedClockStateCritSec);
+		return CachedClockState.TimeStamp;
 	}
 
 	float FQuartzClock::GetEstimatedRunTime()
 	{
-		return (float)Metronome.GetTimeSinceStart();
+		FScopeLock ScopeLock(&CachedClockStateCritSec);
+		return CachedClockState.RunTimeInSeconds;
 	}
 
 	FMixerDevice* FQuartzClock::GetMixerDevice()
@@ -416,6 +536,34 @@ namespace Audio
 		return nullptr;
 	}
 
+	void FQuartzClock::AddQuantizedCommand(FQuartzQuantizedCommandInitInfo& InQuantizationCommandInitInfo)
+	{
+		if (!ensure(InQuantizationCommandInitInfo.QuantizedCommandPtr))
+		{
+			return;
+		}
+
+		// this method can't be utilized by play commands because the AudioMixerSource needs a handle in order to stop it.
+		// PlayCommands must be queued via the clock manager in AudioMixerSourceManager.
+		if (!ensure(EQuartzCommandType::PlaySound != InQuantizationCommandInitInfo.QuantizedCommandPtr->GetCommandType()))
+		{
+			return;
+		}
+
+		// Can this command run without an Audio Device?
+		FMixerDevice* MixerDevice = GetMixerDevice();
+		if (!MixerDevice && InQuantizationCommandInitInfo.QuantizedCommandPtr->RequiresAudioDevice())
+		{
+			InQuantizationCommandInitInfo.QuantizedCommandPtr->Cancel();
+		}
+
+		// this function is a friend of FQuartzClockManager, so we can use FindClock() directly
+		// to access the shared ptr to "this"
+		InQuantizationCommandInitInfo.SetOwningClockPtr(GetClockManager()->FindClock(GetName()));
+		InQuantizationCommandInitInfo.QuantizedCommandPtr->OnQueued(InQuantizationCommandInitInfo);
+		AddQuantizedCommand(InQuantizationCommandInitInfo.QuantizationBoundary, InQuantizationCommandInitInfo.QuantizedCommandPtr);
+	}
+
 	FMixerSourceManager* FQuartzClock::GetSourceManager()
 	{
 		FMixerDevice* MixerDevice = GetMixerDevice();
@@ -429,6 +577,17 @@ namespace Audio
 		return nullptr;
 	}
 
+	FQuartzClockTickRate FQuartzClock::GetTickRate()
+	{
+		FScopeLock ScopeLock(&CachedClockStateCritSec);
+		return CachedClockState.TickRate;
+	}
+
+	FName FQuartzClock::GetName() const
+	{
+		return Name;
+	}
+
 	FQuartzClockManager* FQuartzClock::GetClockManager()
 	{
 		checkSlow(OwningClockManagerPtr);
@@ -444,6 +603,16 @@ namespace Audio
 		Metronome.ResetTransport();
 	}
 
+	void FQuartzClock::AddToTickDelay(int32 NumFramesOfDelayToAdd)
+	{
+		TickDelayLengthInFrames += NumFramesOfDelayToAdd;
+	}
+
+	void FQuartzClock::SetTickDelay(int32 NumFramesOfDelay)
+	{
+		TickDelayLengthInFrames = NumFramesOfDelay;
+	}
+
 	bool FQuartzClock::CancelQuantizedCommandInternal(TSharedPtr<IQuartzQuantizedCommand> InCommandPtr, TArray<PendingCommand>& CommandsToTick)
 	{
 		for (int32 i = CommandsToTick.Num() - 1; i >= 0; --i)
diff --git a/Engine/Source/Runtime/AudioMixer/Private/Quartz/AudioMixerClockHandle.cpp b/Engine/Source/Runtime/AudioMixer/Private/Quartz/AudioMixerClockHandle.cpp
--- a/Engine/Source/Runtime/AudioMixer/Private/Quartz/AudioMixerClockHandle.cpp
+++ b/Engine/Source/Runtime/AudioMixer/Private/Quartz/AudioMixerClockHandle.cpp
@@ -6,119 +6,6 @@
 #include "AudioMixerDevice.h"
 #include "Engine/GameInstance.h"
 
-// FQuartzTickableObject implementation
-void FQuartzTickableObject::FQuartzTickableObjectGCObjectMembers::AddReferencedObjects(FReferenceCollector& Collector)
-{
-	Collector.AddReferencedObject(QuartzSubsystem);
-	Collector.AddReferencedObject(WorldPtr);
-}
-
-FQuartzTickableObject::~FQuartzTickableObject()
-{
-	Shutdown();
-
-	// shutdown the shared command queue
-	if (CommandQueuePtr.IsValid())
-	{
-		CommandQueuePtr->StopTakingCommands();
-		CommandQueuePtr.Reset();
-	}
-}
-
-FQuartzTickableObject* FQuartzTickableObject::Init(UWorld* InWorldPtr)
-{
-	if (!InWorldPtr)
-	{
-		// can't initialize if we don't have a valid world
-		return this;
-	}
-
-	GCObjectMembers.WorldPtr = InWorldPtr;
-	GCObjectMembers.QuartzSubsystem = UQuartzSubsystem::Get(GCObjectMembers.WorldPtr);
-	CommandQueuePtr = GCObjectMembers.QuartzSubsystem->CreateQuartzCommandQueue();
-	GCObjectMembers.QuartzSubsystem->SubscribeToQuartzTick(this);
-
-	bHasBeenInitialized = true;
-
-	return this;
-}
-
-int32 FQuartzTickableObject::AddCommandDelegate(const FOnQuartzCommandEventBP& InDelegate, TArray<TSharedPtr<Audio::FShareableQuartzCommandQueue, ESPMode::ThreadSafe>>& TargetSubscriberArray)
-{
-	TargetSubscriberArray.Add(CommandQueuePtr);
-
-	const int32 Num = QuantizedCommandDelegates.Num();
-	int32 SlotId = 0;
-
-	for (; SlotId < Num; ++SlotId)
-	{
-		if (!QuantizedCommandDelegates[SlotId].MulticastDelegate.IsBound())
-		{
-			QuantizedCommandDelegates[SlotId].MulticastDelegate.AddUnique(InDelegate);
-			return SlotId;
-		}
-	}
-
-	// need a new slot
-	QuantizedCommandDelegates.AddDefaulted_GetRef().MulticastDelegate.AddUnique(InDelegate);
-	return SlotId;
-}
-
-// called by UQuartzSubsystem::BeginDestroy()
-void FQuartzTickableObject::Shutdown()
-{
-	if (GCObjectMembers.WorldPtr && GCObjectMembers.QuartzSubsystem)
-	{
-		GCObjectMembers.QuartzSubsystem->UnsubscribeFromQuartzTick(this);
-		GCObjectMembers.WorldPtr = nullptr;
-		GCObjectMembers.QuartzSubsystem = nullptr;
-		bHasBeenInitialized = false;
-	}
-}
-
-TSharedPtr<Audio::FShareableQuartzCommandQueue, ESPMode::ThreadSafe> FQuartzTickableObject::GetCommandQueue()
-{
-	if (!CommandQueuePtr.IsValid())
-	{
-		CommandQueuePtr = MakeShared<Audio::FShareableQuartzCommandQueue, ESPMode::ThreadSafe>();
-	}
-
-	return CommandQueuePtr;
-}
-
-void FQuartzTickableObject::QuartzTick(float DeltaTime)
-{
-	PumpCommandQueue();
-}
-
-bool FQuartzTickableObject::QuartzIsTickable() const
-{
-	return (CommandQueuePtr && !CommandQueuePtr->IsQueueEmpty());
-}
-
-void FQuartzTickableObject::PumpCommandQueue()
-{
-	if (!CommandQueuePtr.IsValid())
-	{
-		return;
-	}
-
-	// gather move all the current commands into our temp container
-	TFunction<void(FQuartzTickableObject*)> Function;
-	while (CommandQueuePtr->EventDelegateQueue.Dequeue(Function))
-	{
-		TempCommandQueue.Emplace(Function);
-	}
-
-	// execute the commands (may alter the original container)
-	for (auto& Command : TempCommandQueue)
-	{
-		Command(this);
-	}
-
-	TempCommandQueue.Reset();
-}
-
 
 
 // Clock Handle implementation
@@ -128,32 +15,14 @@ UQuartzClockHandle::UQuartzClockHandle()
 
 UQuartzClockHandle::~UQuartzClockHandle()
 {
-
 }
 
 void UQuartzClockHandle::BeginDestroy()
 {
 	Super::BeginDestroy();
 
-	// un-subscribe from Subsystem tick and metronome events
-	if (UQuartzSubsystem* QuartzSubsystem = GetQuartzSubsystem())
-	{
-		QuartzSubsystem->UnsubscribeFromQuartzTick(this);
-
-		Audio::FQuartzClockManager* ClockManager = QuartzSubsystem->GetManagerForClock(QuartzGetWorld(), GetClockName());
-
-		if (ClockManager)
-		{
-			ClockManager->UnsubscribeFromAllTimeDivisions(CurrentClockId, GetCommandQueue());
-		}
-	}
-
-	// shutdown the shared command queue
-	if (CommandQueuePtr.IsValid())
-	{
-		CommandQueuePtr->StopTakingCommands();
-		CommandQueuePtr.Reset();
-	}
+	auto Subscriber = GetQuartzSubscriber();
+	RawHandle.SendCommandToClock([Subscriber](Audio::FQuartzClock* InClock) { InClock->UnsubscribeFromAllTimeDivisions(Subscriber); });
 }
 
 void UQuartzClockHandle::StartClock(const UObject* WorldContextObject, UQuartzClockHandle*& ClockHandle)
@@ -162,46 +31,23 @@ void UQuartzClockHandle::StartClock(const UObject* WorldContextObject, UQuartzCl
 	ResumeClock(WorldContextObject, ClockHandle);
 }
 
-void UQuartzClockHandle::StopClock(const UObject* WorldContextObject, bool CancelPendingEvents, UQuartzClockHandle*& ClockHandle)
+void UQuartzClockHandle::StopClock(const UObject* WorldContextObject, bool bCancelPendingEvents, UQuartzClockHandle*& ClockHandle)
 {
 	ClockHandle = this;
-	if (UQuartzSubsystem* QuartzSubsystem = GetQuartzSubsystem())
-	{
-		Audio::FQuartzClockManager* ClockManager = QuartzSubsystem->GetManagerForClock(WorldContextObject, GetClockName());
-
-		if (ClockManager)
-		{
-			ClockManager->StopClock(CurrentClockId, CancelPendingEvents);
-		}
-	}
+	RawHandle.SendCommandToClock([bCancelPendingEvents](Audio::FQuartzClock* InClock) { InClock->Stop(bCancelPendingEvents); });
 }
 
 void UQuartzClockHandle::PauseClock(const UObject* WorldContextObject, UQuartzClockHandle*& ClockHandle)
 {
 	ClockHandle = this;
-	if (UQuartzSubsystem* QuartzSubsystem = GetQuartzSubsystem())
-	{
-		Audio::FQuartzClockManager* ClockManager = QuartzSubsystem->GetManagerForClock(WorldContextObject, GetClockName());
-		if (ClockManager)
-		{
-			ClockManager->PauseClock(CurrentClockId);
-		}
-	}
+	RawHandle.SendCommandToClock([](Audio::FQuartzClock* InClock) { InClock->Pause(); });
 }
 
 // Begin BP interface
 void UQuartzClockHandle::ResumeClock(const UObject* WorldContextObject, UQuartzClockHandle*& ClockHandle)
 {
 	ClockHandle = this;
-	if (UQuartzSubsystem* QuartzSubsystem = GetQuartzSubsystem())
-	{
-		Audio::FQuartzClockManager* ClockManager = QuartzSubsystem->GetManagerForClock(WorldContextObject, GetClockName());
-
-		if (ClockManager)
-		{
-			ClockManager->ResumeClock(CurrentClockId);
-		}
-	}
+	RawHandle.SendCommandToClock([](Audio::FQuartzClock* InClock) { InClock->Resume(); });
 }
 
 void UQuartzClockHandle::QueueQuantizedSound(const UObject* WorldContextObject, UQuartzClockHandle*& InClockHandle, const FAudioComponentCommandInfo& InAudioComponentData, const FOnQuartzCommandEventBP& InDelegate, const FQuartzQuantizationBoundary& InTargetBoundary)
@@ -209,89 +55,54 @@ void UQuartzClockHandle::QueueQuantizedSound(const UObject* WorldContextObject,
 	InClockHandle = this;
 	FName ClockName = GetClockName();
 
-	if (UQuartzSubsystem* QuartzSubsystem = GetQuartzSubsystem())
-	{
-		//Create a Queue Command, and give it the additional data that it needs
-		TSharedPtr<Audio::FQuantizedQueueCommand> QueueCommandPtr = MakeShared<Audio::FQuantizedQueueCommand>();
-		QueueCommandPtr->SetQueueCommand(InAudioComponentData);
+	//Create a Queue Command, and give it the additional data that it needs
+	TSharedPtr<Audio::FQuantizedQueueCommand> QueueCommandPtr = MakeShared<Audio::FQuantizedQueueCommand>();
+	QueueCommandPtr->SetQueueCommand(InAudioComponentData);
 
-		//Set up initial command info
-		Audio::FQuartzQuantizedCommandInitInfo CommandInitInfo = QuartzSubsystem->CreateDataDataForSchedulePlaySound(InClockHandle, InDelegate, InTargetBoundary);
+	//Set up initial command info
+	Audio::FQuartzQuantizedCommandInitInfo CommandInitInfo = UQuartzSubsystem::CreateRequestDataForSchedulePlaySound(InClockHandle, InDelegate, InTargetBoundary);
 
-		//(Queue's setup is identical to PlaySound except for the command ptr, so fix that here)
-		CommandInitInfo.QuantizedCommandPtr.Reset();
-		CommandInitInfo.QuantizedCommandPtr = QueueCommandPtr;
+	//(Queue's setup is identical to PlaySound except for the command ptr, so fix that here)
+	CommandInitInfo.QuantizedCommandPtr.Reset();
+	CommandInitInfo.QuantizedCommandPtr = QueueCommandPtr;
 
-		QuartzSubsystem->AddCommandToClock(WorldContextObject, CommandInitInfo, ClockName);	
-	}
+	RawHandle.SendCommandToClock([CommandInitInfo](Audio::FQuartzClock* InClock) mutable { InClock->AddQuantizedCommand(CommandInitInfo); });
 }
 
 // deprecated: use ResetTransportQuantized
 void UQuartzClockHandle::ResetTransport(const UObject* WorldContextObject, const FOnQuartzCommandEventBP& InDelegate)
 {
-	if (UQuartzSubsystem* QuartzSubsystem = GetQuartzSubsystem())
-	{
-		Audio::FQuartzQuantizedCommandInitInfo Data(QuartzSubsystem->CreateDataForTransportReset(this, FQuartzQuantizationBoundary(EQuartzCommandQuantization::Bar), InDelegate));
-		QuartzSubsystem->AddCommandToClock(WorldContextObject, Data, GetClockName());
-	}
+	Audio::FQuartzQuantizedCommandInitInfo Data(UQuartzSubsystem::CreateRequestDataForTransportReset(this, FQuartzQuantizationBoundary(EQuartzCommandQuantization::Bar), InDelegate));
+	RawHandle.SendCommandToClock([Data](Audio::FQuartzClock* InClock) mutable { InClock->AddQuantizedCommand(Data); });
 }
 
 void UQuartzClockHandle::ResetTransportQuantized(const UObject* WorldContextObject, FQuartzQuantizationBoundary InQuantizationBoundary, const FOnQuartzCommandEventBP& InDelegate, UQuartzClockHandle*& ClockHandle)
 {
 	ClockHandle = this;
-	if (UQuartzSubsystem* QuartzSubsystem = GetQuartzSubsystem())
-	{
-		Audio::FQuartzQuantizedCommandInitInfo Data(QuartzSubsystem->CreateDataForTransportReset(this, InQuantizationBoundary, InDelegate));
-		QuartzSubsystem->AddCommandToClock(WorldContextObject, Data, GetClockName());
-	}
+	Audio::FQuartzQuantizedCommandInitInfo Data(UQuartzSubsystem::CreateRequestDataForTransportReset(this, InQuantizationBoundary, InDelegate));
+	RawHandle.SendCommandToClock([Data](Audio::FQuartzClock* InClock) mutable { InClock->AddQuantizedCommand(Data); });
 }
 
 
 
 bool UQuartzClockHandle::IsClockRunning(const UObject* WorldContextObject)
 {
-	if (UQuartzSubsystem* QuartzSubsystem = GetQuartzSubsystem())
-	{
-		return QuartzSubsystem->IsClockRunning(WorldContextObject, CurrentClockId);
-	}
-
-	return false;
+	return RawHandle.IsClockRunning();
 }
 
 float UQuartzClockHandle::GetDurationOfQuantizationTypeInSeconds(const UObject* WorldContextObject, const EQuartzCommandQuantization& QuantizationType, float Multiplier)
 {
-	if (UQuartzSubsystem* QuartzSubsystem = GetQuartzSubsystem())
-	{
-		return QuartzSubsystem->GetDurationOfQuantizationTypeInSeconds(WorldContextObject, CurrentClockId, QuantizationType, Multiplier);
-	}
-	else
-	{
-		return INDEX_NONE;
-	}
+	return RawHandle.GetDurationOfQuantizationTypeInSeconds(QuantizationType, Multiplier);
 }
 
 FQuartzTransportTimeStamp UQuartzClockHandle::GetCurrentTimestamp(const UObject* WorldContextObject)
 {
-	if (UQuartzSubsystem* QuartzSubsystem = GetQuartzSubsystem())
-	{
-		return QuartzSubsystem->GetCurrentClockTimestamp(WorldContextObject, CurrentClockId);
-	}
-	else
-	{
-		return FQuartzTransportTimeStamp();
-	}
+	return RawHandle.GetCurrentClockTimestamp();
 }
 
 float UQuartzClockHandle::GetEstimatedRunTime(const UObject* WorldContextObject)
 {
-	if (UQuartzSubsystem* QuartzSubsystem = GetQuartzSubsystem())
-	{
-		return QuartzSubsystem->GetEstimatedClockRunTime(WorldContextObject, CurrentClockId);
-	}
-	else
-	{
-		return INDEX_NONE;
-	}
+	return RawHandle.GetEstimatedClockRunTimeSeconds();
 }
 
 void UQuartzClockHandle::StartOtherClock(const UObject* WorldContextObject, FName OtherClockName, FQuartzQuantizationBoundary InQuantizationBoundary, const FOnQuartzCommandEventBP& InDelegate)
@@ -302,13 +113,11 @@ void UQuartzClockHandle::StartOtherClock(const UObject* WorldContextObject, FNam
 		return;
 	}
 
-	if (UQuartzSubsystem* QuartzSubsystem = GetQuartzSubsystem())
-	{
-		Audio::FQuartzQuantizedCommandInitInfo Data(QuartzSubsystem->CreateDataForStartOtherClock(this, OtherClockName, InQuantizationBoundary, InDelegate));
-		QuartzSubsystem->AddCommandToClock(WorldContextObject, Data, GetClockName());
-	}
+	Audio::FQuartzQuantizedCommandInitInfo Data(UQuartzSubsystem::CreateRequestDataForStartOtherClock(this, OtherClockName, InQuantizationBoundary, InDelegate));
+	RawHandle.SendCommandToClock([Data](Audio::FQuartzClock* InClock) mutable { InClock->AddQuantizedCommand(Data); });
 }
 
+// todo: Move the bulk of these functions to FQuartzTcikableObject once lightweight clock handles are spun up.
 void UQuartzClockHandle::SubscribeToQuantizationEvent(const UObject* WorldContextObject, EQuartzCommandQuantization InQuantizationBoundary, const FOnQuartzMetronomeEventBP& OnQuantizationEvent, UQuartzClockHandle*& ClockHandle)
 {
 	ClockHandle = this;
@@ -319,59 +128,39 @@ void UQuartzClockHandle::SubscribeToQuantizationEvent(const UObject* WorldContex
 		return;
 	}
 
-	if (UQuartzSubsystem* QuartzSubsystem = GetQuartzSubsystem())
-	{
-		Audio::FQuartzClockManager* ClockManager = QuartzSubsystem->GetManagerForClock(WorldContextObject, GetClockName());
-		if (ClockManager && ClockManager->DoesClockExist(CurrentClockId) && OnQuantizationEvent.IsBound())
-		{
-			MetronomeDelegates[static_cast<int32>(InQuantizationBoundary)].MulticastDelegate.AddUnique(OnQuantizationEvent);
-			ClockManager->SubscribeToTimeDivision(CurrentClockId, GetCommandQueue(), InQuantizationBoundary);
-		}
-	}
+	AddMetronomeBpDelegate(InQuantizationBoundary, OnQuantizationEvent);
+
+	auto Subscriber = GetQuartzSubscriber();
+	RawHandle.SendCommandToClock([Subscriber, InQuantizationBoundary](Audio::FQuartzClock* InClock) { InClock->SubscribeToTimeDivision(Subscriber, InQuantizationBoundary); });
 }
 
 void UQuartzClockHandle::SubscribeToAllQuantizationEvents(const UObject* WorldContextObject, const FOnQuartzMetronomeEventBP& OnQuantizationEvent, UQuartzClockHandle*& ClockHandle)
 {
 	ClockHandle = this;
-	if (UQuartzSubsystem* QuartzSubsystem = GetQuartzSubsystem())
-	{
-		Audio::FQuartzClockManager* ClockManager = QuartzSubsystem->GetManagerForClock(WorldContextObject, GetClockName());
-		if (ClockManager && ClockManager->DoesClockExist(CurrentClockId) && OnQuantizationEvent.IsBound())
-		{
-			for (int32 i = 0; i < static_cast<int32>(EQuartzCommandQuantization::Count) - 1; ++i)
-			{
-				MetronomeDelegates[i].MulticastDelegate.AddUnique(OnQuantizationEvent);
-			}
 
-			ClockManager->SubscribeToAllTimeDivisions(CurrentClockId, GetCommandQueue());
-		}
+	for (int32 i = 0; i < static_cast<int32>(EQuartzCommandQuantization::Count) - 1; ++i)
+	{
+		AddMetronomeBpDelegate(static_cast<EQuartzCommandQuantization>(i), OnQuantizationEvent);
 	}
+
+	auto Subscriber = GetQuartzSubscriber();
+	RawHandle.SendCommandToClock([Subscriber](Audio::FQuartzClock* InClock) { InClock->SubscribeToAllTimeDivisions(Subscriber); });
 }
 
 void UQuartzClockHandle::UnsubscribeFromTimeDivision(const UObject* WorldContextObject, EQuartzCommandQuantization InQuantizationBoundary, UQuartzClockHandle*& ClockHandle)
 {
 	ClockHandle = this;
-	if (UQuartzSubsystem* QuartzSubsystem = GetQuartzSubsystem())
-	{
-		Audio::FQuartzClockManager* ClockManager = QuartzSubsystem->GetManagerForClock(WorldContextObject, GetClockName());
-		if (ClockManager && ClockManager->DoesClockExist(CurrentClockId))
-		{
-			ClockManager->UnsubscribeFromTimeDivision(CurrentClockId, GetCommandQueue(), InQuantizationBoundary);
-		}
-	}
+
+	auto Subscriber = GetQuartzSubscriber();
+	RawHandle.SendCommandToClock([Subscriber, InQuantizationBoundary](Audio::FQuartzClock* InClock) { InClock->UnsubscribeFromTimeDivision(Subscriber, InQuantizationBoundary); });
 }
 
 void UQuartzClockHandle::UnsubscribeFromAllTimeDivisions(const UObject* WorldContextObject, UQuartzClockHandle*& ClockHandle)
 {
 	ClockHandle = this;
-	if (UQuartzSubsystem* QuartzSubsystem = GetQuartzSubsystem())
-	{
-		Audio::FQuartzClockManager* ClockManager = QuartzSubsystem->GetManagerForClock(WorldContextObject, GetClockName());
-		if (ClockManager && ClockManager->DoesClockExist(CurrentClockId))
-		{
-			ClockManager->UnsubscribeFromAllTimeDivisions(CurrentClockId, GetCommandQueue());
-		}
-	}
+
+	auto Subscriber = GetQuartzSubscriber();
+	RawHandle.SendCommandToClock([Subscriber](Audio::FQuartzClock* InClock) { InClock->UnsubscribeFromAllTimeDivisions(Subscriber); });
 }
 
 // Metronome Alteration (setters)
@@ -384,14 +173,9 @@ void UQuartzClockHandle::SetMillisecondsPerTick(const UObject* WorldContextObjec
 		return;
 	}
 
-	if (UQuartzSubsystem* QuartzSubsystem = GetQuartzSubsystem())
-	{
-		Audio::FQuartzClockTickRate TickRate;
-		TickRate.SetMillisecondsPerTick(MillisecondsPerTick);
-
-		Audio::FQuartzQuantizedCommandInitInfo Data(QuartzSubsystem->CreateDataForTickRateChange(this, InDelegate, TickRate, InQuantizationBoundary));
-		QuartzSubsystem->AddCommandToClock(WorldContextObject, Data, GetClockName());
-	}
+	Audio::FQuartzClockTickRate TickRate;
+	TickRate.SetMillisecondsPerTick(MillisecondsPerTick);
+	SetTickRateInternal(InQuantizationBoundary, InDelegate, TickRate);
 }
 
 void UQuartzClockHandle::SetTicksPerSecond(const UObject* WorldContextObject, const FQuartzQuantizationBoundary& InQuantizationBoundary, const FOnQuartzCommandEventBP& InDelegate, UQuartzClockHandle*& ClockHandle, float TicksPerSecond)
@@ -403,14 +187,9 @@ void UQuartzClockHandle::SetTicksPerSecond(const UObject* WorldContextObject, co
 		return;
 	}
 
-	if (UQuartzSubsystem* QuartzSubsystem = GetQuartzSubsystem())
-	{
-		Audio::FQuartzClockTickRate TickRate;
-		TickRate.SetSecondsPerTick(1.f / TicksPerSecond);
-
-		Audio::FQuartzQuantizedCommandInitInfo Data(QuartzSubsystem->CreateDataForTickRateChange(this, InDelegate, TickRate, InQuantizationBoundary));
-		QuartzSubsystem->AddCommandToClock(WorldContextObject, Data, GetClockName());
-	}
+	Audio::FQuartzClockTickRate TickRate;
+	TickRate.SetSecondsPerTick(1.f / TicksPerSecond);
+	SetTickRateInternal(InQuantizationBoundary, InDelegate, TickRate);
 }
 
 void UQuartzClockHandle::SetSecondsPerTick(const UObject* WorldContextObject, const FQuartzQuantizationBoundary& InQuantizationBoundary, const FOnQuartzCommandEventBP& InDelegate, UQuartzClockHandle*& ClockHandle, float SecondsPerTick)
@@ -422,14 +201,9 @@ void UQuartzClockHandle::SetSecondsPerTick(const UObject* WorldContextObject, co
 		return;
 	}
 
-	if (UQuartzSubsystem* QuartzSubsystem = GetQuartzSubsystem())
-	{
-		Audio::FQuartzClockTickRate TickRate;
-		TickRate.SetSecondsPerTick(SecondsPerTick);
-
-		Audio::FQuartzQuantizedCommandInitInfo Data(QuartzSubsystem->CreateDataForTickRateChange(this, InDelegate, TickRate, InQuantizationBoundary));
-		QuartzSubsystem->AddCommandToClock(WorldContextObject, Data, GetClockName());
-	}
+	Audio::FQuartzClockTickRate TickRate;
+	TickRate.SetSecondsPerTick(SecondsPerTick);
+	SetTickRateInternal(InQuantizationBoundary, InDelegate, TickRate);
 }
 
 void UQuartzClockHandle::SetThirtySecondNotesPerMinute(const UObject* WorldContextObject, const FQuartzQuantizationBoundary& InQuantizationBoundary, const FOnQuartzCommandEventBP& InDelegate, UQuartzClockHandle*& ClockHandle, float ThirtySecondsNotesPerMinute)
@@ -441,14 +215,9 @@ void UQuartzClockHandle::SetThirtySecondNotesPerMinute(const UObject* WorldConte
 		return;
 	}
 
-	if (UQuartzSubsystem* QuartzSubsystem = GetQuartzSubsystem())
-	{
-		Audio::FQuartzClockTickRate TickRate;
-		TickRate.SetThirtySecondNotesPerMinute(ThirtySecondsNotesPerMinute);
-
-		Audio::FQuartzQuantizedCommandInitInfo Data(QuartzSubsystem->CreateDataForTickRateChange(this, InDelegate, TickRate, InQuantizationBoundary));
-		QuartzSubsystem->AddCommandToClock(WorldContextObject, Data, GetClockName());
-	}
+	Audio::FQuartzClockTickRate TickRate;
+	TickRate.SetThirtySecondNotesPerMinute(ThirtySecondsNotesPerMinute);
+	SetTickRateInternal(InQuantizationBoundary, InDelegate, TickRate);
 }
 
 void UQuartzClockHandle::SetBeatsPerMinute(const UObject* WorldContextObject, const FQuartzQuantizationBoundary& InQuantizationBoundary, const FOnQuartzCommandEventBP& InDelegate, UQuartzClockHandle*& ClockHandle, float BeatsPerMinute)
@@ -460,14 +229,15 @@ void UQuartzClockHandle::SetBeatsPerMinute(const UObject* WorldContextObject, co
 		return;
 	}
 
-	if (UQuartzSubsystem* QuartzSubsystem = GetQuartzSubsystem())
-	{
-		Audio::FQuartzClockTickRate TickRate;
-		TickRate.SetBeatsPerMinute(BeatsPerMinute);
+	Audio::FQuartzClockTickRate TickRate;
+	TickRate.SetBeatsPerMinute(BeatsPerMinute);
+	SetTickRateInternal(InQuantizationBoundary, InDelegate, TickRate);
+}
 
-		Audio::FQuartzQuantizedCommandInitInfo Data(QuartzSubsystem->CreateDataForTickRateChange(this, InDelegate, TickRate, InQuantizationBoundary));
-		QuartzSubsystem->AddCommandToClock(WorldContextObject, Data, GetClockName());
-	}
+void UQuartzClockHandle::SetTickRateInternal(const FQuartzQuantizationBoundary& InQuantizationBoundary, const FOnQuartzCommandEventBP& InDelegate, const Audio::FQuartzClockTickRate& NewTickRate)
+{
+	Audio::FQuartzQuantizedCommandInitInfo Data(UQuartzSubsystem::CreateRequestDataForTickRateChange(this, InDelegate, NewTickRate, InQuantizationBoundary));
+	RawHandle.SendCommandToClock([Data](Audio::FQuartzClock* InClock) mutable { InClock->AddQuantizedCommand(Data); });
 }
 
 // Metronome getters
@@ -490,7 +260,7 @@ float UQuartzClockHandle::GetTicksPerSecond(const UObject* WorldContextObject) c
 	if (GetCurrentTickRate(WorldContextObject, OutTickRate))
 	{
 		const float SecondsPerTick = OutTickRate.GetSecondsPerTick();
-		
+
 		if (!FMath::IsNearlyZero(SecondsPerTick))
 		{
 			return 1.f / SecondsPerTick;
@@ -535,20 +305,44 @@ float UQuartzClockHandle::GetBeatsPerMinute(const UObject* WorldContextObject) c
 
 	return 0.f;
 }
+
+// todo: un-comment when metronome events support the offset
+// void UQuartzClockHandle::SetNotificationAnticipationAmountInMilliseconds(const UObject* WorldContextObject, UQuartzClockHandle*& ClockHandle, const double Milliseconds)
+// {
+// 	ClockHandle = this;
+// 	if(Milliseconds < 0.0)
+// 	{
+// 		UE_LOG(LogAudioQuartz, Warning, TEXT("Setting a negative notification anticipation amount is not supported. (request ignored)"));
+// 		return;
+// 	}
+//
+// 	SetNotificationAnticipationAmountMilliseconds(Milliseconds);
+// }
+//
+//
+// void UQuartzClockHandle::SetNotificationAnticipationAmountAsMusicalDuration(const UObject* WorldContextObject, UQuartzClockHandle*& ClockHandle, const EQuartzCommandQuantization MusicalDuration, const double Multiplier)
+// {
+// 	ClockHandle = this;
+// 	if(Multiplier < 0.0)
+// 	{
+// 		UE_LOG(LogAudioQuartz, Warning, TEXT("Setting a negative notification anticipation amount is not supported. (request ignored)"));
+// 		return;
+// 	}
+//
+// 	SetNotificationAnticipationAmountMusicalDuration(MusicalDuration, Multiplier);
+// }
+
 // End BP interface
 
 
-UQuartzClockHandle* UQuartzClockHandle::SubscribeToClock(const UObject* WorldContextObject, FName ClockName)
+UQuartzClockHandle* UQuartzClockHandle::SubscribeToClock(const UObject* WorldContextObject, FName ClockName, Audio::FQuartzClockProxy const* InHandlePtr)
 {
-	// create ID
 	CurrentClockId = ClockName;
 
-	FString TempId = WorldContextObject->GetFName().ToString();
-	TempId.Append(CurrentClockId.ToString());
-	ClockHandleId = FName(*TempId);
-
-	// TODO: subscribe to clock w/ ClockHandleId
-	bConnectedToClock = true;
+	if (InHandlePtr)
+	{
+		RawHandle = *InHandlePtr;
+	}
 
 	return this;
 }
@@ -557,57 +351,12 @@ UQuartzClockHandle* UQuartzClockHandle::SubscribeToClock(const UObject* WorldCon
 // returns true if OutTickRate is valid and was updated
 bool UQuartzClockHandle::GetCurrentTickRate(const UObject* WorldContextObject, Audio::FQuartzClockTickRate& OutTickRate) const
 {
-	if (UQuartzSubsystem* QuartzSubsystem = GetQuartzSubsystem())
+	if (RawHandle.IsValid())
 	{
-		Audio::FQuartzClockManager* ClockManager = QuartzSubsystem->GetManagerForClock(WorldContextObject, GetClockName());
-
-		if (ClockManager)
-		{
-			OutTickRate = ClockManager->GetTickRateForClock(CurrentClockId);
-			return true;
-		}
+		OutTickRate = RawHandle.GetTickRate();
+		return true;
 	}
 
 	OutTickRate = {};
 	return false;
 }
-
-void UQuartzClockHandle::ProcessCommand(const Audio::FQuartzQuantizedCommandDelegateData& Data)
-{
-	checkSlow(Data.DelegateSubType < EQuartzCommandDelegateSubType::Count && (Data.DelegateID < QuantizedCommandDelegates.Num()));
-	if (UQuartzSubsystem* QuartzSubsystem = GetQuartzSubsystem())
-	{
-		QuartzSubsystem->PushLatencyTrackerResult(Data.RequestRecieved());
-	}
-
-	CommandDelegateGameThreadData& GameThreadEntry = QuantizedCommandDelegates[Data.DelegateID];
-
-	GameThreadEntry.MulticastDelegate.Broadcast(Data.DelegateSubType, "Sample Payload");
-
-	// track the number of active QuantizedCommands that may be sending info back to us.
-	// (new command)
-	if (Data.DelegateSubType == EQuartzCommandDelegateSubType::CommandOnQueued)
-	{
-		GameThreadEntry.RefCount.Increment();
-	}
-
-	// (end of a command)
-	if (Data.DelegateSubType == EQuartzCommandDelegateSubType::CommandOnCanceled)
-	{
-		// are all the commands done?
-		if (GameThreadEntry.RefCount.Decrement() == 0)
-		{
-			GameThreadEntry.MulticastDelegate.Clear();
-		}
-	}
-}
-
-void UQuartzClockHandle::ProcessCommand(const Audio::FQuartzMetronomeDelegateData& Data)
-{
-	if (UQuartzSubsystem* QuartzSubsystem = GetQuartzSubsystem())
-	{
-		QuartzSubsystem->PushLatencyTrackerResult(Data.RequestRecieved());
-	}
-
-	MetronomeDelegates[static_cast<int32>(Data.Quantization)].MulticastDelegate.Broadcast(CurrentClockId, Data.Quantization, Data.Bar, Data.Beat, Data.BeatFraction);
-}
\ No newline at end of file
diff --git a/Engine/Source/Runtime/AudioMixer/Private/Quartz/AudioMixerClockManager.cpp b/Engine/Source/Runtime/AudioMixer/Private/Quartz/AudioMixerClockManager.cpp
--- a/Engine/Source/Runtime/AudioMixer/Private/Quartz/AudioMixerClockManager.cpp
+++ b/Engine/Source/Runtime/AudioMixer/Private/Quartz/AudioMixerClockManager.cpp
@@ -52,7 +52,7 @@ namespace Audio
 		}
 	}
 
-	TSharedPtr<FQuartzClock> FQuartzClockManager::GetOrCreateClock(const FName& InClockName, const FQuartzClockSettings& InClockSettings, bool bOverrideTickRateIfClockExists)
+	FQuartzClockProxy FQuartzClockManager::GetOrCreateClock(const FName& InClockName, const FQuartzClockSettings& InClockSettings, bool bOverrideTickRateIfClockExists)
 	{
 		FScopeLock Lock(&ActiveClockCritSec);
 
@@ -70,11 +70,25 @@ namespace Audio
 				Clock->ChangeTimeSignature(NewSettings.TimeSignature);
 			}
 
-			return Clock;
+			return FQuartzClockProxy(Clock);
 		}
 
 		// doesn't exist, create new clock
-		return ActiveClocks.Emplace_GetRef(MakeShared<FQuartzClock>(InClockName, NewSettings, this));
+		return FQuartzClockProxy(ActiveClocks.Emplace_GetRef(MakeShared<FQuartzClock>(InClockName, NewSettings, this)));
+	}
+
+	FQuartzClockProxy FQuartzClockManager::GetClock(const FName& InClockName)
+	{
+		FScopeLock Lock(&ActiveClockCritSec);
+		TSharedPtr<FQuartzClock> ClockPtr = FindClock(InClockName);
+
+		if (ClockPtr)
+		{
+			return FQuartzClockProxy(ClockPtr);
+		}
+
+		UE_LOG(LogAudioQuartz, Warning, TEXT("Could not find Clock: %s (returning empty handle)"), *ClockPtr->GetName().ToString());
+		return {};
 	}
 
 	bool FQuartzClockManager::DoesClockExist(const FName& InClockName)
@@ -344,7 +358,7 @@ namespace Audio
 		TSharedPtr<FQuartzClock> Clock = FindClock(InClockName);
 		if (Clock)
 		{
-			Clock->SubscribeToTimeDivision(InListenerQueue, InQuantizationBoundary);
+			Clock->SubscribeToTimeDivision(FQuartzGameThreadSubscriber(InListenerQueue), InQuantizationBoundary);
 		}
 	}
 
@@ -365,7 +379,7 @@ namespace Audio
 		TSharedPtr<FQuartzClock> Clock = FindClock(InClockName);
 		if (Clock)
 		{
-			Clock->SubscribeToAllTimeDivisions(InListenerQueue);
+			Clock->SubscribeToAllTimeDivisions(FQuartzGameThreadSubscriber(InListenerQueue));
 		}
 	}
 
@@ -386,7 +400,7 @@ namespace Audio
 		TSharedPtr<FQuartzClock> Clock = FindClock(InClockName);
 		if (Clock)
 		{
-			Clock->UnsubscribeFromTimeDivision(InListenerQueue, InQuantizationBoundary);
+			Clock->UnsubscribeFromTimeDivision(FQuartzGameThreadSubscriber(InListenerQueue), InQuantizationBoundary);
 		}
 	}
 
@@ -407,7 +421,7 @@ namespace Audio
 		TSharedPtr<FQuartzClock> Clock = FindClock(InClockName);
 		if (Clock)
 		{
-			Clock->UnsubscribeFromAllTimeDivisions(InListenerQueue);
+			Clock->UnsubscribeFromAllTimeDivisions(FQuartzGameThreadSubscriber(InListenerQueue));
 		}
 	}
 
diff --git a/Engine/Source/Runtime/AudioMixer/Private/Quartz/AudioMixerQuantizedCommands.cpp b/Engine/Source/Runtime/AudioMixer/Private/Quartz/AudioMixerQuantizedCommands.cpp
--- a/Engine/Source/Runtime/AudioMixer/Private/Quartz/AudioMixerQuantizedCommands.cpp
+++ b/Engine/Source/Runtime/AudioMixer/Private/Quartz/AudioMixerQuantizedCommands.cpp
@@ -132,12 +132,12 @@ namespace Audio
 
 	void FQuantizedQueueCommand::OnFinalCallbackCustom(int32 InNumFramesLeft)
 	{
-		if (OwningClockPtr && AudioComponentData.ComponentCommandPtr.IsValid())
+		if (OwningClockPtr)
 		{
 			FName ClockName = OwningClockPtr->GetName();
 			Audio::FQuartzQueueCommandData CommandData(AudioComponentData, ClockName);
 
-			AudioComponentData.ComponentCommandPtr->PushEvent(CommandData);
+			AudioComponentData.Subscriber.PushEvent(CommandData);
 		}
 	}
 
diff --git a/Engine/Source/Runtime/AudioMixer/Private/Quartz/QuartzMetronome.cpp b/Engine/Source/Runtime/AudioMixer/Private/Quartz/QuartzMetronome.cpp
--- a/Engine/Source/Runtime/AudioMixer/Private/Quartz/QuartzMetronome.cpp
+++ b/Engine/Source/Runtime/AudioMixer/Private/Quartz/QuartzMetronome.cpp
@@ -5,13 +5,14 @@
 
 namespace Audio
 {
-	FQuartzMetronome::FQuartzMetronome() : TimeSinceStart(0)
+	FQuartzMetronome::FQuartzMetronome(FName InClockName)
+		: TimeSinceStart(0), ClockName(InClockName)
 	{
 		SetTickRate(CurrentTickRate);
 	}
 
-	FQuartzMetronome::FQuartzMetronome(const FQuartzTimeSignature& InTimeSignature)
-		: CurrentTimeSignature(InTimeSignature), TimeSinceStart(0)
+	FQuartzMetronome::FQuartzMetronome(const FQuartzTimeSignature& InTimeSignature, FName InClockName)
+		: CurrentTimeSignature(InTimeSignature), TimeSinceStart(0), ClockName(InClockName)
 	{
 		SetTickRate(CurrentTickRate);
 	}
@@ -417,17 +418,18 @@ namespace Audio
 
 	void FQuartzMetronome::FireEvents(int32 EventFlags)
 	{
-		FQuartzMetronomeDelegateData Data;
-		Data.Bar = (CurrentTimeStamp.Bars);
-		Data.Beat = (CurrentTimeStamp.Beat);
-		Data.BeatFraction = (CurrentTimeStamp.BeatFraction);
-
 		if (!(EventFlags &= ListenerFlags))
 		{
 			// no events occurred that we have listeners for
 			return;
 		}
 
+		FQuartzMetronomeDelegateData Data;
+		Data.Bar = (CurrentTimeStamp.Bars);
+		Data.Beat = (CurrentTimeStamp.Beat);
+		Data.BeatFraction = (CurrentTimeStamp.BeatFraction);
+		Data.ClockName = ClockName;
+
 		// loop through quantization boundaries
 		int32 i = -1;
 		for (TArray<MetronomeCommandQueuePtr>& QuantizationBoundarySubscribers : MetronomeSubscriptionMatrix)
diff --git a/Engine/Source/Runtime/AudioMixer/Private/Quartz/QuartzSubsystem.cpp b/Engine/Source/Runtime/AudioMixer/Private/Quartz/QuartzSubsystem.cpp
--- a/Engine/Source/Runtime/AudioMixer/Private/Quartz/QuartzSubsystem.cpp
+++ b/Engine/Source/Runtime/AudioMixer/Private/Quartz/QuartzSubsystem.cpp
@@ -35,7 +35,7 @@ static FAudioDevice* GetAudioDeviceUsingWorldContext(const UObject* WorldContext
 		return nullptr;
 	}
 
-	UWorld* ThisWorld = GEngine->GetWorldFromContextObject(WorldContextObject, EGetWorldErrorMode::LogAndReturnNull);
+	const UWorld* ThisWorld = GEngine->GetWorldFromContextObject(WorldContextObject, EGetWorldErrorMode::LogAndReturnNull);
 	if (!ThisWorld || !ThisWorld->bAllowAudioPlayback || ThisWorld->GetNetMode() == NM_DedicatedServer)
 	{
 		return nullptr;
@@ -62,46 +62,18 @@ static Audio::FMixerDevice* GetAudioMixerDeviceUsingWorldContext(const UObject*
 }
 
 
-UQuartzSubsystem::UQuartzSubsystem()
-{
-}
-
-UQuartzSubsystem::~UQuartzSubsystem()
-{
-}
-
 void UQuartzSubsystem::BeginDestroy()
 {
 	Super::BeginDestroy();
 
 	// force un-subscribe all Quartz tickable objects
-	TArray<FQuartzTickableObject*> SubscribersCopy = QuartzTickSubscribers;
-	for (FQuartzTickableObject* Entry : SubscribersCopy)
-	{
-		Entry->Shutdown();
-	}
-	QuartzTickSubscribers.Reset();
-
 	SubsystemClockManager.Shutdown();
 	SubsystemClockManager.Flush();
 }
 
-bool UQuartzSubsystem::DoesSupportWorldType(EWorldType::Type WorldType) const
+void FQuartzTickableObjectsManager::Tick(float DeltaTime)
 {
-	return Super::DoesSupportWorldType(WorldType) || WorldType == EWorldType::EditorPreview;
-}
-
-
-void UQuartzSubsystem::Tick(float DeltaTime)
-{
-	Super::Tick(DeltaTime);
-	TRACE_CPUPROFILER_EVENT_SCOPE(QuartzSubsystem::Tick);
-
-
-	SubsystemClockManager.LowResoultionUpdate(DeltaTime);
-
 	const int32 NumSubscribers = QuartzTickSubscribers.Num();
-
 	if (MaxQuartzSubscribersToUpdatePerTickCvar <= 0 || NumSubscribers <= MaxQuartzSubscribersToUpdatePerTickCvar)
 	{
 		TArray<FQuartzTickableObject*> SubscribersCopy = QuartzTickSubscribers;
@@ -141,23 +113,14 @@ void UQuartzSubsystem::Tick(float DeltaTime)
 	}
 }
 
-bool UQuartzSubsystem::IsTickable() const
+bool FQuartzTickableObjectsManager::IsTickable() const
 {
 	const int32 NumSubscribers = QuartzTickSubscribers.Num();
-	const int32 NumClocks = SubsystemClockManager.GetNumClocks();
 	const bool bHasTickSubscribers = NumSubscribers > 0;
-	const bool bIsManagingClocks = NumClocks > 0;
-
-	TRACE_INT_VALUE(TEXT("QuartzSubsystem::NumClocks"), NumClocks);
 	TRACE_INT_VALUE(TEXT("QuartzSubsystem::NumSubscribers"), NumSubscribers);
 
-	// if our manager has clocks, we need to tick
-	if (bIsManagingClocks)
-	{
-		return true;
-	}
 	// if our manager has no clocks, and we have no ClockHandle subscribers, we don't need to tick
-	else if (!bHasTickSubscribers)
+	if (!bHasTickSubscribers)
 	{
 		return false;
 	}
@@ -174,13 +137,7 @@ bool UQuartzSubsystem::IsTickable() const
 	return false;
 }
 
-TStatId UQuartzSubsystem::GetStatId() const
-{
-	RETURN_QUICK_DECLARE_CYCLE_STAT(UQuartzSubsystem, STATGROUP_Tickables);
-}
-
-
-void UQuartzSubsystem::SubscribeToQuartzTick(FQuartzTickableObject * InObjectToTick)
+void FQuartzTickableObjectsManager::SubscribeToQuartzTick(FQuartzTickableObject* InObjectToTick)
 {
 	if (!InObjectToTick)
 	{
@@ -190,8 +147,7 @@ void UQuartzSubsystem::SubscribeToQuartzTick(FQuartzTickableObject * InObjectToT
 	QuartzTickSubscribers.AddUnique(InObjectToTick);
 }
 
-
-void UQuartzSubsystem::UnsubscribeFromQuartzTick(FQuartzTickableObject * InObjectToTick)
+void FQuartzTickableObjectsManager::UnsubscribeFromQuartzTick(FQuartzTickableObject* InObjectToTick)
 {
 	if (!InObjectToTick)
 	{
@@ -201,26 +157,69 @@ void UQuartzSubsystem::UnsubscribeFromQuartzTick(FQuartzTickableObject * InObjec
 	QuartzTickSubscribers.RemoveSingleSwap(InObjectToTick);
 }
 
+bool UQuartzSubsystem::DoesSupportWorldType(EWorldType::Type WorldType) const
+{
+	return Super::DoesSupportWorldType(WorldType) || WorldType == EWorldType::EditorPreview;
+}
+
 
-UQuartzSubsystem* UQuartzSubsystem::Get(UWorld* World)
+void UQuartzSubsystem::Tick(float DeltaTime)
 {
-	if (World)
-	{
-		return World->GetSubsystem<UQuartzSubsystem>();
-	}
+	Super::Tick(DeltaTime);
+	TRACE_CPUPROFILER_EVENT_SCOPE(QuartzSubsystem::Tick);
 
-	return nullptr;
+	ensure(TickableObjectManagerPtr.IsValid());
+
+	PruneStaleProxies();
+	SubsystemClockManager.LowResoultionUpdate(DeltaTime);
+	TickableObjectManagerPtr->Tick(DeltaTime);
 }
 
+bool UQuartzSubsystem::IsTickable() const
+{
+	ensure(TickableObjectManagerPtr.IsValid());
+
+	const int32 NumClocks = SubsystemClockManager.GetNumClocks();
+	TRACE_INT_VALUE(TEXT("QuartzSubsystem::NumClocks"), NumClocks);
+
+	// IsTickable() updates unreal insights values
+	const bool bSubscribersNeedUpdate = TickableObjectManagerPtr->IsTickable();
+	const bool bIsManagingClocks = NumClocks > 0;
 
-TSharedPtr<Audio::FShareableQuartzCommandQueue, ESPMode::ThreadSafe> UQuartzSubsystem::CreateQuartzCommandQueue()
+	return bIsManagingClocks || bSubscribersNeedUpdate;
+}
+
+TStatId UQuartzSubsystem::GetStatId() const
 {
-	return MakeShared<Audio::FShareableQuartzCommandQueue, ESPMode::ThreadSafe>();
+	RETURN_QUICK_DECLARE_CYCLE_STAT(UQuartzSubsystem, STATGROUP_Tickables);
 }
 
+ void UQuartzSubsystem::SubscribeToQuartzTick(FQuartzTickableObject * InObjectToTick)
+ {
+	ensure(TickableObjectManagerPtr.IsValid());
+ 	TickableObjectManagerPtr->SubscribeToQuartzTick(InObjectToTick);
+ }
+
+
+ void UQuartzSubsystem::UnsubscribeFromQuartzTick(FQuartzTickableObject * InObjectToTick)
+ {
+	ensure(TickableObjectManagerPtr.IsValid());
+ 	TickableObjectManagerPtr->UnsubscribeFromQuartzTick(InObjectToTick);
+ }
 
 
-Audio::FQuartzQuantizedRequestData UQuartzSubsystem::CreateDataDataForSchedulePlaySound(UQuartzClockHandle* InClockHandle, const FOnQuartzCommandEventBP& InDelegate, const FQuartzQuantizationBoundary& InQuantizationBoundary)
+UQuartzSubsystem* UQuartzSubsystem::Get(const UWorld* const World)
+{
+	if (World)
+	{
+		return World->GetSubsystem<UQuartzSubsystem>();
+	}
+
+	return nullptr;
+}
+
+
+Audio::FQuartzQuantizedRequestData UQuartzSubsystem::CreateRequestDataForSchedulePlaySound(UQuartzClockHandle* InClockHandle, const FOnQuartzCommandEventBP& InDelegate, const FQuartzQuantizationBoundary& InQuantizationBoundary)
 {
 	Audio::FQuartzQuantizedRequestData CommandInitInfo;
 
@@ -230,131 +229,159 @@ Audio::FQuartzQuantizedRequestData UQuartzSubsystem::CreateDataDataForSchedulePl
 	}
 
 	CommandInitInfo.ClockName = InClockHandle->GetClockName();
-	CommandInitInfo.ClockHandleName = InClockHandle->GetHandleName();
 	CommandInitInfo.QuantizationBoundary = InQuantizationBoundary;
 	CommandInitInfo.QuantizedCommandPtr = MakeShared<Audio::FQuantizedPlayCommand>();
 	CommandInitInfo.GameThreadSubscribers.Append(InQuantizationBoundary.GameThreadSubscribers);
-
-	const bool bRequiresAudioDevice = CommandInitInfo.QuantizedCommandPtr->RequiresAudioDevice();
-	const bool bClockManagedByAudioDevice = ClockManagerTypeMap.Contains(CommandInitInfo.ClockName)
-		&& ClockManagerTypeMap[CommandInitInfo.ClockName] == EQuarztClockManagerType::AudioEngine;
-
-	if (!(bRequiresAudioDevice && bClockManagedByAudioDevice))
-	{
-		return {};
-	}
+	CommandInitInfo.GameThreadSubscribers.Add(InClockHandle->GetQuartzSubscriber());
 
 	if (InDelegate.IsBound())
 	{
-		CommandInitInfo.GameThreadDelegateID = InClockHandle->AddCommandDelegate(InDelegate, CommandInitInfo.GameThreadSubscribers);
+		CommandInitInfo.GameThreadDelegateID = InClockHandle->AddCommandDelegate(InDelegate);
 	}
 
 	return CommandInitInfo;
 }
 
 
+Audio::FQuartzQuantizedRequestData UQuartzSubsystem::CreateDataForTickRateChange(UQuartzClockHandle* InClockHandle, const FOnQuartzCommandEventBP& InDelegate, const Audio::FQuartzClockTickRate& InNewTickRate, const FQuartzQuantizationBoundary& InQuantizationBoundary)
+{
+	// new static function
+	return UQuartzSubsystem::CreateRequestDataForTickRateChange(InClockHandle, InDelegate, InNewTickRate, InQuantizationBoundary);
+}
+
+
+Audio::FQuartzQuantizedRequestData UQuartzSubsystem::CreateDataForTransportReset(UQuartzClockHandle* InClockHandle, const FQuartzQuantizationBoundary& InQuantizationBoundary, const FOnQuartzCommandEventBP& InDelegate)
+{
+	// new static function
+	return UQuartzSubsystem::CreateRequestDataForTransportReset(InClockHandle, InQuantizationBoundary, InDelegate);
+}
+
+
+Audio::FQuartzQuantizedRequestData CreateDataForStartOtherClock(UQuartzClockHandle* InClockHandle, FName InClockToStart, const FQuartzQuantizationBoundary& InQuantizationBoundary, const FOnQuartzCommandEventBP& InDelegate)
+{
+	// new static function
+	return UQuartzSubsystem::CreateRequestDataForStartOtherClock(InClockHandle, InClockToStart, InQuantizationBoundary, InDelegate);
+}
+
+
+Audio::FQuartzQuantizedRequestData CreateDataDataForSchedulePlaySound(UQuartzClockHandle* InClockHandle, const FOnQuartzCommandEventBP& InDelegate, const FQuartzQuantizationBoundary& InQuantizationBoundary)
+{
+	// new static function
+	return UQuartzSubsystem::CreateRequestDataForSchedulePlaySound(InClockHandle, InDelegate, InQuantizationBoundary);
+}
+
+
 bool UQuartzSubsystem::IsQuartzEnabled()
 {
 	return true;
 }
 
 
-Audio::FQuartzQuantizedRequestData UQuartzSubsystem::CreateDataForTickRateChange(UQuartzClockHandle* InClockHandle, const FOnQuartzCommandEventBP& InDelegate, const Audio::FQuartzClockTickRate& InNewTickRate, const FQuartzQuantizationBoundary& InQuantizationBoundary)
+Audio::FQuartzQuantizedRequestData UQuartzSubsystem::CreateRequestDataForTickRateChange(UQuartzClockHandle* InClockHandle, const FOnQuartzCommandEventBP& InDelegate, const Audio::FQuartzClockTickRate& InNewTickRate, const FQuartzQuantizationBoundary& InQuantizationBoundary)
 {
 	if (!ensure(InClockHandle))
 	{
 		return { };
 	}
 
-	TSharedPtr<Audio::FQuantizedTickRateChange> TickRateChangeCommandPtr = MakeShared<Audio::FQuantizedTickRateChange>();
+	const TSharedPtr<Audio::FQuantizedTickRateChange> TickRateChangeCommandPtr = MakeShared<Audio::FQuantizedTickRateChange>();
 	TickRateChangeCommandPtr->SetTickRate(InNewTickRate);
 
 	Audio::FQuartzQuantizedRequestData CommandInitInfo;
 
 	CommandInitInfo.ClockName = InClockHandle->GetClockName();
-	CommandInitInfo.ClockHandleName = InClockHandle->GetHandleName();
 	CommandInitInfo.QuantizationBoundary = InQuantizationBoundary;
 	CommandInitInfo.QuantizedCommandPtr = TickRateChangeCommandPtr;
 	CommandInitInfo.GameThreadSubscribers.Append(InQuantizationBoundary.GameThreadSubscribers);
+	CommandInitInfo.GameThreadSubscribers.Add(InClockHandle->GetQuartzSubscriber());
 
 	if (InDelegate.IsBound())
 	{
-		CommandInitInfo.GameThreadDelegateID = InClockHandle->AddCommandDelegate(InDelegate, CommandInitInfo.GameThreadSubscribers);
+		CommandInitInfo.GameThreadDelegateID = InClockHandle->AddCommandDelegate(InDelegate);
 	}
 
 	return CommandInitInfo;
 }
 
-Audio::FQuartzQuantizedRequestData UQuartzSubsystem::CreateDataForTransportReset(UQuartzClockHandle* InClockHandle, const FQuartzQuantizationBoundary& InQuantizationBoundary, const FOnQuartzCommandEventBP& InDelegate)
+Audio::FQuartzQuantizedRequestData UQuartzSubsystem::CreateRequestDataForTransportReset(UQuartzClockHandle* InClockHandle, const FQuartzQuantizationBoundary& InQuantizationBoundary, const FOnQuartzCommandEventBP& InDelegate)
 {
 	if (!ensure(InClockHandle))
 	{
 		return { };
 	}
 
-	TSharedPtr<Audio::FQuantizedTransportReset> TransportResetCommandPtr = MakeShared<Audio::FQuantizedTransportReset>();
+	const TSharedPtr<Audio::FQuantizedTransportReset> TransportResetCommandPtr = MakeShared<Audio::FQuantizedTransportReset>();
 
 	Audio::FQuartzQuantizedRequestData CommandInitInfo;
 
 	CommandInitInfo.ClockName = InClockHandle->GetClockName();
-	CommandInitInfo.ClockHandleName = InClockHandle->GetHandleName();
 	CommandInitInfo.QuantizationBoundary = InQuantizationBoundary;
 	CommandInitInfo.QuantizedCommandPtr = TransportResetCommandPtr;
 	CommandInitInfo.GameThreadSubscribers.Append(InQuantizationBoundary.GameThreadSubscribers);
+	CommandInitInfo.GameThreadSubscribers.Add(InClockHandle->GetQuartzSubscriber());
 
 	if (InDelegate.IsBound())
 	{
-		CommandInitInfo.GameThreadDelegateID = InClockHandle->AddCommandDelegate(InDelegate, CommandInitInfo.GameThreadSubscribers);
+		CommandInitInfo.GameThreadDelegateID = InClockHandle->AddCommandDelegate(InDelegate);
 	}
 
 	return CommandInitInfo;
 }
 
-Audio::FQuartzQuantizedRequestData UQuartzSubsystem::CreateDataForStartOtherClock(UQuartzClockHandle* InClockHandle, FName InClockToStart, const FQuartzQuantizationBoundary& InQuantizationBoundary, const FOnQuartzCommandEventBP& InDelegate)
+Audio::FQuartzQuantizedRequestData UQuartzSubsystem::CreateRequestDataForStartOtherClock(UQuartzClockHandle* InClockHandle, FName InClockToStart, const FQuartzQuantizationBoundary& InQuantizationBoundary, const FOnQuartzCommandEventBP& InDelegate)
 {
 	if (!ensure(InClockHandle))
 	{
 		return { };
 	}
 
-	TSharedPtr<Audio::FQuantizedOtherClockStart> TransportResetCommandPtr = MakeShared<Audio::FQuantizedOtherClockStart>();
+	const TSharedPtr<Audio::FQuantizedOtherClockStart> TransportResetCommandPtr = MakeShared<Audio::FQuantizedOtherClockStart>();
 
 	Audio::FQuartzQuantizedRequestData CommandInitInfo;
 
 	CommandInitInfo.ClockName = InClockHandle->GetClockName();
-	CommandInitInfo.ClockHandleName = InClockHandle->GetHandleName();
 	CommandInitInfo.OtherClockName = InClockToStart;
 	CommandInitInfo.QuantizationBoundary = InQuantizationBoundary;
 	CommandInitInfo.QuantizedCommandPtr = TransportResetCommandPtr;
 	CommandInitInfo.GameThreadSubscribers.Append(InQuantizationBoundary.GameThreadSubscribers);
+	CommandInitInfo.GameThreadSubscribers.Add(InClockHandle->GetQuartzSubscriber());
 
 	if (InDelegate.IsBound())
 	{
-		CommandInitInfo.GameThreadDelegateID = InClockHandle->AddCommandDelegate(InDelegate, CommandInitInfo.GameThreadSubscribers);
+		CommandInitInfo.GameThreadDelegateID = InClockHandle->AddCommandDelegate(InDelegate);
 	}
 
 	return CommandInitInfo;
 }
 
-UQuartzClockHandle* UQuartzSubsystem::CreateNewClock(const UObject* WorldContextObject, FName ClockName, FQuartzClockSettings InSettings, bool bOverrideSettingsIfClockExists, bool bUseAudioEngineClockManager)
+Audio::FQuartzClockManager* UQuartzSubsystem::GetClockManager(const UObject* WorldContextObject, bool bUseAudioEngineClockManager)
 {
-	if (ClockName.IsNone())
+	// decide if the clock should be managed by the AudioDevice (audio engine) or the Subsystem (this object)
+	Audio::FQuartzClockManager* ClockManager;
+	Audio::FMixerDevice* MixerDevice = GetAudioMixerDeviceUsingWorldContext(WorldContextObject);
+	if(!bUseAudioEngineClockManager || SimulateNoAudioDeviceCvar || !MixerDevice)
 	{
-		return nullptr;
+		ClockManager = &SubsystemClockManager;
 	}
-
-	if (!bUseAudioEngineClockManager)
+	else
 	{
-		ClockManagerTypeMap.Add(ClockName, EQuarztClockManagerType::QuartzSubsystem);
+		ClockManager = &MixerDevice->QuantizedEventClockManager;
 	}
 
-	// add or create clock
-	Audio::FQuartzClockManager* ClockManager = GetManagerForClock(WorldContextObject, ClockName);
-	if (!ClockManager)
+	// we should have fallen back to this object
+	return ClockManager;
+}
+
+UQuartzClockHandle* UQuartzSubsystem::CreateNewClock(const UObject* WorldContextObject, FName ClockName, FQuartzClockSettings InSettings, bool bOverrideSettingsIfClockExists, bool bUseAudioEngineClockManager)
+{
+	if (ClockName.IsNone() || !WorldContextObject)
 	{
 		return nullptr;
 	}
 
+	Audio::FQuartzClockManager* ClockManager = GetClockManager(WorldContextObject, bUseAudioEngineClockManager);
+	ensure(ClockManager); // should have at least fallen back to "this" object as a manager
+
 	// numerator of time signature must be >= 1
 	if (InSettings.TimeSignature.NumBeats < 1)
 	{
@@ -362,25 +389,29 @@ UQuartzClockHandle* UQuartzSubsystem::CreateNewClock(const UObject* WorldContext
 		InSettings.TimeSignature.NumBeats = 1;
 	}
 
-	ClockManager->GetOrCreateClock(ClockName, InSettings, bOverrideSettingsIfClockExists);
+	Audio::FQuartzClockProxy ClockProxy = ClockManager->GetOrCreateClock(ClockName, InSettings, bOverrideSettingsIfClockExists);
+	UQuartzClockHandle* ClockHandle = static_cast<UQuartzClockHandle*>(NewObject<UQuartzClockHandle>()->Init(WorldContextObject->GetWorld()));
+	ClockHandle->SubscribeToClock(WorldContextObject, ClockName, &ClockProxy);
 
-	UQuartzClockHandle* ClockHandlePtr = (UQuartzClockHandle *)(NewObject<UQuartzClockHandle>()->Init(WorldContextObject->GetWorld()));
-	ClockHandlePtr->SubscribeToClock(WorldContextObject, ClockName);
+	// if we are not the manager for the clock, it means the FAudioDevice is,
+	// so we hold onto our own copy of the proxy
+	ActiveAudioMixerClockProxies.Add(ClockProxy);
 
-	return ClockHandlePtr;
+	return ClockHandle;
 }
 
 
 void UQuartzSubsystem::DeleteClockByName(const UObject* WorldContextObject, FName ClockName)
 {
-	Audio::FQuartzClockManager* ClockManager = GetManagerForClock(WorldContextObject, ClockName);
+	Audio::FQuartzClockManager* ClockManager = GetClockManager(WorldContextObject);
+
 	if (!ClockManager)
 	{
 		return;
 	}
 
 	ClockManager->RemoveClock(ClockName);
-	ClockManagerTypeMap.Remove(ClockName);
+	PruneStaleProxies();
 }
 
 void UQuartzSubsystem::DeleteClockByHandle(const UObject* WorldContextObject, UQuartzClockHandle*& InClockHandle)
@@ -393,80 +424,99 @@ void UQuartzSubsystem::DeleteClockByHandle(const UObject* WorldContextObject, UQ
 
 UQuartzClockHandle* UQuartzSubsystem::GetHandleForClock(const UObject* WorldContextObject, FName ClockName)
 {
-	Audio::FQuartzClockManager* ClockManager = GetManagerForClock(WorldContextObject, ClockName);
-	if (!ClockManager)
+	Audio::FQuartzClockManager* ClockManager = GetClockManager(WorldContextObject);
+	if (!ClockManager || !ClockManager->DoesClockExist(ClockName))
 	{
 		return nullptr;
 	}
 
-	if (ClockManager->DoesClockExist(ClockName))
+	Audio::FQuartzClockProxy ClockHandle = ClockManager->GetClock(ClockName);
+
+	UQuartzClockHandle* ClockHandlePtr = static_cast<UQuartzClockHandle*>(NewObject<UQuartzClockHandle>()->Init(WorldContextObject->GetWorld()));
+	return ClockHandlePtr->SubscribeToClock(WorldContextObject, ClockName, &ClockHandle);
+}
+
+
+Audio::FQuartzClockProxy UQuartzSubsystem::GetProxyForClock(FName ClockName) const
+{
+	if(Audio::FQuartzClockProxy const* ProxyPtr = FindProxyByName(ClockName))
 	{
-		UQuartzClockHandle* ClockHandlePtr = (UQuartzClockHandle *)(NewObject<UQuartzClockHandle>()->Init(WorldContextObject->GetWorld()));
-		return ClockHandlePtr->SubscribeToClock(WorldContextObject, ClockName);
+		return *ProxyPtr; // caller gets their own copy
 	}
 
-	return nullptr;
+	return {};
+}
+
+void UQuartzSubsystem::AddProxyForExternalClock(const Audio::FQuartzClockProxy& InProxy)
+{
+	// make sure we aren't adding a duplicate name
+	if(FindProxyByName(InProxy.GetClockName()))
+	{
+		UE_LOG(LogAudioQuartz, Warning, TEXT("Recieved request to add external Clock: (%s) when a clock of that name already exists (Ignoring Request)"), *InProxy.GetClockName().ToString());
+		return;
+	}
 }
 
 
 bool UQuartzSubsystem::DoesClockExist(const UObject* WorldContextObject, FName ClockName)
 {
-	Audio::FQuartzClockManager* ClockManager = GetManagerForClock(WorldContextObject, ClockName);
-	if (!ClockManager)
+	Audio::FQuartzClockProxy const* Proxy = FindProxyByName(ClockName);
+	if(Proxy && Proxy->IsValid())
 	{
-		return false;
+		return Proxy->DoesClockExist();
 	}
 
-	return ClockManager->DoesClockExist(ClockName);
+	return {};
 }
 
 bool UQuartzSubsystem::IsClockRunning(const UObject* WorldContextObject, FName ClockName)
 {
-	Audio::FQuartzClockManager* ClockManager = GetManagerForClock(WorldContextObject, ClockName);
-	if (!ClockManager)
+	Audio::FQuartzClockProxy const* Proxy = FindProxyByName(ClockName);
+	if(Proxy && Proxy->IsValid())
 	{
-		return false;
+		return Proxy->IsClockRunning();
 	}
 
-	return ClockManager->IsClockRunning(ClockName);
+	return {};
 }
 
 float UQuartzSubsystem::GetDurationOfQuantizationTypeInSeconds(const UObject* WorldContextObject, FName ClockName, const EQuartzCommandQuantization& QuantizationType, float Multiplier)
 {
-	Audio::FQuartzClockManager* ClockManager = GetManagerForClock(WorldContextObject, ClockName);
-	if (!ClockManager)
+	Audio::FQuartzClockProxy const* Proxy = FindProxyByName(ClockName);
+	if(Proxy && Proxy->IsValid())
 	{
-		return INDEX_NONE;
+		return Proxy->GetDurationOfQuantizationTypeInSeconds(QuantizationType, Multiplier);
 	}
 
-	return ClockManager->GetDurationOfQuantizationTypeInSeconds(ClockName, QuantizationType, Multiplier);
+	return {};
 }
 
 FQuartzTransportTimeStamp UQuartzSubsystem::GetCurrentClockTimestamp(const UObject* WorldContextObject, const FName& InClockName)
 {
-	Audio::FQuartzClockManager* ClockManager = GetManagerForClock(WorldContextObject, InClockName);
-	if (!ClockManager)
+	Audio::FQuartzClockProxy const* Proxy = FindProxyByName(InClockName);
+	if(Proxy && Proxy->IsValid())
 	{
-		return FQuartzTransportTimeStamp();
+		return Proxy->GetCurrentClockTimestamp();
 	}
 
-	return ClockManager->GetCurrentTimestamp(InClockName);
+	return {};
 }
 
 float UQuartzSubsystem::GetEstimatedClockRunTime(const UObject* WorldContextObject, const FName& InClockName)
 {
-	Audio::FQuartzClockManager* ClockManager = GetManagerForClock(WorldContextObject, InClockName);
-	if (!ClockManager)
+	Audio::FQuartzClockProxy const* Proxy = FindProxyByName(InClockName);
+	if(Proxy && Proxy->IsValid())
 	{
-		return INDEX_NONE;
+		return Proxy->GetEstimatedClockRunTimeSeconds();
 	}
 
-	return ClockManager->GetEstimatedRunTime(InClockName);
+	return {};
 }
 
+// todo: move FQuartLatencyTracker off the AudioMixerClockManager? (GameThread->AudioRenderThread tracking)
 float UQuartzSubsystem::GetGameThreadToAudioRenderThreadAverageLatency(const UObject* WorldContextObject)
 {
-	Audio::FQuartzClockManager* ClockManager = GetManagerForClock(WorldContextObject);
+	Audio::FQuartzClockManager* ClockManager = GetClockManager(WorldContextObject);
 	if (!ClockManager)
 	{
 		return { };
@@ -477,7 +527,7 @@ float UQuartzSubsystem::GetGameThreadToAudioRenderThreadAverageLatency(const UOb
 
 float UQuartzSubsystem::GetGameThreadToAudioRenderThreadMinLatency(const UObject* WorldContextObject)
 {
-	Audio::FQuartzClockManager* ClockManager = GetManagerForClock(WorldContextObject);
+	Audio::FQuartzClockManager* ClockManager = GetClockManager(WorldContextObject);
 	if (!ClockManager)
 	{
 		return { };
@@ -488,7 +538,7 @@ float UQuartzSubsystem::GetGameThreadToAudioRenderThreadMinLatency(const UObject
 
 float UQuartzSubsystem::GetGameThreadToAudioRenderThreadMaxLatency(const UObject* WorldContextObject)
 {
-	Audio::FQuartzClockManager* ClockManager = GetManagerForClock(WorldContextObject);
+	Audio::FQuartzClockManager* ClockManager = GetClockManager(WorldContextObject);
 	if (!ClockManager)
 	{
 		return { };
@@ -499,19 +549,22 @@ float UQuartzSubsystem::GetGameThreadToAudioRenderThreadMaxLatency(const UObject
 
 float UQuartzSubsystem::GetAudioRenderThreadToGameThreadAverageLatency()
 {
-	return GetLifetimeAverageLatency();
+	ensure(TickableObjectManagerPtr.IsValid());
+	return TickableObjectManagerPtr->GetLifetimeAverageLatency();
 }
 
 
 float UQuartzSubsystem::GetAudioRenderThreadToGameThreadMinLatency()
 {
-	return GetMinLatency();
+	ensure(TickableObjectManagerPtr.IsValid());
+	return TickableObjectManagerPtr->GetMinLatency();
 }
 
 
 float UQuartzSubsystem::GetAudioRenderThreadToGameThreadMaxLatency()
 {
-	return GetMaxLatency();
+	ensure(TickableObjectManagerPtr.IsValid());
+	return TickableObjectManagerPtr->GetMaxLatency();
 }
 
 
@@ -536,50 +589,54 @@ float UQuartzSubsystem::GetRoundTripMaxLatency(const UObject* WorldContextObject
 
 void UQuartzSubsystem::AddCommandToClock(const UObject* WorldContextObject, Audio::FQuartzQuantizedCommandInitInfo& InQuantizationCommandInitInfo, FName ClockName)
 {
-	Audio::FQuartzClockManager* ClockManager = GetManagerForClock(WorldContextObject, ClockName);
-	if (!ClockManager)
+	Audio::FQuartzClockProxy* Proxy = FindProxyByName(ClockName);
+	if(Proxy && Proxy->IsValid())
 	{
-		return;
+		Proxy->SendCommandToClock([=](Audio::FQuartzClock* InClock) mutable{ InClock->AddQuantizedCommand(InQuantizationCommandInitInfo); });
 	}
+}
 
-	ClockManager->AddCommandToClock(InQuantizationCommandInitInfo);
+TWeakPtr<FQuartzTickableObjectsManager> UQuartzSubsystem::GetTickableObjectManager() const
+{
+	return TickableObjectManagerPtr;
 }
 
+void UQuartzSubsystem::PruneStaleProxies()
+{
+	PruneStaleProxiesInternal(ActiveExternalClockProxies);
+	PruneStaleProxiesInternal(ActiveAudioMixerClockProxies);
+}
 
-Audio::FQuartzClockManager* UQuartzSubsystem::GetManagerForClock(const UObject* WorldContextObject, FName ExistingClockName)
+void UQuartzSubsystem::PruneStaleProxiesInternal(TArray<Audio::FQuartzClockProxy>& ContainerToPrune)
 {
-	// if the enum has changed, the logic in this function needs to be updated
-	ensure((int32)EQuarztClockManagerType::Count == 2);
-	if (!WorldContextObject)
-	{
-		// Do not mutate the ClockManagerTypeMap if we do not have a world context yet
-		return nullptr;
-	}
+	 for(int32 i = 0; i < ContainerToPrune.Num(); ++i)
+	 {
+		 if(ContainerToPrune[i].IsValid() == false)
+		 {
+		 	ContainerToPrune.RemoveAtSwap(i--, 1, false);
+		 }
+	 }
+}
 
-	Audio::FMixerDevice* MixerDevice = nullptr;
-	if (!SimulateNoAudioDeviceCvar)
-	{
-		MixerDevice = GetAudioMixerDeviceUsingWorldContext(WorldContextObject);
-	}
+Audio::FQuartzClockProxy* UQuartzSubsystem::FindProxyByName(const FName& ClockName)
+{
+	Audio::FQuartzClockProxy* Result = ActiveAudioMixerClockProxies.FindByKey((ClockName));
 
-	// see if this clock prefers the subsystem clock be used
-	if (ClockManagerTypeMap.Contains(ExistingClockName) && (ClockManagerTypeMap[ExistingClockName] == EQuarztClockManagerType::QuartzSubsystem))
+	if(!Result)
 	{
-		return &SubsystemClockManager;
+		Result = ActiveAudioMixerClockProxies.FindByKey(ClockName);
 	}
-	else // otherwise 
-	{
-		// if the clock doesn't exist in our type map, try to use the mixer device's manager
-		if (MixerDevice)
-		{
-			ClockManagerTypeMap.Add(ExistingClockName, EQuarztClockManagerType::AudioEngine);
-			return &MixerDevice->QuantizedEventClockManager;
-		}
 
-		// no mixer device, use the subsystem's manager
-		ClockManagerTypeMap.Add(ExistingClockName, EQuarztClockManagerType::QuartzSubsystem);
-		return &SubsystemClockManager;
+	return Result;
+}
+
+Audio::FQuartzClockProxy const* UQuartzSubsystem::FindProxyByName(const FName& ClockName) const
+{
+	Audio::FQuartzClockProxy const* Result = Result = ActiveAudioMixerClockProxies.FindByKey(ClockName);
+	if(!Result)
+	{
+		Result = ActiveAudioMixerClockProxies.FindByKey(ClockName);
 	}
 
-	return nullptr;
+	return Result;
 }
diff --git a/Engine/Source/Runtime/AudioMixer/Public/Quartz/AudioMixerClock.h b/Engine/Source/Runtime/AudioMixer/Public/Quartz/AudioMixerClock.h
--- a/Engine/Source/Runtime/AudioMixer/Public/Quartz/AudioMixerClock.h
+++ b/Engine/Source/Runtime/AudioMixer/Public/Quartz/AudioMixerClock.h
@@ -10,27 +10,97 @@ namespace Audio
 {
 	// forwards
 	class FMixerDevice;
+	class FQuartzClock;
 	class FMixerSourceManager;
 	class FQuartzClockManager;
 
-	// Contains the pending command and the number of frames it has to wait to fire
-	struct PendingCommand
+	template<class ListenerType>
+	class TQuartzShareableCommandQueue;
+
+	using FQuartzClockCommandQueuePtr = TSharedPtr<TQuartzShareableCommandQueue<FQuartzClock>, ESPMode::ThreadSafe>;
+	using FQuartzClockCommandQueueWeakPtr = TWeakPtr<TQuartzShareableCommandQueue<FQuartzClock>, ESPMode::ThreadSafe>;
+
+	
+	/**
+	 *	FQuartzClockProxy:
+	 *
+	 *		This class is a C++ handle to the underlying clock.
+	 *		
+	 *		It is mostly a wrapper around a TWeakPtr<FQuartzClock> and
+	 *		TSharedPtr<TQuartzShareableCommandQueue<FQuartzClock>.
+	 *		
+	 *		The getters query the underlying FQuartzClock directly,
+	 *		which returns values updated during the last audio-engine tick
+	 *
+	 *		If you need to add more getters, add copies of the members in question to
+	 *		FQuartzClock::FQuartzClockState and update FQuartzClock::UpdateCachedState()
+	 *		for thread-safe access (or manually protect access w/ CachedClockStateCritSec)
+	 *
+	 *		SendCommandToClock() can be used to execute lambdas at the beginning
+	 *		of the next clock tick.  These lambdas can call FQuartzClock's public methods safely.
+	 *
+	 *		Your lambda will take an FQuartzClock* as an argument, which will be passed in by the
+	 *		FQuartzClock itself when it pumps the command queue.
+	 *
+	 */
+	class AUDIOMIXER_API FQuartzClockProxy
 	{
+	public:
 		// ctor
-		PendingCommand(TSharedPtr<IQuartzQuantizedCommand> InCommand, int32 InNumFramesUntilExec)
-		: Command(InCommand)
-		, NumFramesUntilExec(InNumFramesUntilExec)
-		{
-		}
+		FQuartzClockProxy() {}
+		FQuartzClockProxy(const FName& Name) : ClockId(Name){ } // conv ctor from FName
+		FQuartzClockProxy(TSharedPtr<FQuartzClock, ESPMode::ThreadSafe> InClock);
+
+		FName GetClockName() const { return ClockId; }
+
+		bool IsValid() const;
+		operator bool() const { return IsValid(); }
+
+		bool operator==(const FName& Name) const { return ClockId == Name; }
+
+		bool DoesClockExist() const;
+
+		bool IsClockRunning() const;
 
-		// Quantized Command Object
-		TSharedPtr<IQuartzQuantizedCommand> Command;
+		Audio::FQuartzClockTickRate GetTickRate() const;
 
-		// Countdown to execution
-		int32 NumFramesUntilExec{ 0 };
-	}; // struct PendingCommand
+		float GetEstimatedClockRunTimeSeconds() const;
 
-	// Class that encapsulates sample-accurate timing logic, as well as firing QuantizedAudioCommands
+		FQuartzTransportTimeStamp GetCurrentClockTimestamp() const;
+
+		float GetDurationOfQuantizationTypeInSeconds(const EQuartzCommandQuantization& QuantizationType, float Multiplier) const;
+
+		// returns false if the clock is not valid or has shut down
+		bool SendCommandToClock(TFunction<void(FQuartzClock*)> InCommand);
+
+		// implicit cast to underlying ID (FName)
+		operator const FName&() const { return ClockId; }
+
+	private:
+		FName ClockId;
+
+		FQuartzClockCommandQueueWeakPtr SharedQueue;
+
+	protected:
+		TWeakPtr<FQuartzClock, ESPMode::ThreadSafe> ClockWeakPtr;
+
+	}; // class FQuartzClockProxy
+
+
+	
+	/**
+	 *	FQuartzClock:
+	 *
+	 *		This class receives, schedules, and fires quantized commands. 
+	 *		The underlying FQuartzMetronome handles all counting / timing logic.
+	 *
+	 *		This class gets ticked externally (i.e. by some Clock Manager)
+	 *		and counts down the time-to-fire the commands in audio frames.
+	 *
+	 *
+	 *		UpdateCachedState() updates a game-thread copy of data accessed via FQuartzClockProxy
+	 *		(see FQuartzClockState)
+	 */
 	class AUDIOMIXER_API FQuartzClock
 	{
 	public:
@@ -41,79 +111,51 @@ namespace Audio
 		// dtor
 		~FQuartzClock();
 
+		// Transport Control:
 		// alter the tick rate (take by-value to make sample-rate adjustments in-place)
 		void ChangeTickRate(FQuartzClockTickRate InNewTickRate, int32 NumFramesLeft = 0);
 
-		// alter the time signature
 		void ChangeTimeSignature(const FQuartzTimeSignature& InNewTimeSignature);
 
-		// start ticking the clock
 		void Resume();
 
-		// stop ticking and reset the clock
-		void Stop(bool CancelPendingEvents);
-
-		// stop ticking the clock
 		void Pause();
 
-		// reset the metronome
 		void Restart(bool bPause = true);
 
-		// shutdown
-		void Shutdown();
-
-		// low-resolution clock update
-		// (not sample-accurate!, useful when running without an Audio Device)
-		void LowResolutionTick(float InDeltaTimeSeconds);
-
-		// sample accurate clock update
-		void Tick(int32 InNumFramesUntilNextTick);
+		void Stop(bool CancelPendingEvents); // Pause + Restart
 
-		// Set the sample rate of the clock
 		void SetSampleRate(float InNewSampleRate);
 
-		// get the tick rate
-		FQuartzClockTickRate GetTickRate() const { return Metronome.GetTickRate(); }
-
-		// get the identifier of the clock
-		FName GetName() const { return Name; }
-
-		// clock will persist across level changes
-		bool IgnoresFlush();
+		void ResetTransport();
 
-		// Does this clock match the provided settings
-		bool DoesMatchSettings(const FQuartzClockSettings& InClockSettings) const;
+		// (used for StartOtherClock command to handle the sub-tick as the target clock)
+		void AddToTickDelay(int32 NumFramesOfDelayToAdd);
 
-		void SubscribeToTimeDivision(MetronomeCommandQueuePtr InListenerQueue, EQuartzCommandQuantization InQuantizationBoundary);
+		// (used for StartOtherClock command to handle the sub-tick as the target clock)
+		void SetTickDelay(int32 NumFramesOfDelay);
 
-		void SubscribeToAllTimeDivisions(MetronomeCommandQueuePtr InListenerQueue);
+		void Shutdown();
 
-		void UnsubscribeFromTimeDivision(MetronomeCommandQueuePtr InListenerQueue, EQuartzCommandQuantization InQuantizationBoundary);
+		// Getters:
+		FQuartzClockTickRate GetTickRate();
 
-		void UnsubscribeFromAllTimeDivisions(MetronomeCommandQueuePtr InListenerQueue);
+		FName GetName() const;
 
-		// Add a new event to be triggered by this clock
-		// TODO: return a handle to this event so "looping" events can be canceled
-		void AddQuantizedCommand(FQuartzQuantizationBoundary InQuantizationBondary, TSharedPtr<IQuartzQuantizedCommand> InNewEvent);
+		bool IgnoresFlush() const;
 
-		// Cancel pending command
-		bool CancelQuantizedCommand(TSharedPtr<IQuartzQuantizedCommand> InCommandPtr);
+		bool DoesMatchSettings(const FQuartzClockSettings& InClockSettings) const;
 
-		// does the clock have any pending events
 		bool HasPendingEvents() const;
 
 		int32 NumPendingEvents() const;
 
-		// is the clock currently ticking?
-		bool IsRunning();
+		bool IsRunning() const;
 
-		// Returns the duration in seconds of the given Quantization Type
 		float GetDurationOfQuantizationTypeInSeconds(const EQuartzCommandQuantization& QuantizationType, float Multiplier);
 
-		// Returns the current location of the clock in the transport
 		FQuartzTransportTimeStamp GetCurrentTimestamp();
 
-		// Returns the amount of time, in seconds, the clock has been running. Caution: due to latency, this will not be perfectly accurate
 		float GetEstimatedRunTime();
 
 		FMixerDevice* GetMixerDevice();
@@ -122,19 +164,61 @@ namespace Audio
 
 		FQuartzClockManager* GetClockManager();
 
-		void ResetTransport();
+		FQuartzClockCommandQueueWeakPtr GetCommandQueue() const;
 
-		void AddToTickDelay(int32 NumFramesOfDelayToAdd)
-		{
-			TickDelayLengthInFrames += NumFramesOfDelayToAdd;
-		}
+		// Metronome Event Subscription:
+		void SubscribeToTimeDivision(FQuartzGameThreadSubscriber InSubscriber, EQuartzCommandQuantization InQuantizationBoundary);
 
-		void SetTickDelay(int32 NumFramesOfDelay)
-		{
-			TickDelayLengthInFrames = NumFramesOfDelay;
-		}
+		void SubscribeToAllTimeDivisions(FQuartzGameThreadSubscriber InSubscriber);
+
+		void UnsubscribeFromTimeDivision(FQuartzGameThreadSubscriber InSubscriber, EQuartzCommandQuantization InQuantizationBoundary);
+
+		void UnsubscribeFromAllTimeDivisions(FQuartzGameThreadSubscriber InSubscriber);
+
+		// Quantized Command Management:
+		void AddQuantizedCommand(FQuartzQuantizedCommandInitInfo& InQuantizationCommandInitInfo);
+
+		void AddQuantizedCommand(FQuartzQuantizationBoundary InQuantizationBondary, TSharedPtr<IQuartzQuantizedCommand> InNewEvent);
+
+		bool CancelQuantizedCommand(TSharedPtr<IQuartzQuantizedCommand> InCommandPtr);
+		
+		// low-resolution clock update
+		// (not sample-accurate!, useful when running without an Audio Device)
+		void LowResolutionTick(float InDeltaTimeSeconds);
+
+		// sample accurate clock update
+		void Tick(int32 InNumFramesUntilNextTick);
 
 	private:
+		// Contains the pending command and the number of frames it has to wait to fire
+		struct PendingCommand
+		{
+			// ctor
+			PendingCommand(TSharedPtr<IQuartzQuantizedCommand> InCommand, int32 InNumFramesUntilExec)
+				: Command(InCommand)
+				, NumFramesUntilExec(InNumFramesUntilExec)
+			{
+			}
+
+			// Quantized Command Object
+			TSharedPtr<IQuartzQuantizedCommand> Command;
+
+			// Countdown to execution
+			int32 NumFramesUntilExec{ 0 };
+		}; // struct PendingCommand
+
+		// mutex-protected update at the end of Tick()
+		FCriticalSection CachedClockStateCritSec;
+		void UpdateCachedState();
+
+		// data is cached when an FQuartzClock is ticked
+		struct FQuartzClockState
+		{
+			FQuartzClockTickRate TickRate;
+			FQuartzTransportTimeStamp TimeStamp;
+			float RunTimeInSeconds;
+		} CachedClockState;
+
 		void TickInternal(int32 InNumFramesUntilNextTick, TArray<PendingCommand>& CommandsToTick, int32 FramesOfLatency = 0, int32 FramesOfDelay = 0);
 
 		bool CancelQuantizedCommandInternal(TSharedPtr<IQuartzQuantizedCommand> InCommandPtr, TArray<PendingCommand>& CommandsToTick);
@@ -149,9 +233,11 @@ namespace Audio
 
 		FName Name;
 
-		// TODO: Make this configurable
 		float ThreadLatencyInMilliseconds{ 40.f };
 
+		// Command queue handed out to GameThread objects to queue commands. These get executed at the top of Tick()
+		mutable FQuartzClockCommandQueuePtr PreTickCommands; // (mutable for lazy init in GetQuartzSubscriber())
+
 		// Container of external commands to be executed (TUniquePointer<QuantizedAudioCommand>)
 		TArray<PendingCommand> ClockAlteringPendingCommands;
 		TArray<PendingCommand> PendingCommands;
diff --git a/Engine/Source/Runtime/AudioMixer/Public/Quartz/AudioMixerClockHandle.h b/Engine/Source/Runtime/AudioMixer/Public/Quartz/AudioMixerClockHandle.h
--- a/Engine/Source/Runtime/AudioMixer/Public/Quartz/AudioMixerClockHandle.h
+++ b/Engine/Source/Runtime/AudioMixer/Public/Quartz/AudioMixerClockHandle.h
@@ -8,85 +8,22 @@
 #include "Sound/QuartzSubscription.h"
 #include "Quartz/QuartzSubsystem.h"
 #include "Quartz/QuartzMetronome.h"
-#include "UObject/GCObject.h"
 
 #include "AudioMixerClockHandle.generated.h"
 
 
-class UAudioComponent;
 
-// Children of this class MUST be UObjects and they must call QuartzAddReferencedObjects from their own override of AddReferencedObjects.
-// TODO: Ideally this should use CRTP to put the fields on the child or something to avoid the possibility of not calling it.
-class AUDIOMIXER_API FQuartzTickableObject
-{
-	struct AUDIOMIXER_API FQuartzTickableObjectGCObjectMembers
-	{
-	public:
-		UQuartzSubsystem* QuartzSubsystem;
-		UWorld* WorldPtr{ nullptr };
-
-		void AddReferencedObjects(FReferenceCollector& Collector);
-	};
-
-	public:
-		FQuartzTickableObject() {}
-		virtual ~FQuartzTickableObject();
-
-		FQuartzTickableObject* Init(UWorld* InWorldPtr);
-
-		// called by the associated QuartzSubsystem
-		void QuartzTick(float DeltaTime);
-		bool QuartzIsTickable() const;
-		UWorld* QuartzGetWorld() const { return GCObjectMembers.WorldPtr; }
-
-		bool IsInitialized() const { return bHasBeenInitialized; }
-
-		// access to the associated QuartzSubsystem
-		UQuartzSubsystem* GetQuartzSubsystem() const { return GCObjectMembers.QuartzSubsystem; }
-
-		TSharedPtr<Audio::FShareableQuartzCommandQueue, ESPMode::ThreadSafe> GetCommandQueue();
-
-		int32 AddCommandDelegate(const FOnQuartzCommandEventBP& InDelegate, TArray<TSharedPtr<Audio::FShareableQuartzCommandQueue, ESPMode::ThreadSafe>>& TargetSubscriberArray);
-
-		// virtual interface
-		virtual void ProcessCommand(const Audio::FQuartzQuantizedCommandDelegateData& Data) {};
-
-		virtual void ProcessCommand(const Audio::FQuartzMetronomeDelegateData& Data) {};
-
-		virtual void ProcessCommand(const Audio::FQuartzQueueCommandData& Data) {};
-
-		void Shutdown();
-
-		void QuartzAddReferencedObjects(FReferenceCollector& Collector)
-		{
-			GCObjectMembers.AddReferencedObjects(Collector);
-		}
-
-	protected:
-		struct CommandDelegateGameThreadData
-		{
-			FOnQuartzCommandEvent MulticastDelegate;
-			FThreadSafeCounter RefCount;
-		};
-
-		struct MetronomeDelegateGameThreadData
-		{
-			FOnQuartzMetronomeEvent MulticastDelegate;
-		};
-	
-		void PumpCommandQueue();
-
-		TSharedPtr<Audio::FShareableQuartzCommandQueue, ESPMode::ThreadSafe> CommandQueuePtr;
-		TArray<TFunction<void(FQuartzTickableObject*)>> TempCommandQueue;
-
-		TArray<CommandDelegateGameThreadData> QuantizedCommandDelegates;
-		MetronomeDelegateGameThreadData MetronomeDelegates[static_cast<int32>(EQuartzCommandQuantization::Count)];
-
-	private:
-		FQuartzTickableObjectGCObjectMembers GCObjectMembers;
-		bool bHasBeenInitialized = false;
-};
 
+/**
+ *  This class is a BP / Game thread wrapper around FQuartzClockProxy
+ *	(to talk to the underlying clock)
+ 
+ *  ...and inherits from FQuartzTickableObject
+ *	(to listen to the underlying clock)
+ *  
+ *  It can subscribe to Quantized Event & Metronome delegates to synchronize
+ *  gameplay & VFX to Quartz events fired from the Audio Engine
+ */
 UCLASS(BlueprintType, Blueprintable, Transient, ClassGroup = Quartz, meta = (BlueprintSpawnableComponent))
 class AUDIOMIXER_API UQuartzClockHandle : public UObject, public FQuartzTickableObject
 {
@@ -103,7 +40,6 @@ public:
 	void BeginDestroy() override;
 	// end UObject interface
 
-
 // Begin Blueprint Interface
 
 	// Clock manipulation
@@ -178,7 +114,6 @@ public:
 	UFUNCTION(BlueprintCallable, Category = "Quantization", meta = (WorldContext = "WorldContextObject", AdvancedDisplay = "QuantizationBoundary, Delegate", AutoCreateRefTerm = "QuantizationBoundary, Delegate", Keywords = "BPM, Tempo"))
 	void SetBeatsPerMinute(const UObject* WorldContextObject, UPARAM(ref) const FQuartzQuantizationBoundary& QuantizationBoundary, const FOnQuartzCommandEventBP& Delegate, UQuartzClockHandle*& ClockHandle, float BeatsPerMinute = 60.f);
 
-
 	// Metronome getters
 	UFUNCTION(BlueprintCallable, Category = "Quantization", meta = (WorldContext = "WorldContextObject", AutoCreateRefTerm = "InDelegate", Keywords = "BPM, Tempo"))
 	float GetMillisecondsPerTick(const UObject* WorldContextObject) const;
@@ -195,42 +130,41 @@ public:
 	UFUNCTION(BlueprintCallable, Category = "Quantization", meta = (WorldContext = "WorldContextObject", AutoCreateRefTerm = "InDelegate", Keywords = "BPM, Tempo"))
 	float GetBeatsPerMinute(const UObject* WorldContextObject) const;
 
+	// todo: un-comment when metronome events support the offset
+	// Set how early we would like to receive Metronome (not yet supported) and "About To Start" Delegates. (all other command delegates will execute as normal)
+	// UFUNCTION(BlueprintCallable, Category = "Quantization", meta = (WorldContext = "WorldContextObject", AutoCreateRefTerm = "InDelegate", Keywords = "BPM, Tempo"))
+	// void SetNotificationAnticipationAmountInMilliseconds(const UObject* WorldContextObject, UQuartzClockHandle*& ClockHandle, const double Milliseconds = 0.0);
+
+	// // Set how early we would like to receive Metronome (not yet supported) and "About To Start" Delegates. (all other command delegates will execute as normal)
+	// UFUNCTION(BlueprintCallable, Category = "Quantization", meta = (WorldContext = "WorldContextObject", Keywords = "BPM, Tempo"))
+	// void SetNotificationAnticipationAmountAsMusicalDuration(const UObject* WorldContextObject, UQuartzClockHandle*& ClockHandle, const EQuartzCommandQuantization MusicalDuration = EQuartzCommandQuantization::QuarterNote, const double Multiplier = 1.0);
+
 // End Blueprint Interface
 	void QueueQuantizedSound(const UObject* WorldContextObject, UQuartzClockHandle*& ClockHandle, const FAudioComponentCommandInfo& AudioComponentData, const FOnQuartzCommandEventBP& InDelegate, const FQuartzQuantizationBoundary& InTargetBoundary);
 
-	UQuartzClockHandle* SubscribeToClock(const UObject* WorldContextObject, FName ClockName);
+	UQuartzClockHandle* SubscribeToClock(const UObject* WorldContextObject, FName ClockName, Audio::FQuartzClockProxy const* InHandlePtr = nullptr);
 
 	FName GetClockName() const { return CurrentClockId; }
 
-	FName GetHandleName() const { return ClockHandleId; }
-	
 	bool DoesClockExist(const UObject* WorldContextObject) const
 	{
-		if (UQuartzSubsystem* QuartzSubsystem = GetQuartzSubsystem())
-		{
-			return QuartzSubsystem->DoesClockExist(WorldContextObject, CurrentClockId);
-		}
-
-		return false;
+		check(RawHandle);
+		return RawHandle.DoesClockExist();
 	}
 
-	virtual void ProcessCommand(const Audio::FQuartzQuantizedCommandDelegateData& Data) override;
+	UE_DEPRECATED(5.1, "This function should not be called directly, and the original functionality has been moved into FQuartzTickable")
+	virtual void ProcessCommand(const Audio::FQuartzQuantizedCommandDelegateData& Data) override {}
 
-	virtual void ProcessCommand(const Audio::FQuartzMetronomeDelegateData& Data) override;
+	UE_DEPRECATED(5.1, "This function should not be called directly, and the original functionality has been moved into FQuartzTickable")
+	virtual void ProcessCommand(const Audio::FQuartzMetronomeDelegateData& Data) override {};
 
 	bool GetCurrentTickRate(const UObject* WorldContextObject, Audio::FQuartzClockTickRate& OutTickRate) const;
 
-	static void AddReferencedObjects(UObject* InThis, FReferenceCollector& Collector) 
-	{
-		CastChecked<UQuartzClockHandle>(InThis)->QuartzAddReferencedObjects(Collector);
-	}
-
 private:
-	
-	FName CurrentClockId;
+	void SetTickRateInternal(const FQuartzQuantizationBoundary& InQuantizationBoundary, const FOnQuartzCommandEventBP& InDelegate, const Audio::FQuartzClockTickRate& NewTickRate);
 
-	FName ClockHandleId;
+	Audio::FQuartzClockProxy RawHandle;
 
-	bool bConnectedToClock{ false };
+	FName CurrentClockId;
 
 }; // class UQuartzClockHandle
diff --git a/Engine/Source/Runtime/AudioMixer/Public/Quartz/AudioMixerClockManager.h b/Engine/Source/Runtime/AudioMixer/Public/Quartz/AudioMixerClockManager.h
--- a/Engine/Source/Runtime/AudioMixer/Public/Quartz/AudioMixerClockManager.h
+++ b/Engine/Source/Runtime/AudioMixer/Public/Quartz/AudioMixerClockManager.h
@@ -9,6 +9,7 @@ namespace Audio
 {
 	// forwards
 	class FMixerDevice;
+	class FQuartzClock;
 
 	// Class that owns, updates, and provides access to all active clocks
 	// All methods are thread-safe. The method locks if it returns a value, and stages a command if it returns void
@@ -34,8 +35,8 @@ namespace Audio
 
 		// add (and take ownership of) a new clock
 		// safe to call from AudioThread (uses critical section)
-		TSharedPtr<FQuartzClock> GetOrCreateClock(const FName& InClockName, const FQuartzClockSettings& InClockSettings, bool bOverrideTickRateIfClockExists = false);
-
+		FQuartzClockProxy GetOrCreateClock(const FName& InClockName, const FQuartzClockSettings& InClockSettings, bool bOverrideTickRateIfClockExists = false);
+		FQuartzClockProxy GetClock(const FName& InClockName);
 
 		// returns true if a clock with the given name already exists.
 		bool DoesClockExist(const FName& InClockName);
@@ -85,6 +86,8 @@ namespace Audio
 		FQuartzQuantizedCommandHandle AddCommandToClock(FQuartzQuantizedCommandInitInfo& InQuantizationCommandInitInfo);
 
 		// subscribe to a specific time division on a clock
+		// TODO: update the metronome subscription functions to take an FQuartzGameThreadSubscriber instead of the Command queue ptr
+		// (to support metronome event offset)
 		void SubscribeToTimeDivision(FName InClockName, MetronomeCommandQueuePtr InListenerQueue, EQuartzCommandQuantization InQuantizationBoundary);
 
 		// subscribe to all time divisions on a clock
@@ -101,7 +104,7 @@ namespace Audio
 
 		bool HasClockBeenTickedThisUpdate(FName InClockName);
 
-		int32 GetLastUpdateSizeInFrames() { return LastUpdateSizeInFrames; }
+		int32 GetLastUpdateSizeInFrames() const { return LastUpdateSizeInFrames; }
 
 		// get access to the owning FMixerDevice
 		FMixerDevice* GetMixerDevice() const;
@@ -124,5 +127,8 @@ namespace Audio
 
 		FThreadSafeCounter LastClockTickedIndex{ 0 };
 		int32 LastUpdateSizeInFrames{ 0 };
+
+		// allow a clock that is queuing a command directly use FindClock() to retrieve the TSharedPtr<FQuartzClock>
+		friend void FQuartzClock::AddQuantizedCommand(FQuartzQuantizedCommandInitInfo& InQuantizationCommandInitInfo);
 	};
 } // namespace Audio
diff --git a/Engine/Source/Runtime/AudioMixer/Public/Quartz/QuartzMetronome.h b/Engine/Source/Runtime/AudioMixer/Public/Quartz/QuartzMetronome.h
--- a/Engine/Source/Runtime/AudioMixer/Public/Quartz/QuartzMetronome.h
+++ b/Engine/Source/Runtime/AudioMixer/Public/Quartz/QuartzMetronome.h
@@ -6,32 +6,33 @@
 
 namespace Audio
 {
-	using MetronomeCommandQueuePtr = TSharedPtr<FShareableQuartzCommandQueue, ESPMode::ThreadSafe>;
+	using MetronomeCommandQueuePtr = TSharedPtr<TQuartzShareableCommandQueue<FQuartzTickableObject>, ESPMode::ThreadSafe>;
 
 	// Class to track the passage of musical time, and allow subscribers to be notified when these musical events take place
 	class FQuartzMetronome
 	{
 	public:
 		// ctor
-		FQuartzMetronome();
-		FQuartzMetronome(const FQuartzTimeSignature& InTimeSignature);
+		FQuartzMetronome(FName InClockName = {});
+		FQuartzMetronome(const FQuartzTimeSignature& InTimeSignature, FName InClockName = {});
 
 		// dtor
 		~FQuartzMetronome();
 
-		// called by owning FQuartzClock
+		// Transport Control:
 		void Tick(int32 InNumSamples, int32 FramesOfLatency = 0);
 
-		// called by owning FQuartzClock
 		void SetTickRate(FQuartzClockTickRate InNewTickRate, int32 NumFramesLeft = 0);
 
-		FQuartzClockTickRate GetTickRate() const { return CurrentTickRate; }
-
 		void SetSampleRate(float InNewSampleRate);
 
-		// affects bars/beats values we send back to the game thread
 		void SetTimeSignature(const FQuartzTimeSignature& InNewTimeSignature);
 
+		void ResetTransport();
+
+		// Getters
+		const FQuartzClockTickRate& GetTickRate() const { return CurrentTickRate; }
+
 		double GetFramesUntilBoundary(FQuartzQuantizationBoundary InQuantizationBoundary) const;
 
 		const FQuartzTimeSignature & GetTimeSignature() const { return CurrentTimeSignature; }
@@ -40,6 +41,7 @@ namespace Audio
 
 		double GetTimeSinceStart() const { return TimeSinceStart; }
 
+		// Event Subscription
 		void SubscribeToTimeDivision(MetronomeCommandQueuePtr InListenerQueue, EQuartzCommandQuantization InQuantizationBoundary);
 
 		void SubscribeToAllTimeDivisions(MetronomeCommandQueuePtr InListenerQueue);
@@ -48,9 +50,10 @@ namespace Audio
 
 		void UnsubscribeFromAllTimeDivisions(MetronomeCommandQueuePtr InListenerQueue);
 
-		void ResetTransport();
 
 	private:
+
+		// Helpers:
 		void RecalculateDurations();
 
 		void FireEvents(int32 EventFlags);
@@ -69,7 +72,6 @@ namespace Audio
 
 		FQuartzClockTickRate CurrentTickRate;
 
-		// (Subtract one because Tick and 1/32 are the same.)
 		TArray<MetronomeCommandQueuePtr> MetronomeSubscriptionMatrix[static_cast<int32>(EQuartzCommandQuantization::Count)];
 		
 		// wrapper around our array so it can be indexed into by different Enums that represent musical time
@@ -118,5 +120,7 @@ namespace Audio
 		//Keeps track of time in seconds since the Clock was last reset
 		double TimeSinceStart;
 
+		FName ClockName;
+
 	}; // class QuartzMetronome
 } // namespace Audio
\ No newline at end of file
diff --git a/Engine/Source/Runtime/AudioMixer/Public/Quartz/QuartzSubsystem.h b/Engine/Source/Runtime/AudioMixer/Public/Quartz/QuartzSubsystem.h
--- a/Engine/Source/Runtime/AudioMixer/Public/Quartz/QuartzSubsystem.h
+++ b/Engine/Source/Runtime/AudioMixer/Public/Quartz/QuartzSubsystem.h
@@ -12,40 +12,53 @@
 
 #include "QuartzSubsystem.generated.h"
 
+
 // forwards
 namespace Audio
 {
 	class FMixerDevice;
 	class FQuartzClockManager;
-	class FShareableQuartzCommandQueue;
+
+	template<class ListenerType>
+	class TQuartzShareableCommandQueue;
 }
+
 class FQuartzTickableObject;
+class UQuartzClockHandle;
+using MetronomeCommandQueuePtr = TSharedPtr<Audio::TQuartzShareableCommandQueue<FQuartzTickableObject>, ESPMode::ThreadSafe>;
 
-using MetronomeCommandQueuePtr = TSharedPtr<Audio::FShareableQuartzCommandQueue, ESPMode::ThreadSafe>;
 
-// GetManagerForClock() logic will need to be updated if more entries are present
-UENUM(BlueprintType)
-enum class EQuarztClockManagerType : uint8
+
+struct AUDIOMIXER_API FQuartzTickableObjectsManager : public FQuartLatencyTracker
 {
-	AudioEngine		UMETA(DisplayName = "Audio Engine", ToolTip = "Sample-accurate clock managment by the audio renderer"),
-	QuartzSubsystem	UMETA(DisplayName = "Transport Relative", ToolTip = "Loose clock management by the Quartz Subsystem in UObjectTick.  (not sample-accurate. Used automatically when no Audio Device is present)"),
-	Count			UMETA(Hidden)
+public:
+	void Tick(float DeltaTime);
+	bool IsTickable() const;
+	void SubscribeToQuartzTick(FQuartzTickableObject* InObjectToTick);
+	void UnsubscribeFromQuartzTick(FQuartzTickableObject* InObjectToTick);
+
+private:
+	// list of objects needing to be ticked by Quartz
+	TArray<FQuartzTickableObject *> QuartzTickSubscribers;
+
+	// index to track the next subscriber to tick (if updates are being amortized across multiple UObject Ticks)
+	int32 UpdateIndex{ 0 };
 };
 
 
 UCLASS(DisplayName = "Quartz")
-class AUDIOMIXER_API UQuartzSubsystem : public UTickableWorldSubsystem, public FQuartLatencyTracker
+class AUDIOMIXER_API UQuartzSubsystem : public UTickableWorldSubsystem
 {
 	GENERATED_BODY()
 
 public:
 	// ctor/dtor
-	UQuartzSubsystem();
-	~UQuartzSubsystem();
+	UQuartzSubsystem() = default;
+	virtual ~UQuartzSubsystem() override = default;
 
 	//~ Begin UWorldSubsystem Interface
 	virtual bool DoesSupportWorldType(EWorldType::Type WorldType) const override;
-	void BeginDestroy() override;
+	void virtual BeginDestroy() override;
 	//~ End UWorldSubsystem Interface
 
 	//~ Begin FTickableGameObject Interface
@@ -54,21 +67,36 @@ public:
 	virtual TStatId GetStatId() const override;
 	//~ End FTickableGameObject Interface
 
-	// if we have another stakeholder later we can use polymorphism here.
-	// in the mean time, we are avoiding the virtual overhead
+	// these calls are forwarded to the internal FQuartzTickableObjectsManager
 	void SubscribeToQuartzTick(FQuartzTickableObject* InObjectToTick);
 	void UnsubscribeFromQuartzTick(FQuartzTickableObject* InObjectToTick);
 
-	// static methods
-	static UQuartzSubsystem* Get(UWorld* World);
+	// get C++ handle (proxy) to a clock if it exists
+	Audio::FQuartzClockProxy GetProxyForClock(FName ClockName) const;
+
+	// allow an external clock (not ticked by the Audio Mixer or QuartzSubsystem) to be accessible via this subsystem
+	void AddProxyForExternalClock(const Audio::FQuartzClockProxy& InProxy);
 
-	// create a new command queue to be shared between ClockHandles and other threads
-	static TSharedPtr<Audio::FShareableQuartzCommandQueue, ESPMode::ThreadSafe> CreateQuartzCommandQueue();
+	// static methods
+	static UQuartzSubsystem* Get(const UWorld* const World);
 
 	// Helper functions for initializing quantized command initialization struct (to consolidate eyesore)
+	static Audio::FQuartzQuantizedRequestData CreateRequestDataForTickRateChange(UQuartzClockHandle* InClockHandle, const FOnQuartzCommandEventBP& InDelegate, const Audio::FQuartzClockTickRate& InNewTickRate, const FQuartzQuantizationBoundary& InQuantizationBoundary);
+	static Audio::FQuartzQuantizedRequestData CreateRequestDataForTransportReset(UQuartzClockHandle* InClockHandle, const FQuartzQuantizationBoundary& InQuantizationBoundary, const FOnQuartzCommandEventBP& InDelegate);
+	static Audio::FQuartzQuantizedRequestData CreateRequestDataForStartOtherClock(UQuartzClockHandle* InClockHandle, FName InClockToStart, const FQuartzQuantizationBoundary& InQuantizationBoundary, const FOnQuartzCommandEventBP& InDelegate);
+	static Audio::FQuartzQuantizedRequestData CreateRequestDataForSchedulePlaySound(UQuartzClockHandle* InClockHandle, const FOnQuartzCommandEventBP& InDelegate, const FQuartzQuantizationBoundary& InQuantizationBoundary);
+
+	// DEPRECATED HELPERS: non-static versions of the above CreateDataFor...() functions
+	UE_DEPRECATED(5.1, "Use the static (CreateRequestDataFor) version of this function instead")
 	Audio::FQuartzQuantizedRequestData CreateDataForTickRateChange(UQuartzClockHandle* InClockHandle, const FOnQuartzCommandEventBP& InDelegate, const Audio::FQuartzClockTickRate& InNewTickRate, const FQuartzQuantizationBoundary& InQuantizationBoundary);
+
+	UE_DEPRECATED(5.1, "Use the static (CreateRequestDataFor) version of this function instead")
 	Audio::FQuartzQuantizedRequestData CreateDataForTransportReset(UQuartzClockHandle* InClockHandle, const FQuartzQuantizationBoundary& InQuantizationBoundary, const FOnQuartzCommandEventBP& InDelegate);
+	
+	UE_DEPRECATED(5.1, "Use the static (CreateRequestDataFor) version of this function instead")
 	Audio::FQuartzQuantizedRequestData CreateDataForStartOtherClock(UQuartzClockHandle* InClockHandle, FName InClockToStart, const FQuartzQuantizationBoundary& InQuantizationBoundary, const FOnQuartzCommandEventBP& InDelegate);
+
+	UE_DEPRECATED(5.1, "Use the static (CreateRequestDataFor) version of this function instead")
 	Audio::FQuartzQuantizedRequestData CreateDataDataForSchedulePlaySound(UQuartzClockHandle* InClockHandle, const FOnQuartzCommandEventBP& InDelegate, const FQuartzQuantizationBoundary& InQuantizationBoundary);
 
 	UFUNCTION(BlueprintCallable, Category = "Quartz Clock Handle")
@@ -99,12 +127,7 @@ public:
 	UFUNCTION(BlueprintCallable, Category = "Quartz Clock Handle", meta = (WorldContext = "WorldContextObject"))
 	bool IsClockRunning(const UObject* WorldContextObject, FName ClockName);
 
-	/** Returns the duration in seconds of the given Quantization Type
-	 * 
-	 * @param The Quantization type to measure
-	 * @param The quantity of the Quantization Type to calculate the time of
-	 * @return The duration, in seconds, of a multiplier amount of the Quantization Type, or -1 in the case the clock is invalid
-	 */
+	// Returns the duration in seconds of the given Quantization Type
 	UFUNCTION(BlueprintCallable, Category = "Quartz Clock Handle", meta = (WorldContext = "WorldContextObject"))
 	float GetDurationOfQuantizationTypeInSeconds(const UObject* WorldContextObject, FName ClockName, const EQuartzCommandQuantization& QuantizationType, float Multiplier = 1.0f);
 
@@ -146,21 +169,30 @@ public:
 	UFUNCTION(BlueprintCallable, Category = "Quartz Subsystem", meta = (WorldContext = "WorldContextObject"))
 	float GetRoundTripMaxLatency(const UObject* WorldContextObject);
 
+	UE_DEPRECATED(5.1, "Obtain and use a UQuartzClockHandle / FQuartzClockProxy instead")
 	void AddCommandToClock(const UObject* WorldContextObject, Audio::FQuartzQuantizedCommandInitInfo& InQuantizationCommandInitInfo, FName ClockName);
 
-	Audio::FQuartzClockManager* GetManagerForClock(const UObject* WorldContextObject, FName ExistingClockName = FName());
+	// sharable to allow non-UObjects to un-subscribe if the Subsystem is going to outlive them
+	TWeakPtr<FQuartzTickableObjectsManager> GetTickableObjectManager() const;
 
 private:
+	// deletes proxies to clocks that no longer exists
+	void PruneStaleProxies();
+	static void PruneStaleProxiesInternal(TArray<Audio::FQuartzClockProxy>& ContainerToPrune);
+
+	// internal clock manager for game-thread-ticked clocks
 	Audio::FQuartzClockManager SubsystemClockManager;
 
-	// list of objects needing to be ticked by Quartz
-	TArray<FQuartzTickableObject *> QuartzTickSubscribers;
+	// sharable tickable object manager to allow for non-UObject subscription / un-subscription
+	TSharedPtr<FQuartzTickableObjectsManager> TickableObjectManagerPtr { MakeShared<FQuartzTickableObjectsManager>() };
 
-	// index to track the next clock handle to tick (if updates are being amortized across multiple UObject Ticks)
-	int32 UpdateIndex{ 0 };
+	// array of active clock handles (update FindProxyByName() if more are added later)
+	TArray<Audio::FQuartzClockProxy> ActiveExternalClockProxies;
+	TArray<Audio::FQuartzClockProxy> ActiveAudioMixerClockProxies;
 
-	// Tracks which system is managing a clock (given that clock's FName)
-	// currently only this subsystem or the FMixerDevice
-	TMap<FName, EQuarztClockManagerType> ClockManagerTypeMap;
+	// helpers
+	Audio::FQuartzClockProxy* FindProxyByName(const FName& ClockName);
+	Audio::FQuartzClockProxy const* FindProxyByName(const FName& ClockName) const;
+	Audio::FQuartzClockManager* GetClockManager(const UObject* WorldContextObject, bool bUseAudioEngineClockManager = true);
 
-}; // class UQuartzGameSubsystem 
\ No newline at end of file
+}; // class UQuartzGameSubsystem
\ No newline at end of file
diff --git a/Engine/Source/Runtime/Engine/Classes/Components/AudioComponent.h b/Engine/Source/Runtime/Engine/Classes/Components/AudioComponent.h
--- a/Engine/Source/Runtime/Engine/Classes/Components/AudioComponent.h
+++ b/Engine/Source/Runtime/Engine/Classes/Components/AudioComponent.h
@@ -748,6 +748,7 @@ public:
 	virtual void BeginDestroy() override;
 	//~ End UObject Interface.
 
+	// QuartzTodo: delete this after fix-up
 	static void AddReferencedObjects(UObject* InThis, FReferenceCollector& Collector)
 	{
 		CastChecked<UAudioComponent>(InThis)->QuartzAddReferencedObjects(Collector);
diff --git a/Engine/Source/Runtime/Engine/Classes/Sound/QuartzQuantizationUtilities.cpp b/Engine/Source/Runtime/Engine/Classes/Sound/QuartzQuantizationUtilities.cpp
--- a/Engine/Source/Runtime/Engine/Classes/Sound/QuartzQuantizationUtilities.cpp
+++ b/Engine/Source/Runtime/Engine/Classes/Sound/QuartzQuantizationUtilities.cpp
@@ -9,7 +9,6 @@
 
 #define INVALID_DURATION -1
 
-
 DEFINE_LOG_CATEGORY(LogAudioQuartz);
 
 EQuartzCommandQuantization TimeSignatureQuantizationToCommandQuantization(const EQuartzTimeSignatureQuantization& BeatType)
@@ -18,23 +17,18 @@ EQuartzCommandQuantization TimeSignatureQuantizationToCommandQuantization(const
 	{
 		case EQuartzTimeSignatureQuantization::HalfNote :
 			return EQuartzCommandQuantization::HalfNote;
-			break;
 
 		case EQuartzTimeSignatureQuantization::QuarterNote :
 			return EQuartzCommandQuantization::QuarterNote;
-			break;
 
 		case EQuartzTimeSignatureQuantization::EighthNote :
 			return EQuartzCommandQuantization::EighthNote;
-			break;
 
 		case EQuartzTimeSignatureQuantization::SixteenthNote :
 			return EQuartzCommandQuantization::SixteenthNote;
-			break;
 
 		case EQuartzTimeSignatureQuantization::ThirtySecondNote :
 			return EQuartzCommandQuantization::ThirtySecondNote;
-			break;
 
 		default:
 			return EQuartzCommandQuantization::Count;
@@ -73,7 +67,7 @@ bool FQuartzTimeSignature::operator==(const FQuartzTimeSignature& Other) const
 			const bool NumPulsesMatch = (OptionalPulseOverride[i].NumberOfPulses == Other.OptionalPulseOverride[i].NumberOfPulses);
 			const bool DurationsMatch = (OptionalPulseOverride[i].PulseDuration == Other.OptionalPulseOverride[i].PulseDuration);
 
-			if (!(NumPulseEntries && DurationsMatch))
+			if (!(NumPulsesMatch && DurationsMatch))
 			{
 				Result = false;
 				break;
@@ -386,8 +380,6 @@ namespace Audio
 		, int32 InSourceID
 	)
 		: ClockName(RHS.ClockName)
-		, ClockHandleName(RHS.ClockHandleName)
-		, OtherClockName(RHS.OtherClockName)
 		, QuantizedCommandPtr(RHS.QuantizedCommandPtr)
 		, QuantizationBoundary(RHS.QuantizationBoundary)
 		, GameThreadSubscribers(RHS.GameThreadSubscribers)
@@ -404,7 +396,7 @@ namespace Audio
 		return nullptr;
 	}
 
-	void IQuartzQuantizedCommand::AddSubscriber(TSharedPtr<FShareableQuartzCommandQueue, ESPMode::ThreadSafe> InSubscriber)
+	void IQuartzQuantizedCommand::AddSubscriber(FQuartzGameThreadSubscriber InSubscriber)
 	{
 		GameThreadSubscribers.AddUnique(InSubscriber);
 	}
@@ -412,8 +404,8 @@ namespace Audio
 	void IQuartzQuantizedCommand::OnQueued(const FQuartzQuantizedCommandInitInfo& InCommandInitInfo)
 	{
 		TRACE_CPUPROFILER_EVENT_SCOPE(QuartzQuantizedCommand::OnQueued);
-		Audio::FMixerDevice* MixerDevice = InCommandInitInfo.OwningClockPointer->GetMixerDevice();
-		if (MixerDevice)
+
+		if (Audio::FMixerDevice* MixerDevice = InCommandInitInfo.OwningClockPointer->GetMixerDevice())
 		{
 			MixerDevice->QuantizedEventClockManager.PushLatencyTrackerResult(FQuartzCrossThreadMessage::RequestRecieved());
 		}
@@ -429,13 +421,9 @@ namespace Audio
 			Data.DelegateSubType = EQuartzCommandDelegateSubType::CommandOnQueued;
 			Data.DelegateID = GameThreadDelegateID;
 
-
-			for (auto& SubscriberQueue : GameThreadSubscribers)
+			for (auto& Subscriber : GameThreadSubscribers)
 			{
-				if (SubscriberQueue.IsValid())
-				{
-					SubscriberQueue->PushEvent(Data);
-				}
+				Subscriber.PushEvent(Data);
 			}
 		}
 
@@ -443,6 +431,35 @@ namespace Audio
 		OnQueuedCustom(InCommandInitInfo);
 	}
 
+	void IQuartzQuantizedCommand::OnScheduled(const FQuartzClockTickRate& InTickRate)
+	{
+		for(auto& Subscriber : GameThreadSubscribers)
+		{
+			Subscriber.FinalizeOffset(InTickRate);
+		}
+	}
+
+	void IQuartzQuantizedCommand::Update(int32 NumFramesUntilDeadline)
+	{
+		TRACE_CPUPROFILER_EVENT_SCOPE(QuartzQuantizedCommand::Countdown);
+
+		FQuartzQuantizedCommandDelegateData Data;
+		Data.CommandType = GetCommandType();
+		Data.DelegateSubType = EQuartzCommandDelegateSubType::CommandOnAboutToStart;
+		Data.DelegateID = GameThreadDelegateID;
+
+		for(auto& Subscriber : GameThreadSubscribers)
+		{
+			// we only want to send this notification to the subscriber once
+			const int32 NumFramesOfAnticipation = Subscriber.GetOffsetAsAudioFrames();
+			if(!Subscriber.HasBeenNotifiedOfAboutToStart()
+				&&  (NumFramesOfAnticipation >= NumFramesUntilDeadline))
+			{
+				Subscriber.PushEvent(Data);
+			}
+		}
+	}
+
 	void IQuartzQuantizedCommand::FailedToQueue(FQuartzQuantizedRequestData& InGameThreadData)
 	{
 		TRACE_CPUPROFILER_EVENT_SCOPE(QuartzQuantizedCommand::FailedToQueue);
@@ -457,13 +474,9 @@ namespace Audio
 			Data.DelegateSubType = EQuartzCommandDelegateSubType::CommandOnFailedToQueue;
 			Data.DelegateID = GameThreadDelegateID;
 
-
-			for (auto& SubscriberQueue : GameThreadSubscribers)
+			for (auto& Subscriber : GameThreadSubscribers)
 			{
-				if (SubscriberQueue.IsValid())
-				{
-					SubscriberQueue->PushEvent(Data);
-				}
+				Subscriber.PushEvent(Data);
 			}
 		}
 
@@ -474,29 +487,18 @@ namespace Audio
 	void IQuartzQuantizedCommand::AboutToStart()
 	{
 		TRACE_CPUPROFILER_EVENT_SCOPE(QuartzQuantizedCommand::AboutToStart);
-		// only call once for the lifespan of this event
-		if (bAboutToStartHasBeenCalled)
-		{
-			return;
-		}
-
-		bAboutToStartHasBeenCalled = true;
 
+		FQuartzQuantizedCommandDelegateData Data;
+		Data.CommandType = GetCommandType();
+		Data.DelegateSubType = EQuartzCommandDelegateSubType::CommandOnAboutToStart;
+		Data.DelegateID = GameThreadDelegateID;
 
-		if (GameThreadSubscribers.Num())
+		for(auto& Subscriber : GameThreadSubscribers)
 		{
-			FQuartzQuantizedCommandDelegateData Data;
-
-			Data.CommandType = GetCommandType();
-			Data.DelegateSubType = EQuartzCommandDelegateSubType::CommandOnAboutToStart;
-			Data.DelegateID = GameThreadDelegateID;
-
-			for (auto& SubscriberQueue : GameThreadSubscribers)
+			// we only want to send this notification to the subscriber once
+			if(!Subscriber.HasBeenNotifiedOfAboutToStart())
 			{
-				if (SubscriberQueue.IsValid())
-				{
-					SubscriberQueue->PushEvent(Data);
-				}
+				Subscriber.PushEvent(Data);
 			}
 		}
 
@@ -517,10 +519,7 @@ namespace Audio
 
 			for (auto& Subscriber : GameThreadSubscribers)
 			{
-				if (Subscriber.IsValid())
-				{
-					Subscriber->PushEvent(OnStartedData);
-				}
+				Subscriber.PushEvent(OnStartedData);
 			}
 		}
 
@@ -551,13 +550,9 @@ namespace Audio
 		Data.DelegateSubType = EQuartzCommandDelegateSubType::CommandOnCanceled;
 		Data.DelegateID = GameThreadDelegateID;
 
-
-		for (auto& SubscriberQueue : GameThreadSubscribers)
+		for (auto& Subscriber : GameThreadSubscribers)
 		{
-			if (SubscriberQueue.IsValid())
-			{
-				SubscriberQueue->PushEvent(Data);
-			}
+			Subscriber.PushEvent(Data);
 		}
 
 		UE_LOG(LogAudioQuartz, Verbose, TEXT("Cancel() called for quantized event type: [%s]"), *GetCommandName().ToString());
@@ -565,7 +560,6 @@ namespace Audio
 	}
 
 
-
 	bool FQuartzQuantizedCommandHandle::Cancel()
 	{
 		checkSlow(MixerDevice);
@@ -679,7 +673,59 @@ namespace Audio
 		return Timer.GetCurrentTimePassedMs();
 	}
 
+	bool FQuartzOffset::operator==(const FQuartzOffset& Other) const
+	{
+		return OffsetInMilliseconds == Other.OffsetInMilliseconds
+			&& OffsetAsDuration == Other.OffsetAsDuration;
+	}
+
+	bool FQuartzGameThreadSubscriber::operator==(const FQuartzGameThreadSubscriber& Other) const
+	{
+		return Offset == Other.Offset
+			&& Queue == Other.Queue;
+	}
+
+	void FQuartzGameThreadSubscriber::PushEvent(const FQuartzQuantizedCommandDelegateData& Data)
+	{
+		if(ensure(Queue.IsValid()))
+		{
+			Queue->PushEvent(Data);
+
+			// raise the flag if this was a CommandOnAboutToStart notification
+			if(!bHasBeenNotifiedOfAboutToStart)
+			{
+				bHasBeenNotifiedOfAboutToStart = (Data.DelegateSubType == EQuartzCommandDelegateSubType::CommandOnAboutToStart);
+			}
+		}
+	}
+
+	void FQuartzGameThreadSubscriber::PushEvent(const FQuartzMetronomeDelegateData& Data)
+	{
+		if(ensure(Queue.IsValid()))
+		{
+			Queue->PushEvent(Data);
+		}
+	}
+
+	void FQuartzGameThreadSubscriber::PushEvent(const FQuartzQueueCommandData& Data)
+	{
+		if(ensure(Queue.IsValid()))
+		{
+			Queue->PushEvent(Data);
+		}
+	}
 
+	int32 FQuartzGameThreadSubscriber::FinalizeOffset(const FQuartzClockTickRate& TickRate)
+	{
+		bOffsetConvertedToFrames = true;
+		return OffsetInAudioFrames = Offset.GetOffsetInAudioFrames(TickRate);
+	}
+
+	int32 FQuartzGameThreadSubscriber::GetOffsetAsAudioFrames() const
+	{
+		ensureAlwaysMsgf(bOffsetConvertedToFrames, TEXT("FinalizeOffset must be called before calling GetOffsetAsAudioFrames()"));
+		return OffsetInAudioFrames;
+	}
 } // namespace Audio
 
 bool FQuartzTransportTimeStamp::IsZero() const
diff --git a/Engine/Source/Runtime/Engine/Classes/Sound/QuartzQuantizationUtilities.h b/Engine/Source/Runtime/Engine/Classes/Sound/QuartzQuantizationUtilities.h
--- a/Engine/Source/Runtime/Engine/Classes/Sound/QuartzQuantizationUtilities.h
+++ b/Engine/Source/Runtime/Engine/Classes/Sound/QuartzQuantizationUtilities.h
@@ -12,6 +12,7 @@
 
 #include "QuartzQuantizationUtilities.generated.h"
 
+class FQuartzTickableObject;
 ENGINE_API DECLARE_LOG_CATEGORY_EXTERN(LogAudioQuartz, Log, All);
 
 // forwards
@@ -29,14 +30,17 @@
 	class FQuartzClock;
 	class FShareableQuartzCommandQueue;
 
+	template<typename T>
+	class TQuartzShareableCommandQueue;
+
 	class FMixerDevice;
 
 	struct FQuartzQuantizedCommandDelegateData;
 	struct FQuartzMetronomeDelegateData;
+	struct FQuartzQueueCommandData;
 	struct FQuartzQuantizedCommandInitInfo;
 } // namespace Audio
 
-// UOBJECT LAYER:
 
 // An enumeration for specifying quantization for Quartz commands
 UENUM(BlueprintType)
@@ -218,55 +222,6 @@
 DECLARE_DYNAMIC_MULTICAST_DELEGATE_TwoParams(FOnQuartzCommandEvent, EQuartzCommandDelegateSubType, EventType, FName, Name);
 DECLARE_DYNAMIC_DELEGATE_TwoParams(FOnQuartzCommandEventBP, EQuartzCommandDelegateSubType, EventType, FName, Name);
 
-// struct used to specify the quantization boundary of an event
-USTRUCT(BlueprintType)
-struct ENGINE_API FQuartzQuantizationBoundary
-{
-	GENERATED_BODY()
-
-	// resolution we are interested in
-	UPROPERTY(EditAnywhere, BlueprintReadWrite, Category = "Quantized Audio Clock Settings")
-	EQuartzCommandQuantization Quantization{ EQuartzCommandQuantization::None };
-
-	// how many "Resolutions" to wait before the onset we care about
-	UPROPERTY(EditAnywhere, BlueprintReadWrite, Category = "Quantized Audio Clock Settings", meta = (ClampMin = "1.0"))
-	float Multiplier{ 1.0f };
-
-	UPROPERTY(EditAnywhere, BlueprintReadWrite, Category = "Quantized Audio Clock Settings")
-	EQuarztQuantizationReference CountingReferencePoint{ EQuarztQuantizationReference::BarRelative };
-
-	// If this is true and the Clock hasn't started yet, the event will fire immediately when the Clock starts
-	UPROPERTY(EditAnywhere, BlueprintReadWrite, AdvancedDisplay, Category = "Quantized Audio Clock Settings")
-	bool bFireOnClockStart{ true };
-
-	// If this is true, this command will be canceled if the Clock is stopped or otherwise not running
-	UPROPERTY(EditAnywhere, BlueprintReadWrite, AdvancedDisplay, Category = "Quantized Audio Clock Settings")
-	bool bCancelCommandIfClockIsNotRunning{ false };
-
-	// If this is true, queueing the sound will also call a Reset Clock command
-	UPROPERTY(EditAnywhere, BlueprintReadWrite, AdvancedDisplay, Category = "Quantized Audio Clock Settings")
-	bool bResetClockOnQueued{ false };
-
-	// If this is true, queueing the sound will also call a Resume Clock command
-	UPROPERTY(EditAnywhere, BlueprintReadWrite, AdvancedDisplay, Category = "Quantized Audio Clock Settings")
-	bool bResumeClockOnQueued{ false };
-
-	// Game thread subscribers that will be passed to command init data (for C++ implementations)
-	TArray<TSharedPtr<Audio::FShareableQuartzCommandQueue, ESPMode::ThreadSafe>> GameThreadSubscribers;
-
-	// ctor
-	FQuartzQuantizationBoundary(
-		EQuartzCommandQuantization InQuantization = EQuartzCommandQuantization::None, 
-		float InMultiplier = 1.0f, 
-		EQuarztQuantizationReference InReferencePoint = EQuarztQuantizationReference::BarRelative, 
-		bool bInFireOnClockStart = true)
-		: Quantization(InQuantization)
-		, Multiplier(InMultiplier)
-		, CountingReferencePoint(InReferencePoint)
-		, bFireOnClockStart(bInFireOnClockStart)
-	{}
-}; // struct FQuartzQuantizationBoundary
-
 // UStruct version of settings struct used to initialized a clock
 USTRUCT(BlueprintType)
 struct ENGINE_API FQuartzClockSettings
@@ -322,8 +277,8 @@
 	class FAudioMixer;
 
 	// Utility class to set/get/convert tick rate
-// In this context "Tick Rate" refers to the duration of smallest temporal resolution we may care about
-// in musical time, this is locked to a 1/32nd note
+	// In this context "Tick Rate" refers to the duration of smallest temporal resolution we may care about
+	// in musical time, this is locked to a 1/32nd note
 
 	struct ENGINE_API FQuartzClockTickRate
 	{
@@ -440,18 +395,152 @@
 		mutable FQuartzLatencyTimer Timer;
 	};
 
+	struct ENGINE_API FQuartzOffset
+	{
+	public:
+		// ctor
+		FQuartzOffset(double InOffsetInMilliseconds = 0.0);
+		FQuartzOffset(EQuartzCommandQuantization InDuration, double InMultiplier);
 
+		// offset get/set
+		void SetOffsetInMilliseconds(double InMilliseconds);
+
+		void SetOffsetMusical(EQuartzCommandQuantization Duration, double Multiplier);
+
+		bool IsSet() const { return IsSetAsMilliseconds() || IsSetAsMusicalDuration(); }
+
+		bool IsSetAsMilliseconds() const;
+
+		bool IsSetAsMusicalDuration() const;
+
+		int32 GetOffsetInAudioFrames(const FQuartzClockTickRate& InTickRate);
+
+		bool operator==(const FQuartzOffset& Other) const;
+
+	private:
+		// only one of these optionals will be valid at a time
+		// (depending on which setter is called)
+		TOptional<double> OffsetInMilliseconds;
+		TOptional<TPair<EQuartzCommandQuantization, double>> OffsetAsDuration;
+
+	};
+
+
+	using FQuartzGameThreadCommandQueue = Audio::TQuartzShareableCommandQueue<FQuartzTickableObject>;
+	using FQuartzGameThreadCommandQueuePtr = TSharedPtr<FQuartzGameThreadCommandQueue, ESPMode::ThreadSafe>;
+
+	struct ENGINE_API FQuartzGameThreadSubscriber
+	{
+		FQuartzGameThreadSubscriber() = default;
+
+		// this is only for back-compat until metronomes support FQuartzGameThreadSubscribers instead of raw queue ptrs
+		// (for Metronome event offsets)
+		FQuartzGameThreadSubscriber(const FQuartzGameThreadCommandQueuePtr& InQueuePtr)
+		: Queue(InQueuePtr)
+		{ }
+
+
+		// copy ctor
+		FQuartzGameThreadSubscriber(const FQuartzGameThreadCommandQueuePtr& InQueuePtr, FQuartzOffset InOffset)
+		: Offset(InOffset)
+		, Queue(InQueuePtr)
+		{ }
+
+		// change offset
+		void SetOffset(FQuartzOffset InOffset) { Offset = MoveTemp(InOffset); }
+
+		bool HasBeenNotifiedOfAboutToStart() const { return bHasBeenNotifiedOfAboutToStart; }
+
+		// comparison
+		bool operator==(const FQuartzGameThreadSubscriber& Other) const;
+
+		// todo: templatize to match teh underlying TQUartzShareableCommandQueue
+		// notify
+		void PushEvent(const FQuartzQuantizedCommandDelegateData& Data);
+		void PushEvent(const FQuartzMetronomeDelegateData& Data);
+		void PushEvent(const FQuartzQueueCommandData& Data);
+
+		// allow implicit casting to the underlying queue
+		operator FQuartzGameThreadCommandQueuePtr() const { return Queue; }
+
+		// positive: anticipatory amount, negative:
+		int32 FinalizeOffset(const FQuartzClockTickRate& TickRate);
+		int32 GetOffsetAsAudioFrames() const;
+
+
+
+	private:
+		FQuartzOffset Offset;
+		FQuartzGameThreadCommandQueuePtr Queue;
+		bool bOffsetConvertedToFrames = false;
+		bool bHasBeenNotifiedOfAboutToStart = false;
+		int32 OffsetInAudioFrames = 0;
+	};
+} // namespace Audio
+
+
+// struct used to specify the quantization boundary of an event
+USTRUCT(BlueprintType)
+struct ENGINE_API FQuartzQuantizationBoundary
+{
+	GENERATED_BODY()
+
+	// resolution we are interested in
+	UPROPERTY(EditAnywhere, BlueprintReadWrite, Category = "Quantized Audio Clock Settings")
+	EQuartzCommandQuantization Quantization{ EQuartzCommandQuantization::None };
+
+	// how many "Resolutions" to wait before the onset we care about
+	UPROPERTY(EditAnywhere, BlueprintReadWrite, Category = "Quantized Audio Clock Settings", meta = (ClampMin = "1.0"))
+	float Multiplier{ 1.0f };
+
+	UPROPERTY(EditAnywhere, BlueprintReadWrite, Category = "Quantized Audio Clock Settings")
+	EQuarztQuantizationReference CountingReferencePoint{ EQuarztQuantizationReference::BarRelative };
+
+	// If this is true and the Clock hasn't started yet, the event will fire immediately when the Clock starts
+	UPROPERTY(EditAnywhere, BlueprintReadWrite, AdvancedDisplay, Category = "Quantized Audio Clock Settings")
+	bool bFireOnClockStart{ true };
+
+	// If this is true, this command will be canceled if the Clock is stopped or otherwise not running
+	UPROPERTY(EditAnywhere, BlueprintReadWrite, AdvancedDisplay, Category = "Quantized Audio Clock Settings")
+	bool bCancelCommandIfClockIsNotRunning{ false };
+
+	// If this is true, queueing the sound will also call a Reset Clock command
+	UPROPERTY(EditAnywhere, BlueprintReadWrite, AdvancedDisplay, Category = "Quantized Audio Clock Settings")
+	bool bResetClockOnQueued{ false };
+
+	// If this is true, queueing the sound will also call a Resume Clock command
+	UPROPERTY(EditAnywhere, BlueprintReadWrite, AdvancedDisplay, Category = "Quantized Audio Clock Settings")
+	bool bResumeClockOnQueued{ false };
+
+	// Game thread subscribers that will be passed to command init data (for C++ implementations)
+	TArray<Audio::FQuartzGameThreadSubscriber> GameThreadSubscribers;
+
+	// ctor
+	FQuartzQuantizationBoundary(
+		EQuartzCommandQuantization InQuantization = EQuartzCommandQuantization::None,
+		float InMultiplier = 1.0f,
+		EQuarztQuantizationReference InReferencePoint = EQuarztQuantizationReference::BarRelative,
+		bool bInFireOnClockStart = true)
+		: Quantization(InQuantization)
+		, Multiplier(InMultiplier)
+		, CountingReferencePoint(InReferencePoint)
+		, bFireOnClockStart(bInFireOnClockStart)
+	{}
+}; // struct FQuartzQuantizationBoundary
+
+
+namespace Audio
+{
 	// data that is gathered by the AudioThread to get passed from FActiveSound->FMixerSourceVoice
 	// eventually converted to IQuartzQuantizedCommand for the Quantized Command itself
 	struct ENGINE_API FQuartzQuantizedRequestData
 	{
 		// shared with FQuartzQuantizedCommandInitInfo:
 		FName ClockName;
-		FName ClockHandleName;
 		FName OtherClockName;
 		TSharedPtr<IQuartzQuantizedCommand> QuantizedCommandPtr;
 		FQuartzQuantizationBoundary QuantizationBoundary{ EQuartzCommandQuantization::Tick, 1.f, EQuarztQuantizationReference::BarRelative, true };
-		TArray<TSharedPtr<FShareableQuartzCommandQueue, ESPMode::ThreadSafe>> GameThreadSubscribers;
+		TArray<FQuartzGameThreadSubscriber> GameThreadSubscribers;
 		int32 GameThreadDelegateID{ -1 };
 	};
 
@@ -468,15 +557,15 @@
 		void SetOwningClockPtr(TSharedPtr<Audio::FQuartzClock> InClockPointer)
 		{
 			OwningClockPointer = InClockPointer;
+			ensure(OwningClockPointer);
 		}
 
 		// shared with FQuartzQuantizedRequestData
 		FName ClockName;
-		FName ClockHandleName;
 		FName OtherClockName;
 		TSharedPtr<IQuartzQuantizedCommand> QuantizedCommandPtr{ nullptr };
 		FQuartzQuantizationBoundary QuantizationBoundary;
-		TArray<TSharedPtr<FShareableQuartzCommandQueue, ESPMode::ThreadSafe>> GameThreadSubscribers;
+		TArray<FQuartzGameThreadSubscriber> GameThreadSubscribers;
 		int32 GameThreadDelegateID{ -1 };
 
 		// Audio Render thread-specific data:
@@ -501,11 +590,18 @@
 		// allocate a copy of the derived class
 		virtual TSharedPtr<IQuartzQuantizedCommand> GetDeepCopyOfDerivedObject() const;
 
-		void AddSubscriber(TSharedPtr<FShareableQuartzCommandQueue, ESPMode::ThreadSafe> InSubscriber);
+		void AddSubscriber(FQuartzGameThreadSubscriber InSubscriber);
 
 		// Command has reached the AudioRenderThread
 		void OnQueued(const FQuartzQuantizedCommandInitInfo& InCommandInitInfo);
 
+		// scheduled (finalize subscriber offsets) - called by FQuartzClock
+		void OnScheduled(const FQuartzClockTickRate& InTickRate);
+
+		// called during FQuartzClock::Tick() to let us call AboutToStart
+		// at different times for different subscribers
+		void Update(int32 NumFramesUntilDeadline);
+
 		// Perhaps the associated sound failed concurrency and will not be playing
 		void FailedToQueue(FQuartzQuantizedRequestData& InGameThreadData);
 
@@ -530,16 +626,16 @@
 		virtual int32 OverrideFramesUntilExec(int32 NumFramesUntilExec) { return NumFramesUntilExec; }
 
 
-		virtual bool IsLooping() { return false; }
 		virtual bool IsClockAltering() { return false; }
+		virtual bool ShouldDeadlineIgnoresBpmChanges() { return false; }
 		virtual bool RequiresAudioDevice() const { return false; }
 
 		virtual FName GetCommandName() const = 0;
 		virtual EQuartzCommandType GetCommandType() const = 0;
 
 
-	protected:
-		// base classes can override these to add extra functionality
+	private:
+		// derived classes can override these to add extra functionality
 		virtual void OnQueuedCustom(const FQuartzQuantizedCommandInitInfo& InCommandInitInfo) {}
 		virtual void FailedToQueueCustom() {}
 		virtual void AboutToStartCustom() {}
@@ -548,11 +644,9 @@
 		virtual void OnClockStartedCustom() {}
 		virtual void CancelCustom() {}
 
-		TArray<TSharedPtr<FShareableQuartzCommandQueue, ESPMode::ThreadSafe>> GameThreadSubscribers;
+		TArray<FQuartzGameThreadSubscriber> GameThreadSubscribers;
 
-	private:
 		int32 GameThreadDelegateID{ -1 };
-		bool bAboutToStartHasBeenCalled{ false };
 	}; // class IAudioMixerQuantizedCommandBase
 
 	// Audio Render Thread Handle to a queued command
@@ -576,20 +670,20 @@
 	FAudioComponentCommandInfo() {}
 
 	FAudioComponentCommandInfo(const FAudioComponentCommandInfo& Other)
-		: ComponentCommandPtr(Other.ComponentCommandPtr)
+		: Subscriber(Other.Subscriber)
 		, AnticapatoryBoundary(Other.AnticapatoryBoundary)
 		, CommandID(Other.CommandID)
 	{}
 
- 	FAudioComponentCommandInfo(TSharedPtr<Audio::FShareableQuartzCommandQueue, ESPMode::ThreadSafe> InComponentCommandPtr, FQuartzQuantizationBoundary InAnticaptoryBoundary)
-		: ComponentCommandPtr(InComponentCommandPtr)
+ 	FAudioComponentCommandInfo(Audio::FQuartzGameThreadSubscriber InSubscriber, FQuartzQuantizationBoundary InAnticaptoryBoundary)
+		: Subscriber(InSubscriber)
 		, AnticapatoryBoundary(InAnticaptoryBoundary)
 	{
 		static uint32 CommandIDs = 0;
 		CommandID = CommandIDs++;
 	}
 
-	TSharedPtr<Audio::FShareableQuartzCommandQueue, ESPMode::ThreadSafe> ComponentCommandPtr;
+	Audio::FQuartzGameThreadSubscriber Subscriber;
 	FQuartzQuantizationBoundary AnticapatoryBoundary;
 	uint32 CommandID{ (uint32)INDEX_NONE };
 };
diff --git a/Engine/Source/Runtime/Engine/Classes/Sound/QuartzSubscription.cpp b/Engine/Source/Runtime/Engine/Classes/Sound/QuartzSubscription.cpp
--- a/Engine/Source/Runtime/Engine/Classes/Sound/QuartzSubscription.cpp
+++ b/Engine/Source/Runtime/Engine/Classes/Sound/QuartzSubscription.cpp
@@ -1,45 +1,248 @@
 // Copyright Epic Games, Inc. All Rights Reserved.
 
 #include "Sound/QuartzSubscription.h"
-
 #include "Quartz/AudioMixerClockHandle.h"
 
 namespace Audio
 {
 	FQuartzQueueCommandData::FQuartzQueueCommandData(const FAudioComponentCommandInfo& InAudioComponentCommandInfo, FName InClockName)
-		: AudioComponentCommandInfo(InAudioComponentCommandInfo),
-		ClockName(InClockName)
-	{}
+	: AudioComponentCommandInfo(InAudioComponentCommandInfo)
+	, ClockName(InClockName)
+	{
+	}
+} // namespace Audio
+
+PRAGMA_DISABLE_DEPRECATION_WARNINGS
+FQuartzTickableObject::~FQuartzTickableObject()
+{
+	if(const TSharedPtr<FQuartzTickableObjectsManager> ObjManagerPtr = TickableObjectManagerPtr.Pin())
+	{
+		ObjManagerPtr->UnsubscribeFromQuartzTick(this);
+	}
+}
+PRAGMA_ENABLE_DEPRECATION_WARNINGS
+
+FQuartzTickableObject* FQuartzTickableObject::Init(UWorld* InWorldPtr)
+{
+	if (!ensure(InWorldPtr))
+	{
+		// can't initialize if we don't have a valid world
+		return this;
+	}
 
-	void FShareableQuartzCommandQueue::PushEvent(const FQuartzQuantizedCommandDelegateData& Data)
+	PRAGMA_DISABLE_DEPRECATION_WARNINGS
+	if(!CommandQueuePtr.IsValid())
 	{
-		if (bIsAcceptingCommands)
+		CommandQueuePtr = Audio::TQuartzShareableCommandQueue<FQuartzTickableObject>::Create();
+	}
+	PRAGMA_ENABLE_DEPRECATION_WARNINGS
+
+	UQuartzSubsystem* QuartzSubsystemPtr = UQuartzSubsystem::Get(InWorldPtr);
+	if(ensure(QuartzSubsystemPtr))
+	{
+		TickableObjectManagerPtr = QuartzSubsystemPtr->GetTickableObjectManager();
+
+		if(TSharedPtr<FQuartzTickableObjectsManager> ObjManagerPtr = TickableObjectManagerPtr.Pin())
 		{
-			EventDelegateQueue.Enqueue([InData = Data](FQuartzTickableObject* Handle) { Handle->ProcessCommand(InData); });
+			ObjManagerPtr->SubscribeToQuartzTick(this);
 		}
 	}
 
-	void FShareableQuartzCommandQueue::PushEvent(const FQuartzMetronomeDelegateData& Data)
+	return this;
+}
+
+int32 FQuartzTickableObject::AddCommandDelegate(const FOnQuartzCommandEventBP& InDelegate)
+{
+	const int32 Num = QuantizedCommandDelegates.Num();
+	int32 SlotId = 0;
+
+	for (; SlotId < Num; ++SlotId)
 	{
-		if (bIsAcceptingCommands)
+		if (!QuantizedCommandDelegates[SlotId].MulticastDelegate.IsBound())
 		{
-			EventDelegateQueue.Enqueue([InData = Data](FQuartzTickableObject* Handle) { Handle->ProcessCommand(InData); });
+			QuantizedCommandDelegates[SlotId].MulticastDelegate.AddUnique(InDelegate);
+			return SlotId;
 		}
 	}
 
-	void FShareableQuartzCommandQueue::PushEvent(const FQuartzQueueCommandData& Data)
+	// need a new slot
+	QuantizedCommandDelegates.AddDefaulted_GetRef().MulticastDelegate.AddUnique(InDelegate);
+	return SlotId;
+}
+
+
+UQuartzSubsystem* FQuartzTickableObject::GetQuartzSubsystem() const
+{
+	return nullptr;
+}
+
+
+void FQuartzTickableObject::ExecCommand(const Audio::FQuartzQuantizedCommandDelegateData& Data)
+{
+	checkSlow(Data.DelegateSubType < EQuartzCommandDelegateSubType::Count);
+
+	if(const TSharedPtr<FQuartzTickableObjectsManager> ObjManagerPtr = TickableObjectManagerPtr.Pin())
+	{
+		ObjManagerPtr->PushLatencyTrackerResult(Data.RequestRecieved());
+	}
+
+	// Broadcast to the BP delegate if we have one bound
+	if (Data.DelegateID >= 0 && Data.DelegateID < QuantizedCommandDelegates.Num()
+		&& ensure(QuantizedCommandDelegates[Data.DelegateID].MulticastDelegate.IsBound()))
 	{
-		if (bIsAcceptingCommands)
+		FCommandDelegateGameThreadData& GameThreadEntry = QuantizedCommandDelegates[Data.DelegateID];
+
+		GameThreadEntry.MulticastDelegate.Broadcast(Data.DelegateSubType, "Quartz Evenct");
+
+		// track the number of active QuantizedCommands that may be sending info back to us.
+		// this is a bit of a hack because sound cues can play multiple wave instances
+		// and each of those wave instances is sending a delegate back to us.
+		// todo: clean this up at the wave-instance level to avoid ref counting here.
+		// (new command)
+		if (Data.DelegateSubType == EQuartzCommandDelegateSubType::CommandOnQueued)
 		{
-			EventDelegateQueue.Enqueue([InData = Data](FQuartzTickableObject* Handle) { Handle->ProcessCommand(InData); });
+			GameThreadEntry.RefCount.Increment();
+		}
+
+		// (end of a command)
+		if (Data.DelegateSubType == EQuartzCommandDelegateSubType::CommandOnCanceled)
+		{
+			// are all the commands done?
+			if (GameThreadEntry.RefCount.Decrement() == 0)
+			{
+				GameThreadEntry.MulticastDelegate.Clear();
+			}
 		}
 	}
 
+	// call base-class method
+	ProcessCommand(Data);
+}
 
-	void FShareableQuartzCommandQueue::StopTakingCommands()
+void FQuartzTickableObject::ExecCommand(const Audio::FQuartzMetronomeDelegateData& Data)
+{
+	if(const TSharedPtr<FQuartzTickableObjectsManager> ObjManagerPtr = TickableObjectManagerPtr.Pin())
 	{
-		bIsAcceptingCommands = false;
-		EventDelegateQueue.Empty();
+		ObjManagerPtr->PushLatencyTrackerResult(Data.RequestRecieved());
+	}
+
+	MetronomeDelegates[static_cast<int32>(Data.Quantization)].MulticastDelegate
+		.Broadcast(Data.ClockName, Data.Quantization, Data.Bar, Data.Beat, Data.BeatFraction);
+
+	// call base-class method
+	ProcessCommand(Data);
+}
+
+void FQuartzTickableObject::ExecCommand(const Audio::FQuartzQueueCommandData& Data)
+{
+	// call base-class method
+	ProcessCommand(Data);
+}
+
+void FQuartzTickableObject::SetNotificationAnticipationAmountMilliseconds(const double Milliseconds)
+{
+	// todo: update metronome subscriptions w/ new value (once metronome observes offsets)
+	NotificationOffset.SetOffsetInMilliseconds(Milliseconds);
+}
+
+void FQuartzTickableObject::SetNotificationAnticipationAmountMusicalDuration(const EQuartzCommandQuantization Duration,
+	const double Multiplier)
+{
+	NotificationOffset.SetOffsetMusical(Duration, Multiplier);
+}
+
+Audio::FQuartzGameThreadSubscriber FQuartzTickableObject::GetQuartzSubscriber()
+ {
+	PRAGMA_DISABLE_DEPRECATION_WARNINGS
+ 	if (!CommandQueuePtr.IsValid())
+ 	{
+ 		CommandQueuePtr = MakeShared<Audio::TQuartzShareableCommandQueue<FQuartzTickableObject>, ESPMode::ThreadSafe>();
+ 	}
+
+ 	return { CommandQueuePtr, NotificationOffset };
+	PRAGMA_ENABLE_DEPRECATION_WARNINGS
+}
+
+void FQuartzTickableObject::QuartzTick(float DeltaTime)
+{
+	PRAGMA_DISABLE_DEPRECATION_WARNINGS
+	CommandQueuePtr->PumpCommandQueue(this);
+	PRAGMA_ENABLE_DEPRECATION_WARNINGS
+}
+
+bool FQuartzTickableObject::QuartzIsTickable() const
+{
+	PRAGMA_DISABLE_DEPRECATION_WARNINGS
+	return CommandQueuePtr.IsValid();
+	PRAGMA_ENABLE_DEPRECATION_WARNINGS
+}
+
+void FQuartzTickableObject::AddMetronomeBpDelegate(EQuartzCommandQuantization InQuantizationBoundary, const FOnQuartzMetronomeEventBP& OnQuantizationEvent)
+{
+	MetronomeDelegates[static_cast<int32>(InQuantizationBoundary)].MulticastDelegate.AddUnique(OnQuantizationEvent);
+}
+
+namespace Audio
+{
+
+	FQuartzOffset::FQuartzOffset(double InOffsetInMilliseconds)
+	: OffsetInMilliseconds(MoveTemp(InOffsetInMilliseconds))
+	{
+	}
+
+	FQuartzOffset::FQuartzOffset(EQuartzCommandQuantization InDuration, double InMultiplier)
+	: OffsetAsDuration(TPair<EQuartzCommandQuantization, double>(MoveTemp(InDuration), MoveTemp(InMultiplier)))
+	{
+	}
+
+	void FQuartzOffset::SetOffsetInMilliseconds(double InMilliseconds)
+	{
+		OffsetAsDuration.Reset();
+		OffsetInMilliseconds.Emplace(InMilliseconds);
+	}
+
+	void FQuartzOffset::SetOffsetMusical(EQuartzCommandQuantization Duration, double Multiplier)
+	{
+		OffsetInMilliseconds.Reset();
+		OffsetAsDuration.Emplace(TPair<EQuartzCommandQuantization, float>(Duration, Multiplier));
+	}
+
+	bool FQuartzOffset::IsSetAsMilliseconds() const
+	{
+		return OffsetInMilliseconds.IsSet();
+	}
+
+	bool FQuartzOffset::IsSetAsMusicalDuration() const
+	{
+		return OffsetAsDuration.IsSet();
+	}
+
+	int32 FQuartzOffset::GetOffsetInAudioFrames(const FQuartzClockTickRate& InTickRate)
+	{
+		// only one should be set at a time: updated by last SetOffset[In/As][Milliseconds/Duration]()
+		check(!(OffsetInMilliseconds.IsSet() && OffsetAsDuration.IsSet()));
+
+		if(OffsetInMilliseconds.IsSet())
+		{
+			const double OffsetInSeconds = OffsetInMilliseconds.GetValue() * 1000.0;
+			const int32 OffsetInFrames = OffsetInSeconds * InTickRate.GetSampleRate();
+
+			return OffsetInFrames;
+		}
+		else if(OffsetAsDuration.IsSet())
+		{
+			auto [QuantizationType, Multiplier] = OffsetAsDuration.GetValue();
+
+			if (QuantizationType == EQuartzCommandQuantization::None)
+			{
+				return 0;
+			}
+
+			return static_cast<int32>(Multiplier * InTickRate.GetFramesPerDuration(QuantizationType));
+		}
+
+		ensure(false); // one of our durations should have been set (even if by a constructor)
+		return 0;
 	}
 
-} // namespace Audio
\ No newline at end of file
+} // namespace Audio
diff --git a/Engine/Source/Runtime/Engine/Classes/Sound/QuartzSubscription.h b/Engine/Source/Runtime/Engine/Classes/Sound/QuartzSubscription.h
--- a/Engine/Source/Runtime/Engine/Classes/Sound/QuartzSubscription.h
+++ b/Engine/Source/Runtime/Engine/Classes/Sound/QuartzSubscription.h
@@ -3,16 +3,19 @@
 #pragma once
 
 #include "Sound/QuartzQuantizationUtilities.h"
-#include "UObject/ObjectMacros.h"
-#include "Containers/Queue.h"
-#include "DSP/VolumeFader.h"
+#include "Containers/MpscQueue.h"
 
 // forwards
-class FQuartzTickableObject;
+class UQuartzSubsystem;
 class UQuartzClockHandle;
+class FQuartzTickableObject;
+struct FQuartzTickableObjectsManager;
 
 namespace Audio
 {
+	struct FQuartzGameThreadSubscriber;
+	class FQuartzClockProxy;
+
 	// Struct used to communicate command state back to the game play thread
 	struct ENGINE_API FQuartzQuantizedCommandDelegateData : public FQuartzCrossThreadMessage
 	{
@@ -22,8 +25,7 @@ namespace Audio
 		// ID so the clock handle knows which delegate to fire
 		int32 DelegateID{ -1 };
 
-		// todo: more payload
-	};
+	}; // struct FQuartzQuantizedCommandDelegateData
 
 	// Struct used to communicate metronome events back to the game play thread
 	struct ENGINE_API FQuartzMetronomeDelegateData : public FQuartzCrossThreadMessage
@@ -32,7 +34,8 @@ namespace Audio
 		int32 Beat;
 		float BeatFraction;
 		EQuartzCommandQuantization Quantization;
-	};
+		FName ClockName;
+	}; // struct FQuartzMetronomeDelegateData
 
 	//Struct used to queue events to be sent to the Audio Render thread closer to their start time
 	struct ENGINE_API FQuartzQueueCommandData : public FQuartzCrossThreadMessage
@@ -41,32 +44,181 @@ namespace Audio
 		FName ClockName;
 
 		FQuartzQueueCommandData(const FAudioComponentCommandInfo& InAudioComponentCommandInfo, FName InClockName);
-	};
+	}; // struct FQuartzQueueCommandData
 
 
-	// Class that can be shared between Game Thread and any other thread to queue commands
-	// (ONLY the Game Thread may EXECUTE the commands in the queue, enforced by UQuartzClockHandle* argument in the TFunctions)
+	// old non-template\ version of TQuartzShareableCommandQueue
+	class UE_DEPRECATED(5.1, "Message") FShareableQuartzCommandQueue;
 	class ENGINE_API FShareableQuartzCommandQueue
 	{
+	};
+
+	/**
+	*	Template class for mono-directional MPSC command queues 
+	* 
+	*   in order to enforce thread-safe access to the object executing the commands,
+	*	"listener type" is the type of the object that is being accessed in the commands
+	*	that object will have to provide a 'this' ptr (of type ListenerType) in order to 
+	*   invoke the commands on itself. (The lambdas do NOT and should NOT cache a ptr or
+	*	reference to the target).
+	* 
+	*	User-provided lambdas can take any (single) argument type T in PushEvent()
+	*	but there must exist a ListenerType::ExecCommand(T) overload for any PushEvent(T) instantiated. 
+	* 
+	*	(see FQuartzTickableObject and FQuartzClock as examples)
+	* 
+	**/
+	template <class ListenerType>
+	class TQuartzShareableCommandQueue
+	{
 	public:
-		// (add more PushEvent overloads for other event types)
-		void PushEvent(const FQuartzQuantizedCommandDelegateData& Data);
-		void PushEvent(const FQuartzMetronomeDelegateData& Data);
-		void PushEvent(const FQuartzQueueCommandData& Data);
+		// ctor
+		TQuartzShareableCommandQueue() {}
+
+		// dtor
+		~TQuartzShareableCommandQueue() {}
 
-		bool IsQueueEmpty()
+		// static helper to create a new sharable queue
+		static TSharedPtr<TQuartzShareableCommandQueue<ListenerType>, ESPMode::ThreadSafe> Create()
 		{
-			return EventDelegateQueue.IsEmpty();
+			return MakeShared<TQuartzShareableCommandQueue<ListenerType>, ESPMode::ThreadSafe>();
 		}
 
+		// note: ListenerType must have a ExecCommand() overload for each instantiation of this method
+		template <typename T>
+		void PushEvent(const T& Data)
+		{
+			CommandQueue.Enqueue([InData = Data](ListenerType* Listener) { Listener->ExecCommand(InData); });
+		}
 
-		// called when the game object owner is shutting down
-		void StopTakingCommands();
+		void PushCommand(TFunction<void(ListenerType*)> InCommand)
+		{
+			CommandQueue.Enqueue([=](ListenerType* Listener) { InCommand(Listener); });
+		}
 
-	private:
-		TQueue<TFunction<void(FQuartzTickableObject*)>, EQueueMode::Mpsc> EventDelegateQueue;
-		FThreadSafeBool bIsAcceptingCommands{ true };
+		void PumpCommandQueue(ListenerType* InListener)
+		{
+			// gather move all the current commands into our temp container
+			// (in case the commands themselves alter the container)
+			TOptional<TFunction<void(ListenerType*)>> Function = CommandQueue.Dequeue();
+			while (Function.IsSet())
+			{
+				TempCommandQueue.Emplace(Function.GetValue());
+				Function = CommandQueue.Dequeue();
+			}
+
+			for (auto& Command : TempCommandQueue)
+			{
+				Command(InListener);
+			}
+
+			TempCommandQueue.Reset();
+		}
 
-		friend class ::FQuartzTickableObject; // UQuartzTickableObject, and those that inherit from it, are the only classes allowed to pump the command queue
+	private:
+		TMpscQueue<TFunction<void(ListenerType*)>> CommandQueue;
+		TArray<TFunction<void(ListenerType*)>> TempCommandQueue;
 	};
+
 } // namespace Audio
+
+using FQuartzGameThreadCommandQueue = Audio::TQuartzShareableCommandQueue<FQuartzTickableObject>;
+using FQuartzGameThreadCommandQueuePtr = TSharedPtr<FQuartzGameThreadCommandQueue, ESPMode::ThreadSafe>;
+
+/**
+ *	FQuartzTickableObject
+ *
+ *		This is the base class for non-Audio Render Thread objects that want to receive
+ *		callbacks for Quartz events.
+ *
+ *		It is a wrapper around TQuartzShareableCommandQueue.
+ *		(see UQuartzClockHandle or UAudioComponent as implementation examples)
+ */
+class ENGINE_API FQuartzTickableObject
+{
+public:
+	// ctor
+	FQuartzTickableObject() {}
+
+	// dtor
+	virtual ~FQuartzTickableObject();
+
+	FQuartzTickableObject* Init(UWorld* InWorldPtr);
+
+	// called by the associated QuartzSubsystem
+	void QuartzTick(float DeltaTime);
+
+	bool QuartzIsTickable() const;
+
+	// QuartzTodo: uncomment post-cleanup
+	// UE_DEPRECATED(5.1, "Derived classes should have access to their own UWorld, this function will always return null")
+	UWorld* QuartzGetWorld() const { return nullptr; }
+
+	// QuartzTodo: delete this static function post-cleanup
+	void QuartzAddReferencedObjects(FReferenceCollector& Collector)
+	{
+	}
+
+	void AddMetronomeBpDelegate(EQuartzCommandQuantization InQuantizationBoundary, const FOnQuartzMetronomeEventBP& OnQuantizationEvent);
+
+
+	bool IsInitialized() const { return TickableObjectManagerPtr.IsValid(); }
+
+	Audio::FQuartzGameThreadSubscriber GetQuartzSubscriber();
+
+	int32 AddCommandDelegate(const FOnQuartzCommandEventBP& InDelegate);
+
+	UE_DEPRECATED(5.1, "use FQuartzTickableObject::AddCommandDelegate(const FOnQuartzCommandEventBP& InDelegate) insead")
+	int32 AddCommandDelegate(const FOnQuartzCommandEventBP& InDelegate, TArray<FQuartzGameThreadCommandQueuePtr>& TargetSubscriberArray){ return -1; }
+
+
+	// QuartzTodo: uncomment  after fix-up
+	// UE_DEPRECATED(5.1, "This object no longer holds any UObject references. Caller should have their own UWorld* and use static 'UQuartzSubsystem::Get()' instead.")
+	UQuartzSubsystem* GetQuartzSubsystem() const;
+
+	// required by TQuartzShareableCommandQueue template
+	void ExecCommand(const Audio::FQuartzQuantizedCommandDelegateData& Data);
+	void ExecCommand(const Audio::FQuartzMetronomeDelegateData& Data);
+	void ExecCommand(const Audio::FQuartzQueueCommandData& Data);
+
+	// virtual interface (ExecCommand will forward the data to derived classes' ProcessCommand() call)
+	virtual void ProcessCommand(const Audio::FQuartzQuantizedCommandDelegateData& Data) {}
+	virtual void ProcessCommand(const Audio::FQuartzMetronomeDelegateData& Data) {}
+	virtual void ProcessCommand(const Audio::FQuartzQueueCommandData& Data) {}
+
+	const Audio::FQuartzOffset& GetQuartzOffset() const { return NotificationOffset; }
+
+protected:
+	void SetNotificationAnticipationAmountMilliseconds(const double Milliseconds);
+	void SetNotificationAnticipationAmountMusicalDuration(const EQuartzCommandQuantization Duration,  const double Multiplier);
+
+private:
+	struct FMetronomeDelegateGameThreadData
+	{
+		FOnQuartzMetronomeEvent MulticastDelegate;
+	};
+
+	struct FCommandDelegateGameThreadData
+	{
+		FOnQuartzCommandEvent MulticastDelegate;
+		FThreadSafeCounter RefCount;
+	};
+
+	// delegate containers
+	FMetronomeDelegateGameThreadData MetronomeDelegates[static_cast<int32>(EQuartzCommandQuantization::Count)];
+	TArray<FCommandDelegateGameThreadData> QuantizedCommandDelegates;
+
+	TArray<TFunction<void(FQuartzTickableObject*)>> TempCommandQueue;
+
+public:
+	// deprecate public access
+	// QuartzTodo: uncomment after fix-up
+	// UE_DEPRECATED(5.1, "FQuartzTickableObject::CommandQueuePtr should no longer be accessed directly. (This member will be private in a future engine version).  Use GetQuartzSubscriber() instead")
+	TSharedPtr<Audio::TQuartzShareableCommandQueue<FQuartzTickableObject>, ESPMode::ThreadSafe> CommandQueuePtr;
+
+private:
+	Audio::FQuartzOffset NotificationOffset;
+	TWeakPtr<FQuartzTickableObjectsManager> TickableObjectManagerPtr;
+
+}; // class FQuartzTickableObject
+
diff --git a/Engine/Source/Runtime/Engine/Private/Components/AudioComponent.cpp b/Engine/Source/Runtime/Engine/Private/Components/AudioComponent.cpp
--- a/Engine/Source/Runtime/Engine/Private/Components/AudioComponent.cpp
+++ b/Engine/Source/Runtime/Engine/Private/Components/AudioComponent.cpp
@@ -43,7 +43,7 @@ FAutoConsoleVariableRef CVarTimeToTakeUpVoiceSlot(
 
 
 //CVar to disable the gamethread-side caching of play requests to minimize time spent rendering a silent voice before the Quartz Deadline
-static int32 bAlwaysTakeVoiceSlotCVar = 0;
+static int32 bAlwaysTakeVoiceSlotCVar = 1;
 FAutoConsoleVariableRef bCVarAlwaysTakeVoiceSlot(
 	TEXT("au.Quartz.bAlwaysTakeVoiceSlot"),
 	bAlwaysTakeVoiceSlotCVar,
@@ -284,13 +284,11 @@ void UAudioComponent::OnUnregister()
 	{
 		Stop();
 	}
-	Shutdown();
 }
 
 void UAudioComponent::EndPlay(const EEndPlayReason::Type EndPlayReason)
 {
 	Super::EndPlay(EndPlayReason);
-	Shutdown();
 }
 
 const UObject* UAudioComponent::AdditionalStatObject() const
@@ -448,13 +446,7 @@ void UAudioComponent::ProcessCommand(const Audio::FQuartzQuantizedCommandDelegat
 
 void UAudioComponent::ProcessCommand(const Audio::FQuartzQueueCommandData& InQueueCommandData)
 {
-	if (UQuartzSubsystem* QuartzSubsystem = GetQuartzSubsystem())
-	{
-		QuartzSubsystem->PushLatencyTrackerResult(InQueueCommandData.RequestRecieved());
-
-		//Queue the sound
-		PlayQueuedQuantizedInternal(GetWorld(), InQueueCommandData.AudioComponentCommandInfo);
-	}
+	PlayQueuedQuantizedInternal(GetWorld(), InQueueCommandData.AudioComponentCommandInfo);
 }
 
 void UAudioComponent::PlayQuantized(
@@ -472,15 +464,15 @@ void UAudioComponent::PlayQuantized(
 	{
 		Init(GetWorld());
 	}
-	check(FQuartzTickableObject::IsInitialized() && CommandQueuePtr.IsValid());
+	check(FQuartzTickableObject::IsInitialized());
 
 	int32 TimeCVarVal = FMath::Clamp(CVarTimeToTakeUpVoiceSlot->GetInt(), 0, (int32)EQuartzCommandQuantization::Count - 1);
 	EQuartzCommandQuantization MinimumQuantization = EQuartzCommandQuantization(TimeCVarVal);
 
-	// Make an anticapatory quantization boundary to try to avoid taking up a whole voice slot while waiting for a queued event
+	// Make an anticipatory quantization boundary to try to avoid taking up a whole voice slot while waiting for a queued event
 	FQuartzQuantizationBoundary AnticipationQuantizationBoundary = FQuartzQuantizationBoundary(MinimumQuantization, 1.0f, EQuarztQuantizationReference::CurrentTimeRelative, true);
 
-	FAudioComponentCommandInfo NewComponentCommandInfo(CommandQueuePtr, AnticipationQuantizationBoundary);
+	FAudioComponentCommandInfo NewComponentCommandInfo(FQuartzTickableObject::GetQuartzSubscriber(), AnticipationQuantizationBoundary);
 
 	// And a new pending quartz command data
 	FAudioComponentPendingQuartzCommandData AudioComponentQuartzCommandData
@@ -534,12 +526,8 @@ void UAudioComponent::PlayQuantized(
 
 void UAudioComponent::PlayQueuedQuantizedInternal(const UObject* WorldContextObject, FAudioComponentCommandInfo InCommandInfo)
 {
-	//Initialize the tickable object portion of the Audio Component, if it hasn't been initialized already
-	if (!FQuartzTickableObject::IsInitialized())
-	{
-		Init(GetWorld());
-	}
-	check(FQuartzTickableObject::IsInitialized() && CommandQueuePtr.IsValid());
+	// confirm the FQuartzTickableObject has been initialized.
+	ensure(FQuartzTickableObject::IsInitialized());
 
 	bool bFoundQuantizedCommand = false;
 	bool bIsValidCommand = true;
@@ -569,7 +557,7 @@ void UAudioComponent::PlayQueuedQuantizedInternal(const UObject* WorldContextObj
 				// confirm a valid handle
 				if (Handle != nullptr)
 				{
-					InternalRequestData.QuantizedRequestData = Handle->GetQuartzSubsystem()->CreateDataDataForSchedulePlaySound(Handle, PendingData.Delegate, PendingData.AnticapatoryBoundary);
+					InternalRequestData.QuantizedRequestData = UQuartzSubsystem::CreateRequestDataForSchedulePlaySound(Handle, PendingData.Delegate, PendingData.AnticapatoryBoundary);
 					UGameplayStatics::PrimeSound(Sound);
 				}
 
@@ -596,17 +584,17 @@ void UAudioComponent::PlayQueuedQuantizedInternal(const UObject* WorldContextObj
 
 				if (bIsValidCommand)
 				{
-					InternalRequestData.QuantizedRequestData.GameThreadSubscribers.Add(CommandQueuePtr);
+					InternalRequestData.QuantizedRequestData.GameThreadSubscribers.Add(GetQuartzSubscriber());
 
 					// Now play the quartz command
 					PlayInternal(InternalRequestData);
+					UE_LOG(LogAudioQuartz, Warning, TEXT("AudioComponent::PlayQueueQUantizedInternal()->PlayInternal()"));
 				} 
 			}
 			else // this component does not have a Sound to play
 			{
 				UE_LOG(LogAudioQuartz, Warning, TEXT("Attempting to play Quantized Sound without supplying a valid Sound to play"));
 			}
-			
 
 
 
@@ -839,10 +827,6 @@ void UAudioComponent::PlayInternal(const PlayInternalRequestData& InPlayRequestD
 		UActorSoundParameterInterface::Fill(Owner, ActorParams);
 		FAudioParameter::Merge(MoveTemp(ActorParams), SoundParams);
 	}
-	else
-	{
-		UE_LOG(LogAudio, VeryVerbose, TEXT("AudioComponent : '%s' playing sound '%s' has no owner"), *GetFullName(), *Sound->GetName());
-	}
 
 	TArray<FAudioParameter> InstanceParamsCopy = InstanceParameters;
 	FAudioParameter::Merge(MoveTemp(InstanceParamsCopy), SoundParams);
