Fix static analysis warning (variable initiatlized through itself). Issue 177151.

diff --git a/Engine/Source/Runtime/AudioMixer/Private/Quartz/QuartzSubsystem.cpp b/Engine/Source/Runtime/AudioMixer/Private/Quartz/QuartzSubsystem.cpp
--- a/Engine/Source/Runtime/AudioMixer/Private/Quartz/QuartzSubsystem.cpp
+++ b/Engine/Source/Runtime/AudioMixer/Private/Quartz/QuartzSubsystem.cpp
@@ -632,7 +632,7 @@ Audio::FQuartzClockProxy* UQuartzSubsystem::FindProxyByName(const FName& ClockNa
 
 Audio::FQuartzClockProxy const* UQuartzSubsystem::FindProxyByName(const FName& ClockName) const
 {
-	Audio::FQuartzClockProxy const* Result = Result = ActiveAudioMixerClockProxies.FindByKey(ClockName);
+	Audio::FQuartzClockProxy const* Result = ActiveAudioMixerClockProxies.FindByKey(ClockName);
 	if(!Result)
 	{
 		Result = ActiveAudioMixerClockProxies.FindByKey(ClockName);
