[CIS] fix for static analysis warnings (wrong variable name used) & fix for deprecation warning generated by implicit copy-asignment ctor
#p4v-preflight-copy 20528511

diff --git a/Engine/Source/Runtime/AudioMixer/Private/Quartz/QuartzSubsystem.cpp b/Engine/Source/Runtime/AudioMixer/Private/Quartz/QuartzSubsystem.cpp
--- a/Engine/Source/Runtime/AudioMixer/Private/Quartz/QuartzSubsystem.cpp
+++ b/Engine/Source/Runtime/AudioMixer/Private/Quartz/QuartzSubsystem.cpp
@@ -620,11 +620,12 @@ void UQuartzSubsystem::PruneStaleProxiesInternal(TArray<Audio::FQuartzClockProxy
 
 Audio::FQuartzClockProxy* UQuartzSubsystem::FindProxyByName(const FName& ClockName)
 {
-	Audio::FQuartzClockProxy* Result = ActiveAudioMixerClockProxies.FindByKey((ClockName));
+	Audio::FQuartzClockProxy* Result = ActiveAudioMixerClockProxies.FindByKey(ClockName);
 
+	// if the subsystem doesn't have a match, check the externally-registered clock proxies
 	if(!Result)
 	{
-		Result = ActiveAudioMixerClockProxies.FindByKey(ClockName);
+		Result = ActiveExternalClockProxies.FindByKey(ClockName);
 	}
 
 	return Result;
@@ -633,9 +634,11 @@ Audio::FQuartzClockProxy* UQuartzSubsystem::FindProxyByName(const FName& ClockNa
 Audio::FQuartzClockProxy const* UQuartzSubsystem::FindProxyByName(const FName& ClockName) const
 {
 	Audio::FQuartzClockProxy const* Result = ActiveAudioMixerClockProxies.FindByKey(ClockName);
+
+	// if the subsystem doesn't have a match, check the externally-registered clock proxies
 	if(!Result)
 	{
-		Result = ActiveAudioMixerClockProxies.FindByKey(ClockName);
+		Result = ActiveExternalClockProxies.FindByKey(ClockName);
 	}
 
 	return Result;
diff --git a/Engine/Source/Runtime/Engine/Classes/Sound/QuartzSubscription.h b/Engine/Source/Runtime/Engine/Classes/Sound/QuartzSubscription.h
--- a/Engine/Source/Runtime/Engine/Classes/Sound/QuartzSubscription.h
+++ b/Engine/Source/Runtime/Engine/Classes/Sound/QuartzSubscription.h
@@ -140,10 +140,11 @@ public:
 	// ctor
 	FQuartzTickableObject() {}
 
-	// copy ctor (to disable deprecation warnings generated by implicit copy ctor)
-	PRAGMA_DISABLE_DEPRECATION_WARNINGS
-	FQuartzTickableObject(const FQuartzTickableObject& Other) = default;
-	PRAGMA_ENABLE_DEPRECATION_WARNINGS
+	// explicitly defaulted ctors (to disable deprecation warnings in compiler-generated functions)
+    PRAGMA_DISABLE_DEPRECATION_WARNINGS
+    FQuartzTickableObject(const FQuartzTickableObject& Other) = default;
+    FQuartzTickableObject& operator=(const FQuartzTickableObject&) = default;
+    PRAGMA_ENABLE_DEPRECATION_WARNINGS
 
 	// dtor
 	virtual ~FQuartzTickableObject();
