Update Audio Settings Telemetry Event Name

diff --git a/Engine/Source/Runtime/AudioMixer/Private/AudioMixerDevice.cpp b/Engine/Source/Runtime/AudioMixer/Private/AudioMixerDevice.cpp
--- a/Engine/Source/Runtime/AudioMixer/Private/AudioMixerDevice.cpp
+++ b/Engine/Source/Runtime/AudioMixer/Private/AudioMixerDevice.cpp
@@ -470,7 +470,7 @@ namespace Audio
 				// Create synchronized Audio Task Queue for this device...
 				CreateSynchronizedAudioTaskQueue((Audio::AudioTaskQueueId)DeviceID);
 
-				FEngineAnalytics::GetProvider().RecordEvent(TEXT("Audio.Usage.EditorSettings"), MakeAnalyticsEventAttributeArray(
+				FEngineAnalytics::GetProvider().RecordEvent(TEXT("Audio.Usage.ProjectSettings"), MakeAnalyticsEventAttributeArray(
 					TEXT("SampleRate"), PlatformSettings.SampleRate,
 					TEXT("BufferSize"), PlatformSettings.CallbackBufferFrameSize,
 					TEXT("NumBuffers"), PlatformSettings.NumBuffers,
