[Backout] - CL20767956 & CL20768414

Original CL Desc for 20767956
-----------------------------------------------------------------
Editor Settings Telemetry

Original CL Desc for 20768414
-----------------------------------------------------------------
Update Audio Settings Telemetry Event Name

diff --git a/Engine/Source/Runtime/AudioMixer/Private/AudioMixerDevice.cpp b/Engine/Source/Runtime/AudioMixer/Private/AudioMixerDevice.cpp
--- a/Engine/Source/Runtime/AudioMixer/Private/AudioMixerDevice.cpp
+++ b/Engine/Source/Runtime/AudioMixer/Private/AudioMixerDevice.cpp
@@ -11,7 +11,6 @@
 #include "AudioMixerEffectsManager.h"
 #include "DSP/Noise.h"
 #include "DSP/SinOsc.h"
-#include "EngineAnalytics.h"
 #include "Sound/AudioSettings.h"
 #include "Sound/SoundSubmix.h"
 #include "Sound/SoundSubmixSend.h"
@@ -470,13 +469,6 @@ namespace Audio
 				// Create synchronized Audio Task Queue for this device...
 				CreateSynchronizedAudioTaskQueue((Audio::AudioTaskQueueId)DeviceID);
 
-				FEngineAnalytics::GetProvider().RecordEvent(TEXT("Audio.Usage.ProjectSettings"), MakeAnalyticsEventAttributeArray(
-					TEXT("SampleRate"), PlatformSettings.SampleRate,
-					TEXT("BufferSize"), PlatformSettings.CallbackBufferFrameSize,
-					TEXT("NumBuffers"), PlatformSettings.NumBuffers,
-					TEXT("NumSources"), PlatformSettings.MaxChannels,
-					TEXT("NumOutputChannels"), PlatformInfo.NumChannels));
-
 				// Start streaming audio
 				return AudioMixerPlatform->StartAudioStream();
 			}
