Audio Proejct Settings Telemetry (Properly this time)

diff --git a/Engine/Source/Runtime/AudioMixer/Private/AudioMixerDevice.cpp b/Engine/Source/Runtime/AudioMixer/Private/AudioMixerDevice.cpp
--- a/Engine/Source/Runtime/AudioMixer/Private/AudioMixerDevice.cpp
+++ b/Engine/Source/Runtime/AudioMixer/Private/AudioMixerDevice.cpp
@@ -9,6 +9,7 @@
 #include "AudioMixerSourceVoice.h"
 #include "AudioPluginUtilities.h"
 #include "AudioMixerEffectsManager.h"
+#include "EngineAnalytics.h"
 #include "DSP/Noise.h"
 #include "DSP/SinOsc.h"
 #include "Sound/AudioSettings.h"
@@ -469,6 +470,16 @@ namespace Audio
 				// Create synchronized Audio Task Queue for this device...
 				CreateSynchronizedAudioTaskQueue((Audio::AudioTaskQueueId)DeviceID);
 
+				if (FEngineAnalytics::IsAvailable())
+				{
+					FEngineAnalytics::GetProvider().RecordEvent(TEXT("Audio.Usage.ProjectSettings"), MakeAnalyticsEventAttributeArray(
+						TEXT("SampleRate"), PlatformSettings.SampleRate,
+						TEXT("BufferSize"), PlatformSettings.CallbackBufferFrameSize,
+						TEXT("NumBuffers"), PlatformSettings.NumBuffers,
+						TEXT("NumSources"), PlatformSettings.MaxChannels,
+						TEXT("NumOutputChannels"), PlatformInfo.NumChannels));
+				}
+
 				// Start streaming audio
 				return AudioMixerPlatform->StartAudioStream();
 			}
