[Quartz] Fix for crash if multiple BPM changes are pending during an audio engine tick and removed an invalid check.

diff --git a/Engine/Source/Runtime/AudioMixer/Private/Quartz/AudioMixerClock.cpp b/Engine/Source/Runtime/AudioMixer/Private/Quartz/AudioMixerClock.cpp
--- a/Engine/Source/Runtime/AudioMixer/Private/Quartz/AudioMixerClock.cpp
+++ b/Engine/Source/Runtime/AudioMixer/Private/Quartz/AudioMixerClock.cpp
@@ -149,7 +149,7 @@ namespace Audio
 		// adjust time-till-fire for existing commands
 		for (auto& Command : PendingCommands)
 		{
-			if(!Command.Command->ShouldDeadlineIgnoresBpmChanges())
+			if(Command.Command && !Command.Command->ShouldDeadlineIgnoresBpmChanges())
 			{
 				Command.NumFramesUntilExec = NumFramesLeft + Ratio * (Command.NumFramesUntilExec - NumFramesLeft);
 			}
@@ -157,7 +157,7 @@ namespace Audio
 
 		for (auto& Command : ClockAlteringPendingCommands)
 		{
-			if(!Command.Command->ShouldDeadlineIgnoresBpmChanges())
+			if(Command.Command && !Command.Command->ShouldDeadlineIgnoresBpmChanges())
 			{
 				Command.NumFramesUntilExec = NumFramesLeft + Ratio * (Command.NumFramesUntilExec - NumFramesLeft);
 			}
diff --git a/Engine/Source/Runtime/AudioMixer/Public/Quartz/AudioMixerClockHandle.h b/Engine/Source/Runtime/AudioMixer/Public/Quartz/AudioMixerClockHandle.h
--- a/Engine/Source/Runtime/AudioMixer/Public/Quartz/AudioMixerClockHandle.h
+++ b/Engine/Source/Runtime/AudioMixer/Public/Quartz/AudioMixerClockHandle.h
@@ -148,7 +148,6 @@ public:
 
 	bool DoesClockExist(const UObject* WorldContextObject) const
 	{
-		check(RawHandle);
 		return RawHandle.DoesClockExist();
 	}
 
