Adding ability to interact with audio buses from arbitrary tasks and the game thread

diff --git a/Engine/Source/Runtime/AudioMixer/Private/AudioMixerBus.cpp b/Engine/Source/Runtime/AudioMixer/Private/AudioMixerBus.cpp
--- a/Engine/Source/Runtime/AudioMixer/Private/AudioMixerBus.cpp
+++ b/Engine/Source/Runtime/AudioMixer/Private/AudioMixerBus.cpp
@@ -19,7 +19,7 @@ namespace Audio
 		SetNumOutputChannels(NumChannels);
 
 		// Make a patch bus input to push mixed audio sent to audio bus from source manager
-		AudioBusInput = PatchMixerSplitter.AddNewInput(4096, 1.0f);
+		AudioBusInput = PatchMixerSplitter.AddNewInput(2*InNumChannels*NumFrames, 1.0f);
 	}
 
 	void FMixerAudioBus::SetNumOutputChannels(int32 InNumOutputChannels)
diff --git a/Engine/Source/Runtime/AudioMixer/Private/AudioMixerDevice.cpp b/Engine/Source/Runtime/AudioMixer/Private/AudioMixerDevice.cpp
--- a/Engine/Source/Runtime/AudioMixer/Private/AudioMixerDevice.cpp
+++ b/Engine/Source/Runtime/AudioMixer/Private/AudioMixerDevice.cpp
@@ -559,7 +559,8 @@ namespace Audio
 	{
 		LLM_SCOPE(ELLMTag::AudioMixer);
 
-
+		// Pump our command queue sending commands to the game thread
+		PumpGameThreadCommandQueue();
 	}
 
 	void FMixerDevice::UpdateHardware()
@@ -1525,6 +1526,11 @@ namespace Audio
 		CommandQueue.Enqueue(MoveTemp(Command));
 	}
 
+	void FMixerDevice::GameThreadMPSCCommand(TFunction<void()> InCommand)
+	{
+		GameThreadCommandQueue.Enqueue(MoveTemp(InCommand));
+	}
+	
 	void FMixerDevice::PumpCommandQueue()
 	{
 		// Execute the pushed lambda functions
@@ -1535,6 +1541,18 @@ namespace Audio
 		}
 	}
 
+	void FMixerDevice::PumpGameThreadCommandQueue()
+	{
+		TOptional Opt { GameThreadCommandQueue.Dequeue() };
+		while (Opt.IsSet())
+		{
+			TFunction<void()> Command = MoveTemp(Opt.GetValue());
+			Command();
+				
+			Opt = GameThreadCommandQueue.Dequeue();
+		}
+	}
+	
 	void FMixerDevice::FlushAudioRenderingCommands(bool bPumpSynchronously)
 	{
 		if (IsInitialized() && (FPlatformProcess::SupportsMultithreading() && !AudioMixerPlatform->IsNonRealtime()))
@@ -1869,11 +1887,12 @@ namespace Audio
 			return;
 		}
 
-		for (TStrongObjectPtr<UAudioBus>& Bus : DefaultAudioBuses)
+		for (TObjectIterator<UAudioBus> It; It; ++It)
 		{
-			if (Bus.IsValid())
+			UAudioBus* AudioBus = *It;
+			if (AudioBus)
 			{
-				StopAudioBus(Bus->GetUniqueID());
+				StopAudioBus(AudioBus->GetUniqueID());
 			}
 		}
 
@@ -2544,6 +2563,14 @@ namespace Audio
 				SourceManager->StartAudioBus(InAudioBusId, InNumChannels, bInIsAutomatic);
 			});
 		}
+		else
+		{
+			// If we're not the game thread, this needs to be on the game thread, so queue up a command to execute it on the game thread
+			GameThreadMPSCCommand([this, InAudioBusId, InNumChannels, bInIsAutomatic]
+			{
+				StartAudioBus(InAudioBusId, InNumChannels, bInIsAutomatic);	
+			});
+		}
 	}
 
 	void FMixerDevice::StopAudioBus(uint32 InAudioBusId)
@@ -2562,6 +2589,14 @@ namespace Audio
 				SourceManager->StopAudioBus(InAudioBusId);
 			});
 		}
+		else
+		{
+			// If we're not the game thread, this needs to be on the game thread, so queue up a command to execute it on the game thread
+			GameThreadMPSCCommand([this, InAudioBusId]
+			{
+				StopAudioBus(InAudioBusId);	
+			});
+		}
 	}
 
 	bool FMixerDevice::IsAudioBusActive(uint32 InAudioBusId) const
@@ -2575,23 +2610,49 @@ namespace Audio
 		return SourceManager->IsAudioBusActive(InAudioBusId);
 	}
 
-	Audio::FPatchOutputStrongPtr FMixerDevice::AddPatchForAudioBus(uint32 InAudioBusId, float InPatchGain)
-	{
-		check(SourceManager);
 
-		const int32 NumChannels = SourceManager->GetAudioBusNumChannels(InAudioBusId);
-		if (NumChannels > 0)
+	FPatchOutputStrongPtr FMixerDevice::AddPatchForAudioBus(uint32 InAudioBusId, float InPatchGain)
+	{
+		if (IsInGameThread())
 		{
+			if (ActiveAudioBuses_GameThread.Find(InAudioBusId))
+			{
+				const int32 NumOutputFrames = SourceManager->GetNumOutputFrames();
+				FPatchOutputStrongPtr StrongOutputPtr = MakeShareable(new FPatchOutput(2 * NumOutputFrames, InPatchGain));
+				FAudioThread::RunCommandOnAudioThread([this, StrongOutputPtr, InAudioBusId]() mutable
+				{
+					SourceManager->AddPatchOutputForAudioBus_AudioThread(InAudioBusId, StrongOutputPtr);
+				});
+				return StrongOutputPtr;
+			}
+			UE_LOG(LogAudioMixer, Warning, TEXT("Unable to add a patch output for audio bus because audio bus id '%d' is not active."), InAudioBusId);
+			return nullptr;
+		}
+		else
+		{
+			// Need to make a strong output patch even if this is not going to ever write to it since the bus may not be running
 			const int32 NumOutputFrames = SourceManager->GetNumOutputFrames();
-			FPatchOutputStrongPtr StrongOutputPtr = MakeShareable(new FPatchOutput(NumOutputFrames * NumChannels, InPatchGain));
-			SourceManager->AddPatchOutputForAudioBus(InAudioBusId, StrongOutputPtr);
+			FPatchOutputStrongPtr StrongOutputPtr = MakeShareable(new FPatchOutput(3 * NumOutputFrames, InPatchGain));
+
+			GameThreadMPSCCommand([this, StrongOutputPtr, InAudioBusId]
+			{
+				if (ActiveAudioBuses_GameThread.Find(InAudioBusId))
+				{
+					FAudioThread::RunCommandOnAudioThread([this, InAudioBusId, StrongOutputPtr]() mutable
+					{
+						SourceManager->AddPatchOutputForAudioBus_AudioThread(InAudioBusId, StrongOutputPtr);
+					});
+				}
+				else
+				{
+					UE_LOG(LogAudioMixer, Warning, TEXT("Unable to add a patch output for audio bus because audio bus id '%d' is not active."), InAudioBusId);
+				}
+			});
 			return StrongOutputPtr;
 		}
-
-		return nullptr;
 	}
 
-	Audio::FPatchOutputStrongPtr FMixerDevice::AddPatchForAudioBus_GameThread(uint32 InAudioBusId, float InPatchGain)
+	FPatchOutputStrongPtr FMixerDevice::AddPatchForAudioBus_GameThread(uint32 InAudioBusId, float InPatchGain)
 	{
 		FPatchOutputStrongPtr StrongOutputPtr;
 		
diff --git a/Engine/Source/Runtime/AudioMixer/Private/AudioMixerSourceManager.cpp b/Engine/Source/Runtime/AudioMixer/Private/AudioMixerSourceManager.cpp
--- a/Engine/Source/Runtime/AudioMixer/Private/AudioMixerSourceManager.cpp
+++ b/Engine/Source/Runtime/AudioMixer/Private/AudioMixerSourceManager.cpp
@@ -1138,10 +1138,21 @@ namespace Audio
 	void FMixerSourceManager::AddPatchOutputForAudioBus(uint32 InAudioBusId, const FPatchOutputStrongPtr& InPatchOutputStrongPtr)
 	{
 		AUDIO_MIXER_CHECK_AUDIO_PLAT_THREAD(MixerDevice);
-		TSharedPtr<FMixerAudioBus> AudioBusPtr = AudioBuses.FindRef(InAudioBusId);
-		if (AudioBusPtr.IsValid())
+		if (MixerDevice->IsAudioRenderingThread())
 		{
-			AudioBusPtr->AddNewPatchOutput(InPatchOutputStrongPtr);
+			TSharedPtr<FMixerAudioBus> AudioBusPtr = AudioBuses.FindRef(InAudioBusId);
+			if (AudioBusPtr.IsValid())
+			{
+				AudioBusPtr->AddNewPatchOutput(InPatchOutputStrongPtr);
+			}
+		}
+		else
+		{
+			// Queue up the command via MPSC command queue
+			AudioMixerThreadMPSCCommand([this, InAudioBusId, InPatchOutputStrongPtr]()
+			{
+				AddPatchOutputForAudioBus(InAudioBusId, InPatchOutputStrongPtr);
+			});
 		}
 	}
 
@@ -2911,6 +2922,10 @@ namespace Audio
 		}
 	}
 
+	void FMixerSourceManager::AudioMixerThreadMPSCCommand(TFunction<void()> InCommand)
+	{
+		MpscCommandQueue.Enqueue(MoveTemp(InCommand));
+	}
 
 	void FMixerSourceManager::AudioMixerThreadCommand(TFunction<void()> InFunction)
 	{
@@ -2955,6 +2970,16 @@ namespace Audio
 			}
 		}
 
+		// Pump the MPSC command queue
+		TOptional Opt{ MpscCommandQueue.Dequeue() };
+		while (Opt.IsSet())
+		{
+			TFunction<void()> Command = MoveTemp(Opt.GetValue());
+			Command();
+				
+			Opt = MpscCommandQueue.Dequeue();
+		}
+
 		int32 CurrentRenderThreadIndex = RenderThreadCommandBufferIndex.GetValue();
 
 		FCommands& Commands = CommandBuffers[CurrentRenderThreadIndex];
diff --git a/Engine/Source/Runtime/AudioMixer/Private/AudioMixerSourceManager.h b/Engine/Source/Runtime/AudioMixer/Private/AudioMixerSourceManager.h
--- a/Engine/Source/Runtime/AudioMixer/Private/AudioMixerSourceManager.h
+++ b/Engine/Source/Runtime/AudioMixer/Private/AudioMixerSourceManager.h
@@ -8,10 +8,9 @@
 #include "AudioMixerDevice.h"
 #include "AudioMixerSourceOutputBuffer.h"
 #include "AudioMixerSubmix.h"
-#include "Containers/Queue.h"
+#include "Containers/MpscQueue.h"
 #include "DSP/BufferVectorOperations.h"
 #include "DSP/EnvelopeFollower.h"
-#include "DSP/Filter.h"
 #include "DSP/InterpolatedOnePole.h"
 #include "DSP/ParamInterpolator.h"
 #include "IAudioExtensionPlugin.h"
@@ -20,6 +19,7 @@
 #include "Sound/QuartzQuantizationUtilities.h"
 #include "Stats/Stats.h"
 
+
 // Tracks the time it takes to up the source manager (computes source buffers, source effects, sample rate conversion)
 DECLARE_CYCLE_STAT_EXTERN(TEXT("Source Manager Update"), STAT_AudioMixerSourceManagerUpdate, STATGROUP_AudioMixer, AUDIOMIXER_API);
 
@@ -263,7 +263,12 @@ namespace Audio
 		void PumpCommandQueue();
 		void UpdatePendingReleaseData(bool bForceWait = false);
 		void FlushCommandQueue(bool bPumpCommandQueue = false);
+
+		// Pushes a TFUnction command into an MPSC queue from an arbitrary thread to the audio render thread
+		void AudioMixerThreadMPSCCommand(TFunction<void()> InCommand);
+		
 	private:
+		
 		void ReleaseSource(const int32 SourceId);
 		void BuildSourceEffectChain(const int32 SourceId, FSoundEffectSourceInitData& InitData, const TArray<FSourceEffectChainEntry>& SourceEffectChain, TArray<TSoundEffectSourcePtr>& OutSourceEffects);
 		void ResetSourceEffectChain(const int32 SourceId);
@@ -281,6 +286,7 @@ namespace Audio
 
 		void AudioMixerThreadCommand(TFunction<void()> InFunction);
 
+		
 		static const int32 NUM_BYTES_PER_SAMPLE = 2;
 
 		// Private class which perform source buffer processing in a worker task
@@ -332,7 +338,7 @@ namespace Audio
 		{
 			TArray<TFunction<void()>> SourceCommandQueue;
 		};
-
+		
 		FCommands CommandBuffers[2];
 		FThreadSafeCounter RenderThreadCommandBufferIndex;
 
@@ -341,6 +347,9 @@ namespace Audio
 
 		TArray<int32> DebugSoloSources;
 
+		// Command queue to communicate commands to teh source manager from arbitrary threads
+		TMpscQueue<TFunction<void()>> MpscCommandQueue;
+		
 		struct FSourceInfo
 		{
 			FSourceInfo() {}
diff --git a/Engine/Source/Runtime/AudioMixer/Public/AudioMixerDevice.h b/Engine/Source/Runtime/AudioMixer/Public/AudioMixerDevice.h
--- a/Engine/Source/Runtime/AudioMixer/Public/AudioMixerDevice.h
+++ b/Engine/Source/Runtime/AudioMixer/Public/AudioMixerDevice.h
@@ -290,12 +290,12 @@ namespace Audio
 		static bool IsEndpointSubmix(const USoundSubmixBase* InSubmix);
 
 		// Audio bus API
-		void StartAudioBus(uint32 InAudioBusId, int32 InNumChannels, bool bInIsAutomatic);
-		void StopAudioBus(uint32 InAudioBusId);
-		bool IsAudioBusActive(uint32 InAudioBusId) const;
+		virtual void StartAudioBus(uint32 InAudioBusId, int32 InNumChannels, bool bInIsAutomatic) override;
+		virtual void StopAudioBus(uint32 InAudioBusId) override;
+		virtual bool IsAudioBusActive(uint32 InAudioBusId) const override;
 
-		FPatchOutputStrongPtr AddPatchForAudioBus(uint32 InAudioBusId, float InPatchGain = 1.0f);
-		FPatchOutputStrongPtr AddPatchForAudioBus_GameThread(uint32 InAudioBusId, float InPatchGain = 1.0f);
+		virtual FPatchOutputStrongPtr AddPatchForAudioBus(uint32 InAudioBusId, float InPatchGain = 1.0f) override;
+		virtual FPatchOutputStrongPtr AddPatchForAudioBus_GameThread(uint32 InAudioBusId, float InPatchGain = 1.0f) override;
 
 		// Clock Manager for quantized event handling on Audio Render Thread
 		FQuartzClockManager QuantizedEventClockManager;
@@ -303,6 +303,8 @@ namespace Audio
 		// Pushes the command to a audio render thread command queue to be executed on render thread
 		void AudioRenderThreadCommand(TFunction<void()> Command);
 
+		// Pushes the command to a MPSC queue to be executed on the game thread
+		void GameThreadMPSCCommand(TFunction<void()> InCommand);
 
 	protected:
 
@@ -336,15 +338,14 @@ namespace Audio
 
 		void UnloadSoundSubmix(const USoundSubmixBase& SoundSubmix);
 
-	private:
-
 		bool IsMasterSubmixType(const USoundSubmixBase* InSubmix) const;
 		FMixerSubmixPtr GetMasterSubmixInstance(uint32 InSubmixId);
 		FMixerSubmixPtr GetMasterSubmixInstance(const USoundSubmixBase* InSubmix);
 		
 		// Pumps the audio render thread command queue
 		void PumpCommandQueue();
-
+		void PumpGameThreadCommandQueue();
+		
 		TArray<USoundSubmix*> MasterSubmixes;
 		TArray<FMixerSubmixPtr> MasterSubmixInstances;
 
@@ -429,6 +430,9 @@ namespace Audio
 		/** Command queue to send commands to audio render thread from game thread or audio thread. */
 		TQueue<TFunction<void()>> CommandQueue;
 
+		/** MPSC command queue to send commands to the game thread */
+		TMpscQueue<TFunction<void()>> GameThreadCommandQueue;
+		
 		/** Whether or not we generate output audio to test multi-platform mixer. */
 		bool bDebugOutputEnabled;
 
diff --git a/Engine/Source/Runtime/Engine/Private/AudioDevice.cpp b/Engine/Source/Runtime/Engine/Private/AudioDevice.cpp
--- a/Engine/Source/Runtime/Engine/Private/AudioDevice.cpp
+++ b/Engine/Source/Runtime/Engine/Private/AudioDevice.cpp
@@ -4392,6 +4392,8 @@ void FAudioDevice::Update(bool bGameTicking)
 
 	if (IsInGameThread())
 	{
+		UpdateGameThread();
+
 		// Make sure our referenced sound waves is up-to-date
 		UpdateReferencedSoundWaves();
 	}
diff --git a/Engine/Source/Runtime/Engine/Public/AudioDevice.h b/Engine/Source/Runtime/Engine/Public/AudioDevice.h
--- a/Engine/Source/Runtime/Engine/Public/AudioDevice.h
+++ b/Engine/Source/Runtime/Engine/Public/AudioDevice.h
@@ -964,6 +964,29 @@ public:
 		return nullptr;
 	}
 
+	virtual void StartAudioBus(uint32 InAudioBusId, int32 InNumChannels, bool bInIsAutomatic)
+	{
+	}
+
+	virtual void StopAudioBus(uint32 InAudioBusId)
+	{
+	}
+
+	virtual bool IsAudioBusActive(uint32 InAudioBusId) const
+	{
+		return false;
+	}
+
+	virtual Audio::FPatchOutputStrongPtr AddPatchForAudioBus(uint32 InAudioBusId, float InPatchGain = 1.0f)
+	{
+		return nullptr;
+	}
+
+	virtual Audio::FPatchOutputStrongPtr AddPatchForAudioBus_GameThread(uint32 InAudioBusId, float InPatchGain = 1.0f)
+	{
+		return nullptr;
+	}
+	
 	virtual void InitSoundEffectPresets() {}
 
 	/**
