[KillPendingKill] Create audio component as transient and copy creation method from parent component to avoid taking pointers to destroyed components when rerunning construction scripts.

diff --git a/Engine/Source/Runtime/AudioMixer/Private/Components/SynthComponent.cpp b/Engine/Source/Runtime/AudioMixer/Private/Components/SynthComponent.cpp
--- a/Engine/Source/Runtime/AudioMixer/Private/Components/SynthComponent.cpp
+++ b/Engine/Source/Runtime/AudioMixer/Private/Components/SynthComponent.cpp
@@ -266,7 +266,8 @@ void USynthComponent::CreateAudioComponent()
 	if (!AudioComponent)
 	{
 		// Create the audio component which will be used to play the procedural sound wave
-		AudioComponent = NewObject<UAudioComponent>(this);
+		AudioComponent = NewObject<UAudioComponent>(this, NAME_None, RF_Transactional | RF_Transient | RF_TextExportTransient);
+		AudioComponent->CreationMethod = CreationMethod;
 
 		AudioComponent->OnAudioSingleEnvelopeValueNative.AddUObject(this, &USynthComponent::OnAudioComponentEnvelopeValue);
 
