[Quartz] Fix for beat-counting edge case with BPM changes.
#p4v-preflight-copy 20948998

diff --git a/Engine/Source/Runtime/AudioMixer/Private/Quartz/QuartzMetronome.cpp b/Engine/Source/Runtime/AudioMixer/Private/Quartz/QuartzMetronome.cpp
--- a/Engine/Source/Runtime/AudioMixer/Private/Quartz/QuartzMetronome.cpp
+++ b/Engine/Source/Runtime/AudioMixer/Private/Quartz/QuartzMetronome.cpp
@@ -72,6 +72,7 @@ namespace Audio
 			}
 		}
 
+
 		// update transport
 		if (ToUpdateBitField & (1 << static_cast<int>(EQuartzCommandQuantization::Bar)))
 		{
@@ -111,12 +112,9 @@ namespace Audio
 		// ratio between new and old rates
 		const double Ratio = InNewTickRate.GetFramesPerTick() / CurrentTickRate.GetFramesPerTick();
 
-		if (NumFramesLeft)
+		for (double& Value : FramesLeftInMusicalDuration.FramesInTimeValueInternal)
 		{
-			for (double& Value : FramesLeftInMusicalDuration.FramesInTimeValueInternal)
-			{
-				Value = NumFramesLeft + Ratio * (Value - NumFramesLeft);
-			}
+			Value = NumFramesLeft + Ratio * (Value - NumFramesLeft);
 		}
 
 		CurrentTickRate = InNewTickRate;
