Adding speculative defense around reading past the end of SourceInfo.CurrentPCMBuffer.

diff --git a/Engine/Source/Runtime/AudioMixer/Private/AudioMixerSourceManager.cpp b/Engine/Source/Runtime/AudioMixer/Private/AudioMixerSourceManager.cpp
--- a/Engine/Source/Runtime/AudioMixer/Private/AudioMixerSourceManager.cpp
+++ b/Engine/Source/Runtime/AudioMixer/Private/AudioMixerSourceManager.cpp
@@ -1813,22 +1813,36 @@ namespace Audio
 		{
 			// Grab the float PCM audio data (which could be a new audio chunk from previous ReadSourceFrame call)
 			const float* AudioData = SourceInfo.CurrentPCMBuffer->AudioData.GetData();
+			const int32 CurrentSampleIndex = SourceInfo.CurrentFrameIndex * NumChannels;
 			const int32 NextSampleIndex = (SourceInfo.CurrentFrameIndex + 1)  * NumChannels;
+			const int32 AudioDataNum = SourceInfo.CurrentPCMBuffer->AudioData.Num();
 
-			if (bReadCurrentFrame)
+			if(ensureAlwaysMsgf(AudioDataNum >= NextSampleIndex + NumChannels
+				, TEXT("Bailing due to bad CurrentPCMBuffer:  AudioData.Num() = %i, NextSampleIndex = %i, NumChannels = %i"), AudioDataNum, NextSampleIndex, NumChannels))
 			{
-				const int32 CurrentSampleIndex = SourceInfo.CurrentFrameIndex * NumChannels;
-				for (int32 Channel = 0; Channel < NumChannels; ++Channel)
+				if (bReadCurrentFrame)
 				{
-					SourceInfo.CurrentFrameValues[Channel] = AudioData[CurrentSampleIndex + Channel];
-					SourceInfo.NextFrameValues[Channel] = AudioData[NextSampleIndex + Channel];
+					for (int32 Channel = 0; Channel < NumChannels; ++Channel)
+					{
+						SourceInfo.CurrentFrameValues[Channel] = AudioData[CurrentSampleIndex + Channel];
+						SourceInfo.NextFrameValues[Channel] = AudioData[NextSampleIndex + Channel];
+					}
+				}
+				else if (NextSampleIndex != SourceInfo.CurrentPCMBuffer->AudioData.Num())
+				{
+					for (int32 Channel = 0; Channel < NumChannels; ++Channel)
+					{
+						SourceInfo.NextFrameValues[Channel] = AudioData[NextSampleIndex + Channel];
+					}
 				}
 			}
-			else if (NextSampleIndex != SourceInfo.CurrentPCMBuffer->AudioData.Num())
+			else
 			{
+				// fill w/ silence instead of the bad access
 				for (int32 Channel = 0; Channel < NumChannels; ++Channel)
 				{
-					SourceInfo.NextFrameValues[Channel] = AudioData[NextSampleIndex + Channel];
+					SourceInfo.CurrentFrameValues[Channel] = 0.f;
+					SourceInfo.NextFrameValues[Channel] = 0.f;
 				}
 			}
 		}
