Kill pending kill - SynthSound referencing garbage components

diff --git a/Engine/Source/Runtime/AudioMixer/Private/Components/SynthComponent.cpp b/Engine/Source/Runtime/AudioMixer/Private/Components/SynthComponent.cpp
--- a/Engine/Source/Runtime/AudioMixer/Private/Components/SynthComponent.cpp
+++ b/Engine/Source/Runtime/AudioMixer/Private/Components/SynthComponent.cpp
@@ -9,7 +9,6 @@
 
 USynthSound::USynthSound(const FObjectInitializer& ObjectInitializer)
 	: Super(ObjectInitializer)
-	, OwningSynthComponent(nullptr)
 {
 }
 
@@ -47,7 +46,7 @@ void USynthSound::StartOnAudioDevice(FAudioDevice* InAudioDevice)
 
 void USynthSound::OnBeginGenerate()
 {
-	if (ensure(OwningSynthComponent))
+	if (ensure(OwningSynthComponent.IsValid()))
 	{
 		OwningSynthComponent->OnBeginGenerate();
 	}
@@ -65,7 +64,7 @@ int32 USynthSound::OnGeneratePCMAudio(TArray<uint8>& OutAudio, int32 NumSamples)
 		OutAudio.AddZeroed(NumSamples * sizeof(float));
 
 		// Mark pending kill can null this out on the game thread in rare cases.
-		if (!OwningSynthComponent)
+		if (!OwningSynthComponent.IsValid())
 		{
 			return 0;
 		}
@@ -79,7 +78,7 @@ int32 USynthSound::OnGeneratePCMAudio(TArray<uint8>& OutAudio, int32 NumSamples)
 		FloatBuffer.AddZeroed(NumSamples * sizeof(float));
 
 		// Mark pending kill can null this out on the game thread in rare cases.
-		if (!OwningSynthComponent)
+		if (!OwningSynthComponent.IsValid())
 		{
 			return 0;
 		}
@@ -103,7 +102,7 @@ int32 USynthSound::OnGeneratePCMAudio(TArray<uint8>& OutAudio, int32 NumSamples)
 void USynthSound::OnEndGenerate()
 {
 	// Mark pending kill can null this out on the game thread in rare cases.
-	if (OwningSynthComponent)
+	if (OwningSynthComponent.IsValid())
 	{
 		OwningSynthComponent->OnEndGenerate();
 	}
@@ -111,7 +110,7 @@ void USynthSound::OnEndGenerate()
 
 ISoundGeneratorPtr USynthSound::CreateSoundGenerator(const FSoundGeneratorInitParams& InParams)
 {
-	if (OwningSynthComponent)
+	if (OwningSynthComponent.IsValid())
 	{
 		return OwningSynthComponent->CreateSoundGeneratorInternal(InParams);
 	}
diff --git a/Engine/Source/Runtime/AudioMixer/Public/Components/SynthComponent.h b/Engine/Source/Runtime/AudioMixer/Public/Components/SynthComponent.h
--- a/Engine/Source/Runtime/AudioMixer/Public/Components/SynthComponent.h
+++ b/Engine/Source/Runtime/AudioMixer/Public/Components/SynthComponent.h
@@ -56,7 +56,7 @@ class AUDIOMIXER_API USynthSound : public USoundWaveProcedural
 
 protected:
 	UPROPERTY()
-	TObjectPtr<USynthComponent> OwningSynthComponent;
+	TWeakObjectPtr<USynthComponent> OwningSynthComponent = nullptr;
 
 	TArray<float> FloatBuffer;
 	bool bAudioMixer;
@@ -64,7 +64,7 @@ protected:
 public:
 	USynthComponent* GetOwningSynthComponent()
 	{
-		return OwningSynthComponent;
+		return OwningSynthComponent.Get();
 	}
 };
 
