Allow submix sends to determine the send stage.

diff --git a/Engine/Source/Runtime/AudioMixer/Private/AudioMixerSource.cpp b/Engine/Source/Runtime/AudioMixer/Private/AudioMixerSource.cpp
--- a/Engine/Source/Runtime/AudioMixer/Private/AudioMixerSource.cpp
+++ b/Engine/Source/Runtime/AudioMixer/Private/AudioMixerSource.cpp
@@ -43,6 +43,21 @@ static FAutoConsoleCommand GSetAudioMixerSourceFadeMin(
 
 namespace Audio
 {
+	namespace MixerSourcePrivate
+	{
+		EMixerSourceSubmixSendStage SubmixSendStageToMixerSourceSubmixSendStage(ESubmixSendStage InSendStage)
+		{
+			switch(InSendStage)
+			{
+				case ESubmixSendStage::PreDistanceAttenuation:
+					return EMixerSourceSubmixSendStage::PreDistanceAttenuation;
+
+				case ESubmixSendStage::PostDistanceAttenuation:
+				default:
+					return EMixerSourceSubmixSendStage::PostDistanceAttenuation;
+			}
+		}
+	}
 	namespace ModulationUtils
 	{
 		static const FSoundModulationDestinationSettings DefaultDestination;
@@ -1488,8 +1503,9 @@ namespace Audio
 					}
 				}
 
-				// set the level for this send
-				MixerSourceVoice->SetSubmixSendInfo(SubmixInstance, SendLevel);
+				// set the level and stage for this send
+				EMixerSourceSubmixSendStage SendStage = MixerSourcePrivate::SubmixSendStageToMixerSourceSubmixSendStage(SendInfo.SendStage);
+				MixerSourceVoice->SetSubmixSendInfo(SubmixInstance, SendLevel, SendStage);
 			}
 		}
  		
diff --git a/Engine/Source/Runtime/AudioMixer/Private/AudioMixerSourceVoice.cpp b/Engine/Source/Runtime/AudioMixer/Private/AudioMixerSourceVoice.cpp
--- a/Engine/Source/Runtime/AudioMixer/Private/AudioMixerSourceVoice.cpp
+++ b/Engine/Source/Runtime/AudioMixer/Private/AudioMixerSourceVoice.cpp
@@ -330,7 +330,7 @@ namespace Audio
 		return SourceManager->GetListenerRotation(SourceId);
 	}
 
-	void FMixerSourceVoice::SetSubmixSendInfo(FMixerSubmixWeakPtr Submix, const float SendLevel)
+	void FMixerSourceVoice::SetSubmixSendInfo(FMixerSubmixWeakPtr Submix, const float SendLevel, const EMixerSourceSubmixSendStage SendStage/* = EMixerSourceSubmixSendStage::PostDistanceAttenuation*/)
 	{
 		AUDIO_MIXER_CHECK_GAME_THREAD(MixerDevice);
 
@@ -345,12 +345,15 @@ namespace Audio
 				NewSubmixSend.Submix = Submix;
 				NewSubmixSend.SendLevel = SendLevel;
 				NewSubmixSend.bIsMainSend = false;
+				NewSubmixSend.SubmixSendStage = SendStage;
+
 				SubmixSends.Add(SubmixPtr->GetId(), NewSubmixSend);
 				SourceManager->SetSubmixSendInfo(SourceId, NewSubmixSend);
 			}
-			else if (!FMath::IsNearlyEqual(SubmixSend->SendLevel, SendLevel))
+			else if (!FMath::IsNearlyEqual(SubmixSend->SendLevel, SendLevel) || SubmixSend->SubmixSendStage != SendStage)
 			{
 				SubmixSend->SendLevel = SendLevel;
+				SubmixSend->SubmixSendStage = SendStage;
 				SourceManager->SetSubmixSendInfo(SourceId, *SubmixSend);
 			}
 		}
diff --git a/Engine/Source/Runtime/AudioMixer/Private/AudioMixerSourceVoice.h b/Engine/Source/Runtime/AudioMixer/Private/AudioMixerSourceVoice.h
--- a/Engine/Source/Runtime/AudioMixer/Private/AudioMixerSourceVoice.h
+++ b/Engine/Source/Runtime/AudioMixer/Private/AudioMixerSourceVoice.h
@@ -116,7 +116,7 @@ namespace Audio
 		const FQuat GetListenerRotationForVoice() const;
 
 		// Sets the submix send levels
-		void SetSubmixSendInfo(FMixerSubmixWeakPtr Submix, const float SendLevel);
+		void SetSubmixSendInfo(FMixerSubmixWeakPtr Submix, const float SendLevel, const EMixerSourceSubmixSendStage SendStage = EMixerSourceSubmixSendStage::PostDistanceAttenuation);
 
 		// Clears the submix send to the given submix
 		void ClearSubmixSendInfo(FMixerSubmixWeakPtr Submix);
