Fix Invalid Audio Bus crashes in Audio Mixer Blueprint Library

diff --git a/Engine/Source/Runtime/AudioMixer/Private/AudioMixerBlueprintLibrary.cpp b/Engine/Source/Runtime/AudioMixer/Private/AudioMixerBlueprintLibrary.cpp
--- a/Engine/Source/Runtime/AudioMixer/Private/AudioMixerBlueprintLibrary.cpp
+++ b/Engine/Source/Runtime/AudioMixer/Private/AudioMixerBlueprintLibrary.cpp
@@ -644,6 +644,11 @@ float UAudioMixerBlueprintLibrary::TrimAudioCache(float InMegabytesToFree)
 
 void UAudioMixerBlueprintLibrary::StartAudioBus(const UObject* WorldContextObject, UAudioBus* AudioBus)
 {
+	if (!AudioBus)
+	{
+		UE_LOG(LogAudioMixer, Error, TEXT("Start Audio Bus called with an invalid Audio Bus."));
+		return;
+	}
 	if (Audio::FMixerDevice* MixerDevice = FAudioDeviceManager::GetAudioMixerDeviceFromWorldContext(WorldContextObject))
 	{
 		uint32 AudioBusId = AudioBus->GetUniqueID();
@@ -658,6 +663,11 @@ void UAudioMixerBlueprintLibrary::StartAudioBus(const UObject* WorldContextObjec
 
 void UAudioMixerBlueprintLibrary::StopAudioBus(const UObject* WorldContextObject, UAudioBus* AudioBus)
 {
+	if (!AudioBus)
+	{
+		UE_LOG(LogAudioMixer, Error, TEXT("Stop Audio Bus called with an invalid Audio Bus."));
+		return;
+	}
 	if (Audio::FMixerDevice* MixerDevice = FAudioDeviceManager::GetAudioMixerDeviceFromWorldContext(WorldContextObject))
 	{
 		uint32 AudioBusId = AudioBus->GetUniqueID();
@@ -671,6 +681,11 @@ void UAudioMixerBlueprintLibrary::StopAudioBus(const UObject* WorldContextObject
 
 bool UAudioMixerBlueprintLibrary::IsAudioBusActive(const UObject* WorldContextObject, UAudioBus* AudioBus)
 {
+	if (!AudioBus)
+	{
+		UE_LOG(LogAudioMixer, Error, TEXT("Is Audio Bus Active called with an invalid Audio Bus."));
+		return false;
+	}
 	if (Audio::FMixerDevice* MixerDevice = FAudioDeviceManager::GetAudioMixerDeviceFromWorldContext(WorldContextObject))
 	{
 		uint32 AudioBusId = AudioBus->GetUniqueID();
