Audio Telemetry Clean-Up

diff --git a/Engine/Plugins/AudioGameplayVolume/Source/AudioGameplayVolume/Private/AudioGameplayVolumeProxy.cpp b/Engine/Plugins/AudioGameplayVolume/Source/AudioGameplayVolume/Private/AudioGameplayVolumeProxy.cpp
--- a/Engine/Plugins/AudioGameplayVolume/Source/AudioGameplayVolume/Private/AudioGameplayVolumeProxy.cpp
+++ b/Engine/Plugins/AudioGameplayVolume/Source/AudioGameplayVolume/Private/AudioGameplayVolumeProxy.cpp
@@ -1,11 +1,11 @@
 // Copyright Epic Games, Inc. All Rights Reserved.
 
 #include "AudioGameplayVolumeProxy.h"
+#include "AudioAnalytics.h"
 #include "AudioGameplayVolumeMutator.h"
 #include "AudioGameplayVolumeSubsystem.h"
 #include "AudioGameplayVolumeLogs.h"
 #include "AudioGameplayVolumeComponent.h"
-#include "EngineAnalytics.h"
 #include "Interfaces/IAudioGameplayCondition.h"
 #include "Components/BrushComponent.h"
 #include "Components/PrimitiveComponent.h"
@@ -53,11 +53,7 @@ void UAudioGameplayVolumeProxy::InitFromComponent(const UAudioGameplayVolumeComp
 		}
 	}
 
-	if (FEngineAnalytics::IsAvailable())
-	{
-		FEngineAnalytics::GetProvider().RecordEvent(TEXT("Audio.Usage.AudioGameplayVolume.InitializedFromComponent"));
-	}
-
+	Audio::Analytics::RecordEvent_Usage(TEXT("AudioGameplayVolume.InitializedFromComponent"));
 }
 
 void UAudioGameplayVolumeProxy::FindMutatorPriority(FAudioProxyMutatorPriorities& Priorities) const
diff --git a/Engine/Plugins/Experimental/MotoSynth/Source/MotoSynthEditor/Private/MotoSynthSourceFactory.cpp b/Engine/Plugins/Experimental/MotoSynth/Source/MotoSynthEditor/Private/MotoSynthSourceFactory.cpp
--- a/Engine/Plugins/Experimental/MotoSynth/Source/MotoSynthEditor/Private/MotoSynthSourceFactory.cpp
+++ b/Engine/Plugins/Experimental/MotoSynth/Source/MotoSynthEditor/Private/MotoSynthSourceFactory.cpp
@@ -1,7 +1,7 @@
 // Copyright Epic Games, Inc. All Rights Reserved.
 #include "MotoSynthSourceFactory.h"
+#include "AudioAnalytics.h"
 #include "MotoSynthSourceAsset.h"
-#include "EngineAnalytics.h"
 #include "MotoSynthEditorModule.h"
 #include "UObject/ObjectMacros.h"
 #include "UObject/Object.h"
@@ -37,10 +37,7 @@ UMotoSynthPresetFactory::UMotoSynthPresetFactory(const FObjectInitializer& Objec
 UObject* UMotoSynthPresetFactory::FactoryCreateNew(UClass* Class, UObject* InParent, FName InName, EObjectFlags Flags, UObject* Context, FFeedbackContext* Warn)
 {
 	UMotoSynthPreset* NewAsset = NewObject<UMotoSynthPreset >(InParent, InName, Flags);
-	if (FEngineAnalytics::IsAvailable())
-	{
-		FEngineAnalytics::GetProvider().RecordEvent(TEXT("Audio.Usage.MotoSynth.PresetCreated"));
-	}
+	Audio::Analytics::RecordEvent_Usage(TEXT("MotoSynth.PresetCreated"));
 	return NewAsset;
 }
 
diff --git a/Engine/Plugins/Runtime/AudioCapture/Source/AudioCapture/Private/AudioCaptureComponent.cpp b/Engine/Plugins/Runtime/AudioCapture/Source/AudioCapture/Private/AudioCaptureComponent.cpp
--- a/Engine/Plugins/Runtime/AudioCapture/Source/AudioCapture/Private/AudioCaptureComponent.cpp
+++ b/Engine/Plugins/Runtime/AudioCapture/Source/AudioCapture/Private/AudioCaptureComponent.cpp
@@ -1,7 +1,7 @@
 // Copyright Epic Games, Inc. All Rights Reserved.
 
 #include "AudioCaptureComponent.h"
-#include "EngineAnalytics.h"
+#include "AudioAnalytics.h"
 
 static const unsigned int MaxBufferSize = 2 * 5 * 48000;
 
@@ -31,10 +31,7 @@ bool UAudioCaptureComponent::Init(int32& SampleRate)
 			// This may fail if capture synths aren't supported on a given platform or if something went wrong with the capture device
 			bIsStreamOpen = CaptureSynth.OpenDefaultStream();
 
-			if (FEngineAnalytics::IsAvailable())
-			{
-				FEngineAnalytics::GetProvider().RecordEvent(TEXT("Audio.Usage.AudioCapture.AudioCaptureComponentInitialized"));
-			}
+			Audio::Analytics::RecordEvent_Usage(TEXT("AudioCapture.AudioCaptureComponentInitialized"));
 
 			return true;
 		}
diff --git a/Engine/Plugins/Runtime/AudioModulation/Source/AudioModulationEditor/Private/Factories/SoundControlBusFactory.cpp b/Engine/Plugins/Runtime/AudioModulation/Source/AudioModulationEditor/Private/Factories/SoundControlBusFactory.cpp
--- a/Engine/Plugins/Runtime/AudioModulation/Source/AudioModulationEditor/Private/Factories/SoundControlBusFactory.cpp
+++ b/Engine/Plugins/Runtime/AudioModulation/Source/AudioModulationEditor/Private/Factories/SoundControlBusFactory.cpp
@@ -2,7 +2,7 @@
 
 #include "SoundControlBusFactory.h"
 
-#include "EngineAnalytics.h"
+#include "AudioAnalytics.h"
 #include "SoundControlBus.h"
 
 
@@ -17,9 +17,6 @@ USoundControlBusFactory::USoundControlBusFactory(const FObjectInitializer& Objec
 
 UObject* USoundControlBusFactory::FactoryCreateNew(UClass* Class, UObject* InParent, FName Name, EObjectFlags Flags, UObject* Context, FFeedbackContext* Warn)
 {
-	if (FEngineAnalytics::IsAvailable())
-	{
-		FEngineAnalytics::GetProvider().RecordEvent(TEXT("Audio.Usage.AudioModulation.ControlBusCreated"));
-	}
+	Audio::Analytics::RecordEvent_Usage(TEXT("AudioModulation.ControlBusCreated"));
 	return NewObject<USoundControlBus>(InParent, Name, Flags);
 }
\ No newline at end of file
diff --git a/Engine/Plugins/Runtime/AudioModulation/Source/AudioModulationEditor/Private/Factories/SoundControlBusMixFactory.cpp b/Engine/Plugins/Runtime/AudioModulation/Source/AudioModulationEditor/Private/Factories/SoundControlBusMixFactory.cpp
--- a/Engine/Plugins/Runtime/AudioModulation/Source/AudioModulationEditor/Private/Factories/SoundControlBusMixFactory.cpp
+++ b/Engine/Plugins/Runtime/AudioModulation/Source/AudioModulationEditor/Private/Factories/SoundControlBusMixFactory.cpp
@@ -2,7 +2,7 @@
 
 #include "SoundControlBusMixFactory.h"
 
-#include "EngineAnalytics.h"
+#include "AudioAnalytics.h"
 #include "SoundControlBusMix.h"
 
 
@@ -17,9 +17,6 @@ USoundControlBusMixFactory::USoundControlBusMixFactory(const FObjectInitializer&
 
 UObject* USoundControlBusMixFactory::FactoryCreateNew(UClass* Class, UObject* InParent, FName Name, EObjectFlags Flags, UObject* Context, FFeedbackContext* Warn)
 {
-	if (FEngineAnalytics::IsAvailable())
-	{
-		FEngineAnalytics::GetProvider().RecordEvent(TEXT("Audio.Usage.AudioModulation.ControlBusMixCreated"));
-	}
+	Audio::Analytics::RecordEvent_Usage(TEXT("AudioModulation.ControlBusMixCreated"));
 	return NewObject<USoundControlBusMix>(InParent, Name, Flags);
 }
diff --git a/Engine/Plugins/Runtime/AudioModulation/Source/AudioModulationEditor/Private/Factories/SoundModulationGeneratorFactory.cpp b/Engine/Plugins/Runtime/AudioModulation/Source/AudioModulationEditor/Private/Factories/SoundModulationGeneratorFactory.cpp
--- a/Engine/Plugins/Runtime/AudioModulation/Source/AudioModulationEditor/Private/Factories/SoundModulationGeneratorFactory.cpp
+++ b/Engine/Plugins/Runtime/AudioModulation/Source/AudioModulationEditor/Private/Factories/SoundModulationGeneratorFactory.cpp
@@ -2,9 +2,9 @@
 
 #include "SoundModulationGeneratorFactory.h"
 
+#include "AudioAnalytics.h"
 #include "ClassViewerFilter.h"
 #include "ClassViewerModule.h"
-#include "EngineAnalytics.h"
 #include "Kismet2/SClassPickerDialog.h"
 #include "SoundModulationGenerator.h"
 #include "Templates/SharedPointer.h"
@@ -76,10 +76,7 @@ UObject* USoundModulationGeneratorFactory::FactoryCreateNew(UClass* InClass, UOb
 {
 	if (GeneratorClass)
 	{
-		if (FEngineAnalytics::IsAvailable())
-		{
-			FEngineAnalytics::GetProvider().RecordEvent(TEXT("Audio.Usage.AudioModulation.ModulationGeneratorCreated"));
-		}
+		Audio::Analytics::RecordEvent_Usage(TEXT("AudioModulation.ModulationGeneratorCreated"));
 		return NewObject<USoundModulationGenerator>(InParent, GeneratorClass, InName, Flags);
 	}
 
diff --git a/Engine/Plugins/Runtime/AudioModulation/Source/AudioModulationEditor/Private/Factories/SoundModulationParameterFactory.cpp b/Engine/Plugins/Runtime/AudioModulation/Source/AudioModulationEditor/Private/Factories/SoundModulationParameterFactory.cpp
--- a/Engine/Plugins/Runtime/AudioModulation/Source/AudioModulationEditor/Private/Factories/SoundModulationParameterFactory.cpp
+++ b/Engine/Plugins/Runtime/AudioModulation/Source/AudioModulationEditor/Private/Factories/SoundModulationParameterFactory.cpp
@@ -2,8 +2,8 @@
 
 #include "SoundModulationParameterFactory.h"
 
+#include "AudioAnalytics.h"
 #include "ClassViewerFilter.h"
-#include "EngineAnalytics.h"
 #include "SoundModulationParameter.h"
 #include "Kismet2/SClassPickerDialog.h"
 
@@ -73,10 +73,7 @@ UObject* USoundModulationParameterFactory::FactoryCreateNew(UClass* InClass, UOb
 	{
 		NewParameter = NewObject<USoundModulationParameter>(InParent, ParameterClass, InName, Flags);
 
-		if (FEngineAnalytics::IsAvailable())
-		{
-			FEngineAnalytics::GetProvider().RecordEvent(TEXT("Audio.Usage.AudioModulation.ModulationParameterCreated"));
-		}
+		Audio::Analytics::RecordEvent_Usage(TEXT("AudioModulation.ModulationParameterCreated"));
 	}
 
 	return NewParameter;
diff --git a/Engine/Plugins/Runtime/AudioModulation/Source/AudioModulationEditor/Private/Factories/SoundModulationPatchFactory.cpp b/Engine/Plugins/Runtime/AudioModulation/Source/AudioModulationEditor/Private/Factories/SoundModulationPatchFactory.cpp
--- a/Engine/Plugins/Runtime/AudioModulation/Source/AudioModulationEditor/Private/Factories/SoundModulationPatchFactory.cpp
+++ b/Engine/Plugins/Runtime/AudioModulation/Source/AudioModulationEditor/Private/Factories/SoundModulationPatchFactory.cpp
@@ -1,7 +1,7 @@
 // Copyright Epic Games, Inc. All Rights Reserved.
 #include "SoundModulationPatchFactory.h"
 
-#include "EngineAnalytics.h"
+#include "AudioAnalytics.h"
 #include "SoundModulationPatch.h"
 
 
@@ -16,9 +16,6 @@ USoundModulationPatchFactory::USoundModulationPatchFactory(const FObjectInitiali
 
 UObject* USoundModulationPatchFactory::FactoryCreateNew(UClass* Class, UObject* InParent, FName Name, EObjectFlags Flags, UObject* Context, FFeedbackContext* Warn)
 {
-	if (FEngineAnalytics::IsAvailable())
-	{
-		FEngineAnalytics::GetProvider().RecordEvent(TEXT("Audio.Usage.AudioModulation.ParameterPatchCreated"));
-	}
+	Audio::Analytics::RecordEvent_Usage(TEXT("AudioModulation.ParameterPatchCreated"));
 	return NewObject<USoundModulationPatch>(InParent, Name, Flags);
 }
diff --git a/Engine/Plugins/Runtime/AudioSynesthesia/Source/AudioSynesthesiaEditor/Private/AudioSynesthesiaNRTFactory.cpp b/Engine/Plugins/Runtime/AudioSynesthesia/Source/AudioSynesthesiaEditor/Private/AudioSynesthesiaNRTFactory.cpp
--- a/Engine/Plugins/Runtime/AudioSynesthesia/Source/AudioSynesthesiaEditor/Private/AudioSynesthesiaNRTFactory.cpp
+++ b/Engine/Plugins/Runtime/AudioSynesthesia/Source/AudioSynesthesiaEditor/Private/AudioSynesthesiaNRTFactory.cpp
@@ -2,9 +2,9 @@
 
 #include "AudioSynesthesiaNRTFactory.h"
 
+#include "AudioAnalytics.h"
 #include "ClassViewerModule.h"
 #include "ClassViewerFilter.h"
-#include "EngineAnalytics.h"
 #include "Kismet2/SClassPickerDialog.h"
 #include "Modules/ModuleManager.h"
 #include "AudioSynesthesiaClassFilter.h"
@@ -57,10 +57,7 @@ UObject* UAudioSynesthesiaNRTFactory::FactoryCreateNew(UClass* InClass, UObject*
 	{
 		NewAudioSynesthesiaNRT = NewObject<UAudioSynesthesiaNRT>(InParent, AudioSynesthesiaNRTClass, InName, Flags);
 
-		if (FEngineAnalytics::IsAvailable())
-		{
-			FEngineAnalytics::GetProvider().RecordEvent(TEXT("Audio.Usage.AudioSynesthesia.NRTCreated"));
-		}
+		Audio::Analytics::RecordEvent_Usage(TEXT("AudioSynesthesia.NRTCreated"));
 	}
 	return NewAudioSynesthesiaNRT;
 }
diff --git a/Engine/Plugins/Runtime/AudioSynesthesia/Source/AudioSynesthesiaEditor/Private/AudioSynesthesiaNRTSettingsFactory.cpp b/Engine/Plugins/Runtime/AudioSynesthesia/Source/AudioSynesthesiaEditor/Private/AudioSynesthesiaNRTSettingsFactory.cpp
--- a/Engine/Plugins/Runtime/AudioSynesthesia/Source/AudioSynesthesiaEditor/Private/AudioSynesthesiaNRTSettingsFactory.cpp
+++ b/Engine/Plugins/Runtime/AudioSynesthesia/Source/AudioSynesthesiaEditor/Private/AudioSynesthesiaNRTSettingsFactory.cpp
@@ -2,9 +2,9 @@
 
 #include "AudioSynesthesiaNRTSettingsFactory.h"
 
+#include "AudioAnalytics.h"
 #include "ClassViewerModule.h"
 #include "ClassViewerFilter.h"
-#include "EngineAnalytics.h"
 #include "Kismet2/SClassPickerDialog.h"
 #include "Modules/ModuleManager.h"
 #include "AudioSynesthesiaClassFilter.h"
@@ -55,10 +55,7 @@ UObject* UAudioSynesthesiaNRTSettingsFactory::FactoryCreateNew(UClass* InClass,
 	if (AudioSynesthesiaNRTSettingsClass != nullptr)
 	{
 		NewAudioSynesthesiaNRTSettings = NewObject<UAudioSynesthesiaNRTSettings>(InParent, AudioSynesthesiaNRTSettingsClass, InName, Flags);
-		if (FEngineAnalytics::IsAvailable())
-		{
-			FEngineAnalytics::GetProvider().RecordEvent(TEXT("Audio.Usage.AudioSynesthesia.NRTSettingsCreated"));
-		}
+		Audio::Analytics::RecordEvent_Usage(TEXT("AudioSynesthesia.NRTSettingsCreated"));
 	}
 	return NewAudioSynesthesiaNRTSettings;
 }
diff --git a/Engine/Plugins/Runtime/AudioSynesthesia/Source/AudioSynesthesiaEditor/Private/AudioSynesthesiaSettingsFactory.cpp b/Engine/Plugins/Runtime/AudioSynesthesia/Source/AudioSynesthesiaEditor/Private/AudioSynesthesiaSettingsFactory.cpp
--- a/Engine/Plugins/Runtime/AudioSynesthesia/Source/AudioSynesthesiaEditor/Private/AudioSynesthesiaSettingsFactory.cpp
+++ b/Engine/Plugins/Runtime/AudioSynesthesia/Source/AudioSynesthesiaEditor/Private/AudioSynesthesiaSettingsFactory.cpp
@@ -2,9 +2,9 @@
 
 #include "AudioSynesthesiaSettingsFactory.h"
 
+#include "AudioAnalytics.h"
 #include "ClassViewerModule.h"
 #include "ClassViewerFilter.h"
-#include "EngineAnalytics.h"
 #include "Kismet2/SClassPickerDialog.h"
 #include "Modules/ModuleManager.h"
 #include "AudioSynesthesiaClassFilter.h"
@@ -56,10 +56,7 @@ UObject* UAudioSynesthesiaSettingsFactory::FactoryCreateNew(UClass* InClass, UOb
 	{
 		NewAudioSynesthesiaSettings = NewObject<UAudioSynesthesiaSettings>(InParent, AudioSynesthesiaSettingsClass, InName, Flags);
 
-		if (FEngineAnalytics::IsAvailable())
-		{
-			FEngineAnalytics::GetProvider().RecordEvent(TEXT("Audio.Usage.AudioSynesthesia.SettingsFactoryCreated"));
-		}
+		Audio::Analytics::RecordEvent_Usage(TEXT("AudioSynesthesia.SettingsFactoryCreated"));
 	}
 	return NewAudioSynesthesiaSettings;
 }
diff --git a/Engine/Plugins/Runtime/Metasound/Source/MetasoundEditor/Private/MetasoundEditor.cpp b/Engine/Plugins/Runtime/Metasound/Source/MetasoundEditor/Private/MetasoundEditor.cpp
--- a/Engine/Plugins/Runtime/Metasound/Source/MetasoundEditor/Private/MetasoundEditor.cpp
+++ b/Engine/Plugins/Runtime/Metasound/Source/MetasoundEditor/Private/MetasoundEditor.cpp
@@ -11,7 +11,6 @@
 #include "EdGraph/EdGraphPin.h"
 #include "EdGraphUtilities.h"
 #include "Editor.h"
-#include "EngineAnalytics.h"
 #include "Styling/AppStyle.h"
 #include "Framework/Application/SlateApplication.h"
 #include "Framework/Commands/GenericCommands.h"
@@ -897,11 +896,6 @@ namespace Metasound
 			{
 				Play();
 			}
-
-			if (FEngineAnalytics::IsAvailable())
-			{
-				FEngineAnalytics::GetProvider().RecordEvent(TEXT("Audio.Usage.MetaSounds.EditorOpened"));
-			}
 		}
 
 		UObject* FEditor::GetMetasoundObject() const
diff --git a/Engine/Plugins/Runtime/MicrosoftSpatialAudio/Source/MicrosoftSpatialSound/Private/MicrosoftSpatialSoundPlugin.cpp b/Engine/Plugins/Runtime/MicrosoftSpatialAudio/Source/MicrosoftSpatialSound/Private/MicrosoftSpatialSoundPlugin.cpp
--- a/Engine/Plugins/Runtime/MicrosoftSpatialAudio/Source/MicrosoftSpatialSound/Private/MicrosoftSpatialSoundPlugin.cpp
+++ b/Engine/Plugins/Runtime/MicrosoftSpatialAudio/Source/MicrosoftSpatialSound/Private/MicrosoftSpatialSoundPlugin.cpp
@@ -1,7 +1,7 @@
 // Copyright Epic Games, Inc. All Rights Reserved.
 
 #include "MicrosoftSpatialSoundPlugin.h"
-#include "EngineAnalytics.h"
+#include "AudioAnalytics.h"
 #include "Features/IModularFeatures.h"
 #include "HAL/PlatformMisc.h"
 
@@ -213,10 +213,7 @@ void FMicrosoftSpatialSound::Initialize(const FAudioPluginInitializationParams I
 
 			SpatialAudioRenderThread = FRunnableThread::Create(this, TEXT("MicrosoftSpatialAudioThread"), 0, TPri_TimeCritical, FPlatformAffinity::GetAudioThreadMask());
 
-			if (FEngineAnalytics::IsAvailable())
-			{
-				FEngineAnalytics::GetProvider().RecordEvent(TEXT("Audio.Usage.MicrosoftSpatialSound.Initialized"));
-			}
+			Audio::Analytics::RecordEvent_Usage(TEXT("MicrosoftSpatialSound.Initialized"));
 		}
 		else
 		{
diff --git a/Engine/Plugins/Runtime/OSCModulationMixing/Source/OSCModulationMixing/Private/OSCModulationMixingStatics.cpp b/Engine/Plugins/Runtime/OSCModulationMixing/Source/OSCModulationMixing/Private/OSCModulationMixingStatics.cpp
--- a/Engine/Plugins/Runtime/OSCModulationMixing/Source/OSCModulationMixing/Private/OSCModulationMixingStatics.cpp
+++ b/Engine/Plugins/Runtime/OSCModulationMixing/Source/OSCModulationMixing/Private/OSCModulationMixingStatics.cpp
@@ -7,7 +7,7 @@
 #include "OSCModulationMixing.h"
 #include "OSCServer.h"
 
-#include "EngineAnalytics.h"
+#include "AudioAnalytics.h"
 #include "HAL/IConsoleManager.h"
 #include "SoundControlBus.h"
 #include "SoundControlBusMix.h"
@@ -76,10 +76,7 @@ FOSCAddress UOSCModulationMixingStatics::GetProfileLoadPath()
 
 FOSCAddress UOSCModulationMixingStatics::GetProfileSavePath()
 {
-	if (FEngineAnalytics::IsAvailable())
-	{
-		FEngineAnalytics::GetProvider().RecordEvent(TEXT("Audio.Usage.OSCModulationMixing.GetProfileSavePath"));
-	}
+	Audio::Analytics::RecordEvent_Usage(TEXT("OSCModulationMixing.GetProfileSavePath"));
 	return OSCModulation::Addresses::ProfileSave;
 }
 
diff --git a/Engine/Plugins/Runtime/Oculus/OculusAudio/Source/OculusAudioEditor/Private/OculusAmbisonicSettingsFactory.cpp b/Engine/Plugins/Runtime/Oculus/OculusAudio/Source/OculusAudioEditor/Private/OculusAmbisonicSettingsFactory.cpp
--- a/Engine/Plugins/Runtime/Oculus/OculusAudio/Source/OculusAudioEditor/Private/OculusAmbisonicSettingsFactory.cpp
+++ b/Engine/Plugins/Runtime/Oculus/OculusAudio/Source/OculusAudioEditor/Private/OculusAmbisonicSettingsFactory.cpp
@@ -1,9 +1,9 @@
 // Copyright Epic Games, Inc. All Rights Reserved.
 
 #include "OculusAmbisonicSettingsFactory.h"
-#include "OculusAmbisonicsSettings.h"
 #include "AssetTypeCategories.h"
-#include "EngineAnalytics.h"
+#include "AudioAnalytics.h"
+#include "OculusAmbisonicsSettings.h"
 #include "UObject/ObjectMacros.h"
 #include "UObject/Object.h"
 
@@ -52,10 +52,7 @@ UOculusAmbisonicsSettingsFactory::UOculusAmbisonicsSettingsFactory(const FObject
 UObject* UOculusAmbisonicsSettingsFactory::FactoryCreateNew(UClass* Class, UObject* InParent, FName Name, EObjectFlags Flags, UObject* Context,
 	FFeedbackContext* Warn)
 {
-	if (FEngineAnalytics::IsAvailable())
-	{
-		FEngineAnalytics::GetProvider().RecordEvent(TEXT("Audio.Usage.OculusAudio.AmbisonicsSettingsCreated"));
-	}
+	Audio::Analytics::RecordEvent_Usage(TEXT("OculusAudio.AmbisonicsSettingsCreated"));
 	return NewObject<UOculusAudioSoundfieldSettings>(InParent, Name, Flags);
 }
 
diff --git a/Engine/Plugins/Runtime/Oculus/OculusAudio/Source/OculusAudioEditor/Private/OculusAudioSourceSettingsFactory.cpp b/Engine/Plugins/Runtime/Oculus/OculusAudio/Source/OculusAudioEditor/Private/OculusAudioSourceSettingsFactory.cpp
--- a/Engine/Plugins/Runtime/Oculus/OculusAudio/Source/OculusAudioEditor/Private/OculusAudioSourceSettingsFactory.cpp
+++ b/Engine/Plugins/Runtime/Oculus/OculusAudio/Source/OculusAudioEditor/Private/OculusAudioSourceSettingsFactory.cpp
@@ -1,9 +1,9 @@
 // Copyright Epic Games, Inc. All Rights Reserved.
 
 #include "OculusAudioSourceSettingsFactory.h"
-#include "OculusAudioSourceSettings.h"
 #include "AssetTypeCategories.h"
-#include "EngineAnalytics.h"
+#include "AudioAnalytics.h"
+#include "OculusAudioSourceSettings.h"
 #include "UObject/ObjectMacros.h"
 #include "UObject/Object.h"
 
@@ -52,10 +52,7 @@ UOculusAudioSourceSettingsFactory::UOculusAudioSourceSettingsFactory(const FObje
 UObject* UOculusAudioSourceSettingsFactory::FactoryCreateNew(UClass* Class, UObject* InParent, FName Name, EObjectFlags Flags, UObject* Context,
     FFeedbackContext* Warn)
 {
-	if (FEngineAnalytics::IsAvailable())
-	{
-		FEngineAnalytics::GetProvider().RecordEvent(TEXT("Audio.Usage.OculusAudio.SourceSettingsCreated"));
-	}
+	Audio::Analytics::RecordEvent_Usage(TEXT("OculusAudio.SourceSettingsCreated"));
     return NewObject<UOculusAudioSourceSettings>(InParent, Name, Flags);
 }
 
diff --git a/Engine/Plugins/Runtime/ResonanceAudio/Source/ResonanceAudioEditor/Private/ResonanceAudioReverbPluginPresetFactory.cpp b/Engine/Plugins/Runtime/ResonanceAudio/Source/ResonanceAudioEditor/Private/ResonanceAudioReverbPluginPresetFactory.cpp
--- a/Engine/Plugins/Runtime/ResonanceAudio/Source/ResonanceAudioEditor/Private/ResonanceAudioReverbPluginPresetFactory.cpp
+++ b/Engine/Plugins/Runtime/ResonanceAudio/Source/ResonanceAudioEditor/Private/ResonanceAudioReverbPluginPresetFactory.cpp
@@ -5,7 +5,7 @@
 #include "ResonanceAudioReverbPluginPresetFactory.h"
 #include "ResonanceAudioReverb.h"
 
-#include "EngineAnalytics.h"
+#include "AudioAnalytics.h"
 #include "UObject/ObjectMacros.h"
 #include "UObject/Object.h"
 
@@ -51,10 +51,7 @@ UResonanceAudioReverbPluginPresetFactory::UResonanceAudioReverbPluginPresetFacto
 
 UObject* UResonanceAudioReverbPluginPresetFactory::FactoryCreateNew(UClass* Class, UObject* InParent, FName InName, EObjectFlags Flags, UObject* Context, FFeedbackContext* Warn)
 {
-	if (FEngineAnalytics::IsAvailable())
-	{
-		FEngineAnalytics::GetProvider().RecordEvent(TEXT("Audio.Usage.ResonanceAudio.ReverbPresetCreated"));
-	}
+	Audio::Analytics::RecordEvent_Usage(TEXT("ResonanceAudio.ReverbPresetCreated"));
 	return NewObject<UResonanceAudioReverbPluginPreset>(InParent, InName, Flags);
 }
 
diff --git a/Engine/Plugins/Runtime/ResonanceAudio/Source/ResonanceAudioEditor/Private/ResonanceAudioSpatializationSourceSettingsFactory.cpp b/Engine/Plugins/Runtime/ResonanceAudio/Source/ResonanceAudioEditor/Private/ResonanceAudioSpatializationSourceSettingsFactory.cpp
--- a/Engine/Plugins/Runtime/ResonanceAudio/Source/ResonanceAudioEditor/Private/ResonanceAudioSpatializationSourceSettingsFactory.cpp
+++ b/Engine/Plugins/Runtime/ResonanceAudio/Source/ResonanceAudioEditor/Private/ResonanceAudioSpatializationSourceSettingsFactory.cpp
@@ -5,7 +5,7 @@
 #include "ResonanceAudioSpatializationSourceSettingsFactory.h"
 #include "ResonanceAudioSpatializationSourceSettings.h"
 
-#include "EngineAnalytics.h"
+#include "AudioAnalytics.h"
 #include "UObject/ObjectMacros.h"
 #include "UObject/Object.h"
 
@@ -51,10 +51,7 @@ UResonanceAudioSpatializationSourceSettingsFactory::UResonanceAudioSpatializatio
 
 UObject* UResonanceAudioSpatializationSourceSettingsFactory::FactoryCreateNew(UClass* Class, UObject* InParent, FName InName, EObjectFlags Flags, UObject* Context, FFeedbackContext* Warn)
 {
-	if (FEngineAnalytics::IsAvailable())
-	{
-		FEngineAnalytics::GetProvider().RecordEvent(TEXT("Audio.Usage.ResonanceAudio.SpatializationSourceSettingsCreated"));
-	}
+	Audio::Analytics::RecordEvent_Usage(TEXT("ResonanceAudio.SpatializationSourceSettingsCreated"));
 	return Cast<UObject>(NewObject<UResonanceAudioSpatializationSourceSettings>(InParent, InName, Flags));
 }
 
diff --git a/Engine/Plugins/Runtime/SoundCueTemplates/Source/SoundCueTemplatesEditor/Private/SoundCueTemplateFactory.cpp b/Engine/Plugins/Runtime/SoundCueTemplates/Source/SoundCueTemplatesEditor/Private/SoundCueTemplateFactory.cpp
--- a/Engine/Plugins/Runtime/SoundCueTemplates/Source/SoundCueTemplatesEditor/Private/SoundCueTemplateFactory.cpp
+++ b/Engine/Plugins/Runtime/SoundCueTemplates/Source/SoundCueTemplatesEditor/Private/SoundCueTemplateFactory.cpp
@@ -3,10 +3,10 @@
 #include "SoundCueTemplateFactory.h"
 #include "SoundCueTemplateClassFilter.h"
 
+#include "AudioAnalytics.h"
 #include "ClassViewerModule.h"
 #include "EdGraph/EdGraphPin.h"
 #include "Engine/Engine.h"
-#include "EngineAnalytics.h"
 #include "Kismet2/SClassPickerDialog.h"
 #include "SoundCueGraph/SoundCueGraphNode.h"
 #include "Sound/SoundCue.h"
@@ -91,10 +91,7 @@ UObject* USoundCueTemplateFactory::FactoryCreateNew(UClass* Class, UObject* InPa
 
 			NewSoundCueTemplate->RebuildGraph(*NewSoundCueTemplate);
 
-			if (FEngineAnalytics::IsAvailable())
-			{
-				FEngineAnalytics::GetProvider().RecordEvent(TEXT("Audio.Usage.SoundCueTemplates.TemplateCreated"));
-			}
+			Audio::Analytics::RecordEvent_Usage(TEXT("SoundCueTemplates.TemplateCreated"));
 
 			return NewSoundCueTemplate;
 		}
diff --git a/Engine/Plugins/Runtime/SoundFields/Source/SoundFields/Private/SoundFields.cpp b/Engine/Plugins/Runtime/SoundFields/Source/SoundFields/Private/SoundFields.cpp
--- a/Engine/Plugins/Runtime/SoundFields/Source/SoundFields/Private/SoundFields.cpp
+++ b/Engine/Plugins/Runtime/SoundFields/Source/SoundFields/Private/SoundFields.cpp
@@ -1,11 +1,10 @@
 // Copyright Epic Games, Inc. All Rights Reserved.
 
 #include "SoundFields.h"
-#include "CoreMinimal.h"
+#include "AudioAnalytics.h"
 #include "DSP/AlignedBuffer.h"
 #include "DSP/Dsp.h"
 #include "DSP/FloatArrayMath.h"
-#include "EngineAnalytics.h"
 #include "SoundFieldRendering.h"
 #include "HAL/IConsoleManager.h"
 
@@ -252,10 +251,7 @@ TUniquePtr<ISoundfieldTranscodeStream> FAmbisonicsSoundfieldFormat::CreateTransc
 
 TUniquePtr<ISoundfieldMixerStream> FAmbisonicsSoundfieldFormat::CreateMixerStream(const ISoundfieldEncodingSettingsProxy& InitialSettings)
 {
-	if (FEngineAnalytics::IsAvailable())
-	{
-		FEngineAnalytics::GetProvider().RecordEvent("Audio.Usage.SoundFields.MixerStreamCreated");
-	}
+	Audio::Analytics::RecordEvent_Usage("SoundFields.MixerStreamCreated");
 	return TUniquePtr<ISoundfieldMixerStream>(new FAmbisonicsMixer());
 }
 
diff --git a/Engine/Plugins/Runtime/SoundMod/Source/SoundModImporter/Private/SoundModImporterFactory.cpp b/Engine/Plugins/Runtime/SoundMod/Source/SoundModImporter/Private/SoundModImporterFactory.cpp
--- a/Engine/Plugins/Runtime/SoundMod/Source/SoundModImporter/Private/SoundModImporterFactory.cpp
+++ b/Engine/Plugins/Runtime/SoundMod/Source/SoundModImporter/Private/SoundModImporterFactory.cpp
@@ -1,7 +1,7 @@
 // Copyright Epic Games, Inc. All Rights Reserved.
 
 #include "SoundModImporterFactory.h"
-#include "EngineAnalytics.h"
+#include "AudioAnalytics.h"
 #include "EngineGlobals.h"
 #include "Editor.h"
 #include "SoundMod.h"
@@ -122,10 +122,7 @@ UObject* USoundModImporterFactory::FactoryCreateBinary
 		ComponentsToRestart[ComponentIndex]->Play();
 	}
 
-	if (FEngineAnalytics::IsAvailable())
-	{
-		FEngineAnalytics::GetProvider().RecordEvent("Audio.Usage.SoundModule.BinaryCreated");
-	}
+	Audio::Analytics::RecordEvent_Usage("SoundModule.BinaryCreated");
 
 	return Sound;
 }
diff --git a/Engine/Plugins/Runtime/SoundUtilities/Source/SoundUtilitiesEditor/Private/SoundSimpleFactory.cpp b/Engine/Plugins/Runtime/SoundUtilities/Source/SoundUtilitiesEditor/Private/SoundSimpleFactory.cpp
--- a/Engine/Plugins/Runtime/SoundUtilities/Source/SoundUtilitiesEditor/Private/SoundSimpleFactory.cpp
+++ b/Engine/Plugins/Runtime/SoundUtilities/Source/SoundUtilitiesEditor/Private/SoundSimpleFactory.cpp
@@ -1,7 +1,7 @@
 // Copyright Epic Games, Inc. All Rights Reserved.
 
 #include "SoundSimpleFactory.h"
-#include "EngineAnalytics.h"
+#include "AudioAnalytics.h"
 #include "SoundSimple.h"
 
 USoundSimpleFactory::USoundSimpleFactory(const FObjectInitializer& ObjectInitializer)
@@ -37,10 +37,7 @@ UObject* USoundSimpleFactory::FactoryCreateNew(UClass* Class, UObject* InParent,
 	// Write out the duration to be the longest duration
 	SoundSimple->Duration = Duration;
 
-	if (FEngineAnalytics::IsAvailable())
-	{
-		FEngineAnalytics::GetProvider().RecordEvent("Audio.Usage.SoundUtilities.SimpleSoundCreated");
-	}
+	Audio::Analytics::RecordEvent_Usage("SoundUtilities.SimpleSoundCreated");
 
 	return SoundSimple;
 }
diff --git a/Engine/Plugins/Runtime/Spatialization/Source/SpatializationEditor/Private/ITDSpatializationSourceSettingsFactory.cpp b/Engine/Plugins/Runtime/Spatialization/Source/SpatializationEditor/Private/ITDSpatializationSourceSettingsFactory.cpp
--- a/Engine/Plugins/Runtime/Spatialization/Source/SpatializationEditor/Private/ITDSpatializationSourceSettingsFactory.cpp
+++ b/Engine/Plugins/Runtime/Spatialization/Source/SpatializationEditor/Private/ITDSpatializationSourceSettingsFactory.cpp
@@ -1,7 +1,7 @@
 // Copyright Epic Games, Inc. All Rights Reserved.
 
 #include "ITDSpatializationSourceSettingsFactory.h"
-#include "EngineAnalytics.h"
+#include "AudioAnalytics.h"
 #include "ITDSpatializationSourceSettings.h"
 
 FText FAssetTypeActions_ITDSpatializationSettings::GetName() const
@@ -46,10 +46,7 @@ UITDSpatializationSettingsFactory::UITDSpatializationSettingsFactory(const FObje
 UObject* UITDSpatializationSettingsFactory::FactoryCreateNew(UClass* Class, UObject* InParent, FName InName, EObjectFlags Flags,
 	UObject* Context, FFeedbackContext* Warn)
 {
-	if (FEngineAnalytics::IsAvailable())
-	{
-		FEngineAnalytics::GetProvider().RecordEvent("Audio.Usage.Spatialization.SettingsCreated");
-	}
+	Audio::Analytics::RecordEvent_Usage("Spatialization.SettingsCreated");
 	return NewObject<UITDSpatializationSourceSettings>(InParent, InName, Flags);
 }
 
diff --git a/Engine/Plugins/Runtime/Steam/SteamAudio/Source/SteamAudioEditor/Private/PhononOcclusionSettingsFactory.cpp b/Engine/Plugins/Runtime/Steam/SteamAudio/Source/SteamAudioEditor/Private/PhononOcclusionSettingsFactory.cpp
--- a/Engine/Plugins/Runtime/Steam/SteamAudio/Source/SteamAudioEditor/Private/PhononOcclusionSettingsFactory.cpp
+++ b/Engine/Plugins/Runtime/Steam/SteamAudio/Source/SteamAudioEditor/Private/PhononOcclusionSettingsFactory.cpp
@@ -4,7 +4,7 @@
 
 #include "PhononOcclusionSettingsFactory.h"
 
-#include "EngineAnalytics.h"
+#include "AudioAnalytics.h"
 #include "PhononOcclusionSourceSettings.h"
 #include "UObject/ObjectMacros.h"
 #include "UObject/Object.h"
@@ -54,10 +54,7 @@ UPhononOcclusionSettingsFactory::UPhononOcclusionSettingsFactory(const FObjectIn
 UObject* UPhononOcclusionSettingsFactory::FactoryCreateNew(UClass* Class, UObject* InParent, FName InName, EObjectFlags Flags, UObject* Context,
 	FFeedbackContext* Warn)
 {
-	if (FEngineAnalytics::IsAvailable())
-	{
-		FEngineAnalytics::GetProvider().RecordEvent(TEXT("Audio.Usage.SteamAudio.PhononOcclusionSettingsCreated"));
-	}
+	Audio::Analytics::RecordEvent_Usage(TEXT("SteamAudio.PhononOcclusionSettingsCreated"));
 	return NewObject<UPhononOcclusionSourceSettings>(InParent, InName, Flags);
 }
 
diff --git a/Engine/Plugins/Runtime/Steam/SteamAudio/Source/SteamAudioEditor/Private/PhononReverbSettingsFactory.cpp b/Engine/Plugins/Runtime/Steam/SteamAudio/Source/SteamAudioEditor/Private/PhononReverbSettingsFactory.cpp
--- a/Engine/Plugins/Runtime/Steam/SteamAudio/Source/SteamAudioEditor/Private/PhononReverbSettingsFactory.cpp
+++ b/Engine/Plugins/Runtime/Steam/SteamAudio/Source/SteamAudioEditor/Private/PhononReverbSettingsFactory.cpp
@@ -4,7 +4,7 @@
 
 #include "PhononReverbSettingsFactory.h"
 
-#include "EngineAnalytics.h"
+#include "AudioAnalytics.h"
 #include "PhononReverbSourceSettings.h"
 #include "UObject/ObjectMacros.h"
 #include "UObject/Object.h"
@@ -54,10 +54,7 @@ UPhononReverbSettingsFactory::UPhononReverbSettingsFactory(const FObjectInitiali
 UObject* UPhononReverbSettingsFactory::FactoryCreateNew(UClass* Class, UObject* InParent, FName InName, EObjectFlags Flags, UObject* Context,
 	FFeedbackContext* Warn)
 {
-	if (FEngineAnalytics::IsAvailable())
-	{
-		FEngineAnalytics::GetProvider().RecordEvent(TEXT("Audio.Usage.SteamAudio.PhononReverbSettingsCreated"));
-	}
+	Audio::Analytics::RecordEvent_Usage(TEXT("SteamAudio.PhononReverbSettingsCreated"));
 	return NewObject<UPhononReverbSourceSettings>(InParent, InName, Flags);
 }
 
diff --git a/Engine/Plugins/Runtime/Steam/SteamAudio/Source/SteamAudioEditor/Private/PhononSpatializationSettingsFactory.cpp b/Engine/Plugins/Runtime/Steam/SteamAudio/Source/SteamAudioEditor/Private/PhononSpatializationSettingsFactory.cpp
--- a/Engine/Plugins/Runtime/Steam/SteamAudio/Source/SteamAudioEditor/Private/PhononSpatializationSettingsFactory.cpp
+++ b/Engine/Plugins/Runtime/Steam/SteamAudio/Source/SteamAudioEditor/Private/PhononSpatializationSettingsFactory.cpp
@@ -4,7 +4,7 @@
 
 #include "PhononSpatializationSettingsFactory.h"
 
-#include "EngineAnalytics.h"
+#include "AudioAnalytics.h"
 #include "PhononSpatializationSourceSettings.h"
 #include "UObject/ObjectMacros.h"
 #include "UObject/Object.h"
@@ -54,10 +54,7 @@ UPhononSpatializationSettingsFactory::UPhononSpatializationSettingsFactory(const
 UObject* UPhononSpatializationSettingsFactory::FactoryCreateNew(UClass* Class, UObject* InParent, FName InName, EObjectFlags Flags,
 	UObject* Context, FFeedbackContext* Warn)
 {
-	if (FEngineAnalytics::IsAvailable())
-	{
-		FEngineAnalytics::GetProvider().RecordEvent(TEXT("Audio.Usage.SteamAudio.PhononSpatializationSettingsCreated"));
-	}
+	Audio::Analytics::RecordEvent_Usage(TEXT("SteamAudio.PhononSpatializationSettingsCreated"));
 	return NewObject<UPhononSpatializationSourceSettings>(InParent, InName, Flags);
 }
 
diff --git a/Engine/Plugins/Runtime/Synthesis/Source/SynthesisEditor/Private/EpicSynth1PresetBank.cpp b/Engine/Plugins/Runtime/Synthesis/Source/SynthesisEditor/Private/EpicSynth1PresetBank.cpp
--- a/Engine/Plugins/Runtime/Synthesis/Source/SynthesisEditor/Private/EpicSynth1PresetBank.cpp
+++ b/Engine/Plugins/Runtime/Synthesis/Source/SynthesisEditor/Private/EpicSynth1PresetBank.cpp
@@ -1,7 +1,7 @@
 // Copyright Epic Games, Inc. All Rights Reserved.
 #include "EpicSynth1PresetBank.h"
 
-#include "EngineAnalytics.h"
+#include "AudioAnalytics.h"
 #include "SynthComponents/EpicSynth1Component.h"
 #include "UObject/ObjectMacros.h"
 #include "UObject/Object.h"
@@ -37,10 +37,7 @@ UModularSynthPresetBankFactory::UModularSynthPresetBankFactory(const FObjectInit
 UObject* UModularSynthPresetBankFactory::FactoryCreateNew(UClass* Class, UObject* InParent, FName InName, EObjectFlags Flags, UObject* Context, FFeedbackContext* Warn)
 {
 	UModularSynthPresetBank* NewPresetBank = NewObject<UModularSynthPresetBank>(InParent, InName, Flags);
-	if (FEngineAnalytics::IsAvailable())
-	{
-		FEngineAnalytics::GetProvider().RecordEvent(TEXT("Audio.Usage.SynthesisAndDSPEffects.ModularSynthPresetBankCreated"));
-	}
+	Audio::Analytics::RecordEvent_Usage(TEXT("SynthesisAndDSPEffects.ModularSynthPresetBankCreated"));
 	return NewPresetBank;
 }
 
diff --git a/Engine/Plugins/Runtime/Synthesis/Source/SynthesisEditor/Private/MonoWaveTablePresetBank.cpp b/Engine/Plugins/Runtime/Synthesis/Source/SynthesisEditor/Private/MonoWaveTablePresetBank.cpp
--- a/Engine/Plugins/Runtime/Synthesis/Source/SynthesisEditor/Private/MonoWaveTablePresetBank.cpp
+++ b/Engine/Plugins/Runtime/Synthesis/Source/SynthesisEditor/Private/MonoWaveTablePresetBank.cpp
@@ -1,7 +1,7 @@
 // Copyright Epic Games, Inc. All Rights Reserved.
 #include "MonoWaveTablePresetBank.h"
 
-#include "EngineAnalytics.h"
+#include "AudioAnalytics.h"
 #include "SynthComponents/SynthComponentMonoWaveTable.h"
 #include "UObject/ObjectMacros.h"
 #include "UObject/Object.h"
@@ -37,10 +37,7 @@ UMonoWaveTableSynthPresetFactory::UMonoWaveTableSynthPresetFactory(const FObject
 UObject* UMonoWaveTableSynthPresetFactory::FactoryCreateNew(UClass* Class, UObject* InParent, FName InName, EObjectFlags Flags, UObject* Context, FFeedbackContext* Warn)
 {
 	UMonoWaveTableSynthPreset* NewPresetBank = NewObject<UMonoWaveTableSynthPreset>(InParent, InName, Flags);
-	if (FEngineAnalytics::IsAvailable())
-	{
-		FEngineAnalytics::GetProvider().RecordEvent(TEXT("Audio.Usage.SynthesisAndDSPEffects.MonoWaveTableSynthPresetCreated"));
-	}
+	Audio::Analytics::RecordEvent_Usage(TEXT("SynthesisAndDSPEffects.MonoWaveTableSynthPresetCreated"));
 	return NewPresetBank;
 }
 #undef LOCTEXT_NAMESPACE
\ No newline at end of file
diff --git a/Engine/Source/Editor/AudioEditor/Private/Factories/SoundFactory.cpp b/Engine/Source/Editor/AudioEditor/Private/Factories/SoundFactory.cpp
--- a/Engine/Source/Editor/AudioEditor/Private/Factories/SoundFactory.cpp
+++ b/Engine/Source/Editor/AudioEditor/Private/Factories/SoundFactory.cpp
@@ -3,9 +3,9 @@
 #include "Factories/SoundFactory.h"
 #include "AssetRegistry/AssetRegistryModule.h"
 #include "Audio.h"
+#include "AudioAnalytics.h"
 #include "Components/AudioComponent.h"
 #include "ContentBrowserModule.h"
-#include "EngineAnalytics.h"
 #include "IContentBrowserSingleton.h"
 #include "Modules/ModuleManager.h"
 #include "Sound/SoundCue.h"
@@ -601,10 +601,7 @@ UObject* USoundFactory::CreateObject
 
 		Sound->SetRedrawThumbnail(true);
 
-		if (FEngineAnalytics::IsAvailable())
-		{
-			FEngineAnalytics::GetProvider().RecordEvent(TEXT("Audio.Usage.SoundFactory.SoundWaveImported"));
-		}
+		Audio::Analytics::RecordEvent_Usage(TEXT("SoundFactory.SoundWaveImported"));
 
 		return Sound;
 	}
diff --git a/Engine/Source/Editor/AudioEditor/Private/SoundClassEditor.cpp b/Engine/Source/Editor/AudioEditor/Private/SoundClassEditor.cpp
--- a/Engine/Source/Editor/AudioEditor/Private/SoundClassEditor.cpp
+++ b/Engine/Source/Editor/AudioEditor/Private/SoundClassEditor.cpp
@@ -175,12 +175,6 @@ void FSoundClassEditor::InitSoundClassEditor( const EToolkitMode::Type Mode, con
 			break;
 		}
 	}
-
-	if (FEngineAnalytics::IsAvailable())
-	{
-		FEngineAnalytics::GetProvider().RecordEvent(TEXT("Audio.Usage.SoundClass.EditorOpened"));
-	}
-
 }
 
 void FSoundClassEditor::BindCommands()
diff --git a/Engine/Source/Editor/AudioEditor/Private/SoundCueEditor.cpp b/Engine/Source/Editor/AudioEditor/Private/SoundCueEditor.cpp
--- a/Engine/Source/Editor/AudioEditor/Private/SoundCueEditor.cpp
+++ b/Engine/Source/Editor/AudioEditor/Private/SoundCueEditor.cpp
@@ -164,11 +164,6 @@ void FSoundCueEditor::InitSoundCueEditor(const EToolkitMode::Type Mode, const TS
 		SpawnToolkitTab(GraphCanvasTabId, FString(), EToolkitTabSpot::Viewport);
 		SpawnToolkitTab(PropertiesTabId, FString(), EToolkitTabSpot::Details);
 	}*/
-
-	if (FEngineAnalytics::IsAvailable())
-	{
-		FEngineAnalytics::GetProvider().RecordEvent(TEXT("Audio.Usage.SoundCue.EditorOpened"));
-	}
 }
 
 USoundCue* FSoundCueEditor::GetSoundCue() const
diff --git a/Engine/Source/Editor/AudioEditor/Private/SoundSubmixEditor.cpp b/Engine/Source/Editor/AudioEditor/Private/SoundSubmixEditor.cpp
--- a/Engine/Source/Editor/AudioEditor/Private/SoundSubmixEditor.cpp
+++ b/Engine/Source/Editor/AudioEditor/Private/SoundSubmixEditor.cpp
@@ -193,12 +193,6 @@ void FSoundSubmixEditor::Init(const EToolkitMode::Type Mode, const TSharedPtr<IT
 			break;
 		}
 	}
-
-	if (FEngineAnalytics::IsAvailable())
-	{
-		FEngineAnalytics::GetProvider().RecordEvent(TEXT("Audio.Usage.Submix.EditorOpened"));
-	}
-
 }
 
 FSoundSubmixEditor::~FSoundSubmixEditor()
diff --git a/Engine/Source/Runtime/AudioLink/AudioLinkEngine/Private/IAudioLink.cpp b/Engine/Source/Runtime/AudioLink/AudioLinkEngine/Private/IAudioLink.cpp
--- a/Engine/Source/Runtime/AudioLink/AudioLinkEngine/Private/IAudioLink.cpp
+++ b/Engine/Source/Runtime/AudioLink/AudioLinkEngine/Private/IAudioLink.cpp
@@ -2,13 +2,9 @@
 
 #include "IAudioLink.h"
 
-#include "EngineAnalytics.h"
+#include "AudioAnalytics.h"
 
 IAudioLink::IAudioLink()
 {
-	if (FEngineAnalytics::IsAvailable())
-	{
-		FEngineAnalytics::GetProvider().RecordEvent(TEXT("Audio.Usage.AudioLink.InstanceCreated"));
-	}
+	Audio::Analytics::RecordEvent_Usage(TEXT("AudioLink.InstanceCreated"));
 }
-
diff --git a/Engine/Source/Runtime/AudioMixer/Private/AudioMixerDevice.cpp b/Engine/Source/Runtime/AudioMixer/Private/AudioMixerDevice.cpp
--- a/Engine/Source/Runtime/AudioMixer/Private/AudioMixerDevice.cpp
+++ b/Engine/Source/Runtime/AudioMixer/Private/AudioMixerDevice.cpp
@@ -2,6 +2,7 @@
 
 #include "AudioMixerDevice.h"
 
+#include "AudioAnalytics.h"
 #include "AudioMixerSource.h"
 #include "AudioMixerSourceManager.h"
 #include "AudioMixerSourceDecode.h"
@@ -9,7 +10,6 @@
 #include "AudioMixerSourceVoice.h"
 #include "AudioPluginUtilities.h"
 #include "AudioMixerEffectsManager.h"
-#include "EngineAnalytics.h"
 #include "DSP/Noise.h"
 #include "DSP/SinOsc.h"
 #include "Sound/AudioSettings.h"
@@ -473,15 +473,12 @@ namespace Audio
 				// Create synchronized Audio Task Queue for this device...
 				CreateSynchronizedAudioTaskQueue((Audio::AudioTaskQueueId)DeviceID);
 
-				if (FEngineAnalytics::IsAvailable())
-				{
-					FEngineAnalytics::GetProvider().RecordEvent(TEXT("Audio.Usage.ProjectSettings"), MakeAnalyticsEventAttributeArray(
-						TEXT("SampleRate"), PlatformSettings.SampleRate,
-						TEXT("BufferSize"), PlatformSettings.CallbackBufferFrameSize,
-						TEXT("NumBuffers"), PlatformSettings.NumBuffers,
-						TEXT("NumSources"), PlatformSettings.MaxChannels,
-						TEXT("NumOutputChannels"), PlatformInfo.NumChannels));
-				}
+				Audio::Analytics::RecordEvent_Usage(TEXT("ProjectSettings"), MakeAnalyticsEventAttributeArray(
+					TEXT("SampleRate"), PlatformSettings.SampleRate,
+					TEXT("BufferSize"), PlatformSettings.CallbackBufferFrameSize,
+					TEXT("NumBuffers"), PlatformSettings.NumBuffers,
+					TEXT("NumSources"), PlatformSettings.MaxChannels,
+					TEXT("NumOutputChannels"), PlatformInfo.NumChannels));
 
 				// Start streaming audio
 				return AudioMixerPlatform->StartAudioStream();
diff --git a/Engine/Source/Runtime/Engine/Private/AudioDeviceManager.cpp b/Engine/Source/Runtime/Engine/Private/AudioDeviceManager.cpp
--- a/Engine/Source/Runtime/Engine/Private/AudioDeviceManager.cpp
+++ b/Engine/Source/Runtime/Engine/Private/AudioDeviceManager.cpp
@@ -3,11 +3,11 @@
 
 #include "Audio.h"
 #include "Audio/AudioDebug.h"
+#include "AudioAnalytics.h"
 #include "AudioDefines.h"
 #include "AudioDevice.h"
 #include "AudioMixerDevice.h"
 #include "Engine/World.h"
-#include "EngineAnalytics.h"
 #include "Sound/AudioSettings.h"
 #include "Sound/SoundWave.h"
 #include "GameFramework/GameUserSettings.h"
@@ -711,13 +711,7 @@ bool FAudioDeviceManager::LoadDefaultAudioDeviceModule()
 		IsUsingAudioMixerCvar->Set(0, ECVF_SetByConstructor);
 	}
 
-	if (FEngineAnalytics::IsAvailable())
-	{ 		
-		if (bForceNoAudioMixer)
-		{
-			FEngineAnalytics::GetProvider().RecordEvent(TEXT("Audio.Usage.ForceNoAudioMixer"));
-		}
-	}
+	Audio::Analytics::RecordEvent_Usage(TEXT("ForceNoAudioMixer"));
 
 	return AudioDeviceModule != nullptr;
 }
@@ -914,10 +908,7 @@ bool FAudioDeviceManager::Initialize()
 		}
 		else
 		{
-			if (FEngineAnalytics::IsAvailable())
-			{
-				FEngineAnalytics::GetProvider().RecordEvent(TEXT("Audio.Usage.AllAudioDisabled"));
-			}
+			Audio::Analytics::RecordEvent_Usage(TEXT("AllAudioDisabled"));
 		}
 	}
 
diff --git a/Engine/Source/Runtime/Engine/Public/AudioAnalytics.h b/Engine/Source/Runtime/Engine/Public/AudioAnalytics.h
new file mode 100644
--- /dev/null
+++ b/Engine/Source/Runtime/Engine/Public/AudioAnalytics.h
@@ -0,0 +1,28 @@
+// Copyright Epic Games, Inc. All Rights Reserved.
+
+#pragma once
+
+#include "Containers/Array.h"
+#include "Containers/UnrealString.h"
+#include "EngineAnalytics.h"
+#include "HAL/Platform.h"
+#include "UObject/TopLevelAssetPath.h"
+
+#define UE_AUDIO_ENABLE_ANALYTICS 1
+
+
+namespace Audio
+{
+	namespace Analytics
+	{
+		FORCEINLINE void RecordEvent_Usage(FString&& InEventName, const TArray<FAnalyticsEventAttribute>& InEventAttributes = TArray<FAnalyticsEventAttribute>())
+		{
+#ifdef UE_AUDIO_ENABLE_ANALYTICS
+			if (FEngineAnalytics::IsAvailable())
+			{
+				FEngineAnalytics::GetProvider().RecordEvent(TEXT("Audio.Usage.") + MoveTemp(InEventName), InEventAttributes);
+			}
+#endif // UE_AUDIO_ENABLE_ANALYTICS
+		}
+	}
+}
