[CrashFix] - Revert previous attempt at KillPendingKill fix for SynthComponent

diff --git a/Engine/Source/Runtime/AudioMixer/Private/Components/SynthComponent.cpp b/Engine/Source/Runtime/AudioMixer/Private/Components/SynthComponent.cpp
--- a/Engine/Source/Runtime/AudioMixer/Private/Components/SynthComponent.cpp
+++ b/Engine/Source/Runtime/AudioMixer/Private/Components/SynthComponent.cpp
@@ -9,6 +9,7 @@
 
 USynthSound::USynthSound(const FObjectInitializer& ObjectInitializer)
 	: Super(ObjectInitializer)
+	, OwningSynthComponent(nullptr)
 {
 }
 
@@ -46,7 +47,7 @@ void USynthSound::StartOnAudioDevice(FAudioDevice* InAudioDevice)
 
 void USynthSound::OnBeginGenerate()
 {
-	if (ensure(OwningSynthComponent.IsValid()))
+	if (ensure(OwningSynthComponent))
 	{
 		OwningSynthComponent->OnBeginGenerate();
 	}
@@ -64,7 +65,7 @@ int32 USynthSound::OnGeneratePCMAudio(TArray<uint8>& OutAudio, int32 NumSamples)
 		OutAudio.AddZeroed(NumSamples * sizeof(float));
 
 		// Mark pending kill can null this out on the game thread in rare cases.
-		if (!OwningSynthComponent.IsValid())
+		if (!OwningSynthComponent)
 		{
 			return 0;
 		}
@@ -78,7 +79,7 @@ int32 USynthSound::OnGeneratePCMAudio(TArray<uint8>& OutAudio, int32 NumSamples)
 		FloatBuffer.AddZeroed(NumSamples * sizeof(float));
 
 		// Mark pending kill can null this out on the game thread in rare cases.
-		if (!OwningSynthComponent.IsValid())
+		if (!OwningSynthComponent)
 		{
 			return 0;
 		}
@@ -102,7 +103,7 @@ int32 USynthSound::OnGeneratePCMAudio(TArray<uint8>& OutAudio, int32 NumSamples)
 void USynthSound::OnEndGenerate()
 {
 	// Mark pending kill can null this out on the game thread in rare cases.
-	if (OwningSynthComponent.IsValid())
+	if (OwningSynthComponent)
 	{
 		OwningSynthComponent->OnEndGenerate();
 	}
@@ -110,7 +111,7 @@ void USynthSound::OnEndGenerate()
 
 ISoundGeneratorPtr USynthSound::CreateSoundGenerator(const FSoundGeneratorInitParams& InParams)
 {
-	if (OwningSynthComponent.IsValid())
+	if (OwningSynthComponent)
 	{
 		return OwningSynthComponent->CreateSoundGeneratorInternal(InParams);
 	}
diff --git a/Engine/Source/Runtime/AudioMixer/Public/Components/SynthComponent.h b/Engine/Source/Runtime/AudioMixer/Public/Components/SynthComponent.h
--- a/Engine/Source/Runtime/AudioMixer/Public/Components/SynthComponent.h
+++ b/Engine/Source/Runtime/AudioMixer/Public/Components/SynthComponent.h
@@ -56,7 +56,7 @@ class AUDIOMIXER_API USynthSound : public USoundWaveProcedural
 
 protected:
 	UPROPERTY()
-	TWeakObjectPtr<USynthComponent> OwningSynthComponent = nullptr;
+	TObjectPtr<USynthComponent> OwningSynthComponent;
 
 	TArray<float> FloatBuffer;
 	bool bAudioMixer;
