Add OnAllSourcesProcessed to source data override plugin

diff --git a/Engine/Source/Runtime/AudioExtensions/Public/IAudioExtensionPlugin.h b/Engine/Source/Runtime/AudioExtensions/Public/IAudioExtensionPlugin.h
--- a/Engine/Source/Runtime/AudioExtensions/Public/IAudioExtensionPlugin.h
+++ b/Engine/Source/Runtime/AudioExtensions/Public/IAudioExtensionPlugin.h
@@ -451,6 +451,10 @@ public:
 	{
 	}
 
+	/** Called when all sources have finished processing. */
+	virtual void OnAllSourcesProcessed()
+	{
+	}
 };
 
 /** This is a class which should be overridden to provide users with settings to use for individual sounds */
diff --git a/Engine/Source/Runtime/AudioMixer/Private/AudioMixerSourceManager.cpp b/Engine/Source/Runtime/AudioMixer/Private/AudioMixerSourceManager.cpp
--- a/Engine/Source/Runtime/AudioMixer/Private/AudioMixerSourceManager.cpp
+++ b/Engine/Source/Runtime/AudioMixer/Private/AudioMixerSourceManager.cpp
@@ -165,6 +165,7 @@ namespace Audio
 		, NumSourceWorkers(4)
 		, bInitialized(false)
 		, bUsingSpatializationPlugin(false)
+		, bUsingSourceDataOverridePlugin(false)
 	{
 		// Get a manual resetable event
 		const bool bIsManualReset = true;
@@ -343,6 +344,12 @@ namespace Audio
 		{
 			bUsingSpatializationPlugin = true;
 		}
+		// Cache the source data override plugin
+		SourceDataOverridePlugin = MixerDevice->SourceDataOverridePluginInterface;
+		if (SourceDataOverridePlugin.IsValid())
+		{
+			bUsingSourceDataOverridePlugin = true;
+		}
 
 		// Spam command queue with nops.
 		static FAutoConsoleCommand SpamNopsCmd(
@@ -2911,6 +2918,14 @@ namespace Audio
 			SpatialInterfaceInfo.SpatializationPlugin->OnAllSourcesProcessed();
 		}
 
+		// Let the plugin know we finished processing all sources
+		if (bUsingSourceDataOverridePlugin)
+		{
+			AUDIO_MIXER_CHECK(SourceDataOverridePlugin.IsValid());
+			LLM_SCOPE(ELLMTag::AudioMixerPlugins);
+			SourceDataOverridePlugin->OnAllSourcesProcessed();
+		}
+
 		// Update the game thread copy of source doneness
 		for (int32 SourceId = 0; SourceId < NumTotalSources; ++SourceId)
 		{		
diff --git a/Engine/Source/Runtime/AudioMixer/Private/AudioMixerSourceManager.h b/Engine/Source/Runtime/AudioMixer/Private/AudioMixerSourceManager.h
--- a/Engine/Source/Runtime/AudioMixer/Private/AudioMixerSourceManager.h
+++ b/Engine/Source/Runtime/AudioMixer/Private/AudioMixerSourceManager.h
@@ -327,8 +327,11 @@ namespace Audio
 
 		FMixerDevice* MixerDevice;
 
-		// Cached ptr to an optional spatialization plugin
+		// Info about spatialization plugin
 		FAudioDevice::FAudioSpatializationInterfaceInfo SpatialInterfaceInfo;
+		
+		// Cached ptr to an optional source data override plugin
+		TAudioSourceDataOverridePtr SourceDataOverridePlugin;
 
 		// Array of pointers to game thread audio source objects
 		TArray<FMixerSourceVoice*> MixerSources;
@@ -548,6 +551,7 @@ namespace Audio
 
 		uint8 bInitialized : 1;
 		uint8 bUsingSpatializationPlugin : 1;
+		uint8 bUsingSourceDataOverridePlugin : 1;
 
 		// Set to true when the audio source manager should pump the command queue
 		FThreadSafeBool bPumpQueue;
