ModulationDestinations Upgrade
- Support sets of input modulators
- Support union of modulation heirarchy

diff --git a/Engine/Plugins/Runtime/Synthesis/Source/Synthesis/Classes/SourceEffects/SourceEffectBitCrusher.h b/Engine/Plugins/Runtime/Synthesis/Source/Synthesis/Classes/SourceEffects/SourceEffectBitCrusher.h
--- a/Engine/Plugins/Runtime/Synthesis/Source/Synthesis/Classes/SourceEffects/SourceEffectBitCrusher.h
+++ b/Engine/Plugins/Runtime/Synthesis/Source/Synthesis/Classes/SourceEffects/SourceEffectBitCrusher.h
@@ -1,6 +1,7 @@
 // Copyright Epic Games, Inc. All Rights Reserved.
 #pragma once
 
+#include "Containers/Set.h"
 #include "DSP/BitCrusher.h"
 #include "Sound/SoundEffectSource.h"
 #include "Sound/SoundModulationDestination.h"
@@ -63,9 +64,11 @@ public:
 	// Process the input block of audio. Called on audio thread.
 	virtual void ProcessAudio(const FSoundEffectSourceInputData& InData, float* OutAudioBufferData) override;
 
-	void SetSampleRateModulator(const USoundModulatorBase* Modulator);
-
 	void SetBitModulator(const USoundModulatorBase* Modulator);
+	void SetBitModulators(const TSet<USoundModulatorBase*>& InModulators);
+
+	void SetSampleRateModulator(const USoundModulatorBase* Modulator);
+	void SetSampleRateModulators(const TSet<USoundModulatorBase*>& InModulators);
 
 protected:
 	Audio::FBitCrusher BitCrusher;
@@ -96,12 +99,18 @@ public:
 	UFUNCTION(BlueprintCallable, Category = "Audio|Effects|BitCrusher")
 	void SetBitModulator(const USoundModulatorBase* Modulator);
 
+	UFUNCTION(BlueprintCallable, Category = "Audio|Effects|BitCrusher")
+	void SetBitModulators(const TSet<USoundModulatorBase*>& InModulators);
+
 	UFUNCTION(BlueprintCallable, Category = "Audio|Effects|BitCrusher")
 	void SetSampleRate(float SampleRate);
 
 	UFUNCTION(BlueprintCallable, Category = "Audio|Effects|BitCrusher")
 	void SetSampleRateModulator(const USoundModulatorBase* Modulator);
 
+	UFUNCTION(BlueprintCallable, Category = "Audio|Effects|BitCrusher")
+	void SetSampleRateModulators(const TSet<USoundModulatorBase*>& InModulators);
+
 	// Sets just base (i.e. carrier) setting values without modifying modulation source references
 	UFUNCTION(BlueprintCallable, Category = "Audio|Effects|BitCrusher")
 	void SetSettings(const FSourceEffectBitCrusherBaseSettings& Settings);
diff --git a/Engine/Plugins/Runtime/Synthesis/Source/Synthesis/Classes/SourceEffects/SourceEffectChorus.h b/Engine/Plugins/Runtime/Synthesis/Source/Synthesis/Classes/SourceEffects/SourceEffectChorus.h
--- a/Engine/Plugins/Runtime/Synthesis/Source/Synthesis/Classes/SourceEffects/SourceEffectChorus.h
+++ b/Engine/Plugins/Runtime/Synthesis/Source/Synthesis/Classes/SourceEffects/SourceEffectChorus.h
@@ -2,8 +2,9 @@
 
 #pragma once
 
-#include "Sound/SoundEffectSource.h"
+#include "Containers/Set.h"
 #include "DSP/Chorus.h"
+#include "Sound/SoundEffectSource.h"
 #include "Sound/SoundModulationDestination.h"
 
 #include "SourceEffectChorus.generated.h"
@@ -115,6 +116,13 @@ public:
 	void SetDryModulator(const USoundModulatorBase* InModulator);
 	void SetSpreadModulator(const USoundModulatorBase* InModulator);
 
+	void SetDepthModulators(const TSet<USoundModulatorBase*>& InModulators);
+	void SetFeedbackModulators(const TSet<USoundModulatorBase*>& InModulators);
+	void SetSpreadModulators(const TSet<USoundModulatorBase*>& InModulators);
+	void SetDryModulators(const TSet<USoundModulatorBase*>& InModulators);
+	void SetWetModulators(const TSet<USoundModulatorBase*>& InModulators);
+	void SetFrequencyModulators(const TSet<USoundModulatorBase*>& InModulators);
+
 protected:
 	Audio::FChorus Chorus;
 
@@ -146,36 +154,54 @@ public:
 	UFUNCTION(BlueprintCallable, Category = "Audio|Effects|Chorus")
 	void SetDepthModulator(const USoundModulatorBase* Modulator);
 
+	UFUNCTION(BlueprintCallable, Category = "Audio|Effects|Chorus")
+	void SetDepthModulators(const TSet<USoundModulatorBase*>& Modulators);
+
 	UFUNCTION(BlueprintCallable, Category = "Audio|Effects|Chorus")
 	void SetFeedback(float Feedback);
 
 	UFUNCTION(BlueprintCallable, Category = "Audio|Effects|Chorus")
 	void SetFeedbackModulator(const USoundModulatorBase* Modulator);
 
+	UFUNCTION(BlueprintCallable, Category = "Audio|Effects|Chorus")
+	void SetFeedbackModulators(const TSet<USoundModulatorBase*>& Modulators);
+
 	UFUNCTION(BlueprintCallable, Category = "Audio|Effects|Chorus")
 	void SetFrequency(float Frequency);
 
 	UFUNCTION(BlueprintCallable, Category = "Audio|Effects|Chorus")
 	void SetFrequencyModulator(const USoundModulatorBase* Modulator);
 
+	UFUNCTION(BlueprintCallable, Category = "Audio|Effects|Chorus")
+	void SetFrequencyModulators(const TSet<USoundModulatorBase*>& Modulators);
+
 	UFUNCTION(BlueprintCallable, Category = "Audio|Effects|Chorus")
 	void SetWet(float WetAmount);
 
 	UFUNCTION(BlueprintCallable, Category = "Audio|Effects|Chorus")
 	void SetWetModulator(const USoundModulatorBase* Modulator);
 
+	UFUNCTION(BlueprintCallable, Category = "Audio|Effects|Chorus")
+	void SetWetModulators(const TSet<USoundModulatorBase*>& Modulators);
+
 	UFUNCTION(BlueprintCallable, Category = "Audio|Effects|Chorus")
 	void SetDry(float DryAmount);
 
 	UFUNCTION(BlueprintCallable, Category = "Audio|Effects|Chorus")
 	void SetDryModulator(const USoundModulatorBase* Modulator);
 
+	UFUNCTION(BlueprintCallable, Category = "Audio|Effects|Chorus")
+	void SetDryModulators(const TSet<USoundModulatorBase*>& Modulators);
+
 	UFUNCTION(BlueprintCallable, Category = "Audio|Effects|Chorus")
 	void SetSpread(float Spread);
 
 	UFUNCTION(BlueprintCallable, Category = "Audio|Effects|Chorus")
 	void SetSpreadModulator(const USoundModulatorBase* Modulator);
 
+	UFUNCTION(BlueprintCallable, Category = "Audio|Effects|Chorus")
+	void SetSpreadModulators(const TSet<USoundModulatorBase*>& Modulators);
+
 	// Sets just base (i.e. carrier) setting values without modifying modulation source references
 	UFUNCTION(BlueprintCallable, Category = "Audio|Effects|Chorus")
 	void SetSettings(const FSourceEffectChorusBaseSettings& Settings);
diff --git a/Engine/Plugins/Runtime/Synthesis/Source/Synthesis/Private/SourceEffectBitCrusher.cpp b/Engine/Plugins/Runtime/Synthesis/Source/Synthesis/Private/SourceEffectBitCrusher.cpp
--- a/Engine/Plugins/Runtime/Synthesis/Source/Synthesis/Private/SourceEffectBitCrusher.cpp
+++ b/Engine/Plugins/Runtime/Synthesis/Source/Synthesis/Private/SourceEffectBitCrusher.cpp
@@ -2,6 +2,7 @@
 
 #include "SourceEffects/SourceEffectBitCrusher.h"
 
+#include "Algo/Transform.h"
 #include "IAudioModulation.h"
 
 
@@ -15,8 +16,15 @@ void FSourceEffectBitCrusher::Init(const FSoundEffectSourceInitData& InitData)
 		BitMod.Init(InitData.AudioDeviceId, FName("BitDepth"), false /* bInIsBuffered */);
 		SampleRateMod.Init(InitData.AudioDeviceId, FName("SampleRate"), false /* bInIsBuffered */);
 
-		SetBitModulator(ProcPreset->Settings.BitModulation.Modulator);
-		SetSampleRateModulator(ProcPreset->Settings.SampleRateModulation.Modulator);
+		auto TransformObjectPtrSet = [](const TSet<TObjectPtr<USoundModulatorBase>>& InPtrs)
+		{
+			TSet<USoundModulatorBase*> Objects;
+			Algo::Transform(InPtrs, Objects, [](const TObjectPtr<USoundModulatorBase>& Obj) { return Obj; });
+			return Objects;
+		};
+
+		SetBitModulators(TransformObjectPtrSet(ProcPreset->Settings.BitModulation.Modulators));
+		SetSampleRateModulators(TransformObjectPtrSet(ProcPreset->Settings.SampleRateModulation.Modulators));
 	}
 }
 
@@ -44,18 +52,35 @@ void FSourceEffectBitCrusher::ProcessAudio(const FSoundEffectSourceInputData& In
 
 void FSourceEffectBitCrusher::SetBitModulator(const USoundModulatorBase* InModulator)
 {
-	BitMod.UpdateModulator(InModulator);
+	BitMod.UpdateModulators({ InModulator });
+}
+
+void FSourceEffectBitCrusher::SetBitModulators(const TSet<USoundModulatorBase*>& InModulators)
+{
+	BitMod.UpdateModulators(InModulators);
 }
 
 void FSourceEffectBitCrusher::SetSampleRateModulator(const USoundModulatorBase* InModulator)
 {
-	SampleRateMod.UpdateModulator(InModulator);
+	SampleRateMod.UpdateModulators({ InModulator });
+}
+
+void FSourceEffectBitCrusher::SetSampleRateModulators(const TSet<USoundModulatorBase*>& InModulators)
+{
+	SampleRateMod.UpdateModulators(InModulators);
 }
 
 void USourceEffectBitCrusherPreset::OnInit()
 {
-	SetBitModulator(Settings.BitModulation.Modulator);
-	SetSampleRateModulator(Settings.SampleRateModulation.Modulator);
+	auto TransformObjectPtrSet = [](const TSet<TObjectPtr<USoundModulatorBase>>& InPtrs)
+	{
+		TSet<USoundModulatorBase*> Objects;
+		Algo::Transform(InPtrs, Objects, [](const TObjectPtr<USoundModulatorBase>& Obj) { return Obj; });
+		return Objects;
+	};
+
+	SetBitModulators(TransformObjectPtrSet(Settings.BitModulation.Modulators));
+	SetSampleRateModulators(TransformObjectPtrSet(Settings.SampleRateModulation.Modulators));
 }
 
 void USourceEffectBitCrusherPreset::Serialize(FArchive& Ar)
@@ -96,6 +121,14 @@ void USourceEffectBitCrusherPreset::SetBitModulator(const USoundModulatorBase* I
 	});
 }
 
+void USourceEffectBitCrusherPreset::SetBitModulators(const TSet<USoundModulatorBase*>& InModulators)
+{
+	IterateEffects<FSourceEffectBitCrusher>([&InModulators](FSourceEffectBitCrusher& InCrusher)
+	{
+		InCrusher.SetBitModulators(InModulators);
+	});
+}
+
 void USourceEffectBitCrusherPreset::SetSampleRate(float InSampleRate)
 {
 	UpdateSettings([NewSampleRate = InSampleRate](FSourceEffectBitCrusherSettings& OutSettings)
@@ -112,6 +145,14 @@ void USourceEffectBitCrusherPreset::SetSampleRateModulator(const USoundModulator
 	});
 }
 
+void USourceEffectBitCrusherPreset::SetSampleRateModulators(const TSet<USoundModulatorBase*>& InModulators)
+{
+	IterateEffects<FSourceEffectBitCrusher>([&InModulators](FSourceEffectBitCrusher& InCrusher)
+	{
+		InCrusher.SetSampleRateModulators(InModulators);
+	});
+}
+
 void USourceEffectBitCrusherPreset::SetSettings(const FSourceEffectBitCrusherBaseSettings& InSettings)
 {
 	UpdateSettings([NewBaseSettings = InSettings](FSourceEffectBitCrusherSettings& OutSettings)
diff --git a/Engine/Plugins/Runtime/Synthesis/Source/Synthesis/Private/SourceEffectChorus.cpp b/Engine/Plugins/Runtime/Synthesis/Source/Synthesis/Private/SourceEffectChorus.cpp
--- a/Engine/Plugins/Runtime/Synthesis/Source/Synthesis/Private/SourceEffectChorus.cpp
+++ b/Engine/Plugins/Runtime/Synthesis/Source/Synthesis/Private/SourceEffectChorus.cpp
@@ -2,6 +2,9 @@
 
 #include "SourceEffects/SourceEffectChorus.h"
 
+#include "Algo/Transform.h"
+
+
 void FSourceEffectChorus::Init(const FSoundEffectSourceInitData& InitData)
 {
 	bIsActive = true;
@@ -74,42 +77,79 @@ void FSourceEffectChorus::ProcessAudio(const FSoundEffectSourceInputData& InData
 
 void FSourceEffectChorus::SetDepthModulator(const USoundModulatorBase* InModulator)
 {
-	DepthMod.UpdateModulator(InModulator);
+	DepthMod.UpdateModulators({ InModulator });
 }
 
 void FSourceEffectChorus::SetFeedbackModulator(const USoundModulatorBase* InModulator)
 {
-	FeedbackMod.UpdateModulator(InModulator);
+	FeedbackMod.UpdateModulators({ InModulator });
 }
 
 void FSourceEffectChorus::SetSpreadModulator(const USoundModulatorBase* InModulator)
 {
-	SpreadMod.UpdateModulator(InModulator);
+	SpreadMod.UpdateModulators({ InModulator });
 }
 
 void FSourceEffectChorus::SetDryModulator(const USoundModulatorBase* InModulator)
 {
-	DryMod.UpdateModulator(InModulator);
+	DryMod.UpdateModulators({ InModulator });
 }
 
 void FSourceEffectChorus::SetWetModulator(const USoundModulatorBase* InModulator)
 {
-	WetMod.UpdateModulator(InModulator);
+	WetMod.UpdateModulators({ InModulator });
 }
 
 void FSourceEffectChorus::SetFrequencyModulator(const USoundModulatorBase* InModulator)
 {
-	FrequencyMod.UpdateModulator(InModulator);
+	FrequencyMod.UpdateModulators({ InModulator });
+}
+
+void FSourceEffectChorus::SetDepthModulators(const TSet<USoundModulatorBase*>& InModulators)
+{
+	DepthMod.UpdateModulators(InModulators);
+}
+
+void FSourceEffectChorus::SetFeedbackModulators(const TSet<USoundModulatorBase*>& InModulators)
+{
+	FeedbackMod.UpdateModulators(InModulators);
+}
+
+void FSourceEffectChorus::SetSpreadModulators(const TSet<USoundModulatorBase*>& InModulators)
+{
+	SpreadMod.UpdateModulators(InModulators);
+}
+
+void FSourceEffectChorus::SetDryModulators(const TSet<USoundModulatorBase*>& InModulators)
+{
+	DryMod.UpdateModulators(InModulators);
+}
+
+void FSourceEffectChorus::SetWetModulators(const TSet<USoundModulatorBase*>& InModulators)
+{
+	WetMod.UpdateModulators(InModulators);
+}
+
+void FSourceEffectChorus::SetFrequencyModulators(const TSet<USoundModulatorBase*>& InModulators)
+{
+	FrequencyMod.UpdateModulators(InModulators);
 }
 
 void USourceEffectChorusPreset::OnInit()
 {
-	SetDepthModulator(Settings.SpreadModulation.Modulator);
-	SetDryModulator(Settings.DryModulation.Modulator);
-	SetWetModulator(Settings.WetModulation.Modulator);
-	SetFeedbackModulator(Settings.FeedbackModulation.Modulator);
-	SetFrequencyModulator(Settings.FrequencyModulation.Modulator);
-	SetSpreadModulator(Settings.SpreadModulation.Modulator);
+	auto TransformObjectPtrSet = [](const TSet<TObjectPtr<USoundModulatorBase>>& InPtrs)
+	{
+		TSet<USoundModulatorBase*> Objects;
+		Algo::Transform(InPtrs, Objects, [] (const TObjectPtr<USoundModulatorBase>& Obj) { return Obj; });
+		return Objects;
+	};
+
+	SetDepthModulators(TransformObjectPtrSet(Settings.SpreadModulation.Modulators));
+	SetDryModulators(TransformObjectPtrSet(Settings.DryModulation.Modulators));
+	SetWetModulators(TransformObjectPtrSet(Settings.WetModulation.Modulators));
+	SetFeedbackModulators(TransformObjectPtrSet(Settings.FeedbackModulation.Modulators));
+	SetFrequencyModulators(TransformObjectPtrSet(Settings.FrequencyModulation.Modulators));
+	SetSpreadModulators(TransformObjectPtrSet(Settings.SpreadModulation.Modulators));
 }
 
 void USourceEffectChorusPreset::SetDepth(float InDepth)
@@ -128,6 +168,14 @@ void USourceEffectChorusPreset::SetDepthModulator(const USoundModulatorBase* InM
 	});
 }
 
+void USourceEffectChorusPreset::SetDepthModulators(const TSet<USoundModulatorBase*>& Modulators)
+{
+	IterateEffects<FSourceEffectChorus>([&Modulators](FSourceEffectChorus& InDelay)
+	{
+		InDelay.SetDepthModulators(Modulators);
+	});
+}
+
 void USourceEffectChorusPreset::SetFeedback(float InFeedback)
 {
 	UpdateSettings([NewFeedback = InFeedback](FSourceEffectChorusSettings& OutSettings)
@@ -144,6 +192,14 @@ void USourceEffectChorusPreset::SetFeedbackModulator(const USoundModulatorBase*
 	});
 }
 
+void USourceEffectChorusPreset::SetFeedbackModulators(const TSet<USoundModulatorBase*>& Modulators)
+{
+	IterateEffects<FSourceEffectChorus>([&Modulators](FSourceEffectChorus& InDelay)
+	{
+		InDelay.SetFeedbackModulators(Modulators);
+	});
+}
+
 void USourceEffectChorusPreset::SetFrequency(float InFrequency)
 {
 	UpdateSettings([NewFrequency = InFrequency](FSourceEffectChorusSettings& OutSettings)
@@ -160,6 +216,14 @@ void USourceEffectChorusPreset::SetFrequencyModulator(const USoundModulatorBase*
 	});
 }
 
+void USourceEffectChorusPreset::SetFrequencyModulators(const TSet<USoundModulatorBase*>& Modulators)
+{
+	IterateEffects<FSourceEffectChorus>([&Modulators](FSourceEffectChorus& InDelay)
+	{
+		InDelay.SetFrequencyModulators(Modulators);
+	});
+}
+
 void USourceEffectChorusPreset::SetWet(float InWet)
 {
 	UpdateSettings([NewWet = InWet](FSourceEffectChorusSettings& OutSettings)
@@ -176,6 +240,14 @@ void USourceEffectChorusPreset::SetWetModulator(const USoundModulatorBase* InMod
 	});
 }
 
+void USourceEffectChorusPreset::SetWetModulators(const TSet<USoundModulatorBase*>& Modulators)
+{
+	IterateEffects<FSourceEffectChorus>([&Modulators](FSourceEffectChorus& InDelay)
+	{
+		InDelay.SetWetModulators(Modulators);
+	});
+}
+
 void USourceEffectChorusPreset::SetDry(float InDry)
 {
 	UpdateSettings([NewDry = InDry](FSourceEffectChorusSettings& OutSettings)
@@ -192,6 +264,14 @@ void USourceEffectChorusPreset::SetDryModulator(const USoundModulatorBase* InMod
 	});
 }
 
+void USourceEffectChorusPreset::SetDryModulators(const TSet<USoundModulatorBase*>& Modulators)
+{
+	IterateEffects<FSourceEffectChorus>([&Modulators](FSourceEffectChorus& InDelay)
+	{
+		InDelay.SetDryModulators(Modulators);
+	});
+}
+
 void USourceEffectChorusPreset::SetSpread(float InSpread)
 {
 	UpdateSettings([NewSpread = InSpread](FSourceEffectChorusSettings& OutSettings)
@@ -208,6 +288,14 @@ void USourceEffectChorusPreset::SetSpreadModulator(const USoundModulatorBase* In
 	});
 }
 
+void USourceEffectChorusPreset::SetSpreadModulators(const TSet<USoundModulatorBase*>& Modulators)
+{
+	IterateEffects<FSourceEffectChorus>([&Modulators](FSourceEffectChorus& InDelay)
+	{
+		InDelay.SetSpreadModulators(Modulators);
+	});
+}
+
 void USourceEffectChorusPreset::SetSettings(const FSourceEffectChorusBaseSettings& InSettings)
 {
 	UpdateSettings([NewBaseSettings = InSettings](FSourceEffectChorusSettings& OutSettings)
diff --git a/Engine/Source/Editor/AudioEditor/Private/SoundModulationDestinationLayout.cpp b/Engine/Source/Editor/AudioEditor/Private/SoundModulationDestinationLayout.cpp
--- a/Engine/Source/Editor/AudioEditor/Private/SoundModulationDestinationLayout.cpp
+++ b/Engine/Source/Editor/AudioEditor/Private/SoundModulationDestinationLayout.cpp
@@ -88,7 +88,7 @@ namespace ModDestinationLayoutUtils
 		return FName();
 	}
 
-	bool IsParamMismatched(TSharedRef<IPropertyHandle> ModulatorHandle, TSharedRef<IPropertyHandle> StructPropertyHandle, FName* OutModParamName = nullptr, FName* OutDestParamName = nullptr)
+	bool IsParamMismatched(TSharedRef<IPropertyHandle> ModulatorHandle, TSharedRef<IPropertyHandle> StructPropertyHandle, FName* OutModParamName = nullptr, FName* OutDestParamName = nullptr, FName* OutModName = nullptr)
 	{
 		if (OutModParamName)
 		{
@@ -100,6 +100,11 @@ namespace ModDestinationLayoutUtils
 			*OutDestParamName = FName();
 		}
 
+		if (OutModName)
+		{
+			*OutModName = FName();
+		}
+
 		UObject* ModObject = nullptr;
 		ModulatorHandle->GetValue(ModObject);
 
@@ -122,6 +127,12 @@ namespace ModDestinationLayoutUtils
 			{
 				*OutDestParamName = DestParamName;
 			}
+
+			if (OutModName)
+			{
+				*OutModName = ModBase->GetFName();
+			}
+
 			return true;
 		}
 
@@ -199,7 +210,7 @@ namespace ModDestinationLayoutUtils
 		IDetailChildrenBuilder& ChildBuilder,
 		TSharedRef<IPropertyHandle> StructPropertyHandle,
 		TSharedRef<IPropertyHandle> ValueHandle,
-		TSharedRef<IPropertyHandle> ModulatorHandle,
+		TSharedRef<IPropertyHandle> ModulatorsHandle,
 		TSharedRef<IPropertyHandle> EnablementHandle)
 	{
 		FText UnitDisplayText = FText::GetEmpty();
@@ -209,12 +220,12 @@ namespace ModDestinationLayoutUtils
 		const FText DisplayName = StructPropertyHandle->GetPropertyDisplayName();
 		ChildBuilder.AddCustomRow(DisplayName)
 		.NameContent()
-			[
-				SNew(STextBlock)
-				.Font(IDetailLayoutBuilder::GetDetailFont())
-				.Text(DisplayName)
-				.ToolTipText(StructPropertyHandle->GetToolTipText())
-			]
+		[
+			SNew(STextBlock)
+			.Font(IDetailLayoutBuilder::GetDetailFont())
+			.Text(DisplayName)
+			.ToolTipText(StructPropertyHandle->GetToolTipText())
+		]
 		.ValueContent()
 		.MinDesiredWidth(250.0f)
 		[
@@ -284,7 +295,7 @@ namespace ModDestinationLayoutUtils
 				]
 		];
 
-		EnablementHandle->SetOnPropertyValueChanged(FSimpleDelegate::CreateLambda([EnablementHandle, ValueHandle, StructPropertyHandle, ModulatorHandle]()
+		EnablementHandle->SetOnPropertyValueChanged(FSimpleDelegate::CreateLambda([EnablementHandle, ValueHandle, StructPropertyHandle, ModulatorsHandle]()
 		{
 			bool bEnabled = false;
 			EnablementHandle->GetValue(bEnabled);
@@ -296,8 +307,7 @@ namespace ModDestinationLayoutUtils
 			}
 			else
 			{
-				UObject* NullObj = nullptr;
-				ModulatorHandle->SetValue(NullObj);
+				ModulatorsHandle->AsSet()->Empty();
 			}
 		}));
 	}
@@ -331,76 +341,84 @@ namespace ModDestinationLayoutUtils
 		];
 	}
 
-	void CustomizeChildren_AddModulatorRow(IDetailChildrenBuilder& ChildBuilder, TSharedRef<IPropertyHandle> StructPropertyHandle, TSharedRef<IPropertyHandle> ModulatorHandle, TSharedRef<IPropertyHandle> EnablementHandle)
+	void CustomizeChildren_AddModulatorRow(IDetailChildrenBuilder& ChildBuilder, TSharedRef<IPropertyHandle> StructPropertyHandle, TSharedRef<IPropertyHandle> ModulatorsHandle, TSharedRef<IPropertyHandle> EnablementHandle)
 	{
 		const FText DisplayName = StructPropertyHandle->GetPropertyDisplayName();
-		ChildBuilder.AddCustomRow(DisplayName)
+
+		auto ModEnabled = [EnablementHandle]()
+		{
+			bool bModulationEnabled = false;
+			EnablementHandle->GetValue(bModulationEnabled);
+			return bModulationEnabled
+				? EVisibility::Visible
+				: EVisibility::Collapsed;
+		};
+
+		const TAttribute<EVisibility> ModulatorVisibility = TAttribute<EVisibility>::Create(ModEnabled);
+
+		ChildBuilder.AddProperty(ModulatorsHandle)
+			.DisplayName(FText::Format(LOCTEXT("SoundModulationParameter_ModulatorsFormat", "{0} Modulators"), DisplayName))
+			.Visibility(ModulatorVisibility);
+
+		ChildBuilder.AddCustomRow(LOCTEXT("SoundModulationDestinationLayout_UnitMismatchHeadingWarning", "Unit Mismatch Warning"))
 			.NameContent()
 			[
 				SNew(STextBlock)
 				.Font(IDetailLayoutBuilder::GetDetailFont())
-				.Text(FText::Format(LOCTEXT("SoundModulationParameter_ModulatorFormat", "{0} Modulator"), DisplayName))
-				.ToolTipText(ModulatorHandle->GetToolTipText())
-			]
-			.ValueContent()
-			.MinDesiredWidth(200.0f)
-			[
-				SNew(SHorizontalBox)
-				+ SHorizontalBox::Slot()
-				.FillWidth(1.0f)
-				.Padding(4.0f, 0.0f, 0.0f, 0.0f)
-				.VAlign(VAlign_Center)
-				[
-					ModulatorHandle->CreatePropertyValueWidget()
-				]
+				.Text(FText::Format(LOCTEXT("ModulationDestinationLayout_UnitMismatchHeader", "{0} Mismatched Parameters"), DisplayName))
+				.ToolTipText(StructPropertyHandle->GetToolTipText())
 			]
-			.Visibility(TAttribute<EVisibility>::Create([EnablementHandle]()
-			{
-				bool bModulationEnabled = false;
-				EnablementHandle->GetValue(bModulationEnabled);
-				return bModulationEnabled 
-					? EVisibility::Visible
-					: EVisibility::Hidden;
-			}));
-
-		ChildBuilder.AddCustomRow(LOCTEXT("SoundModulationDestinationLayout_UnitMismatchHeadingWarning", "Unit Mismatch Warning"))
 			.ValueContent()
 			.MinDesiredWidth(150.0f)
 			[
 				SNew(SHorizontalBox)
 				+ SHorizontalBox::Slot()
 					.FillWidth(1.0f)
-					.Padding(10.0f, 0.0f, 0.0f, 0.0f)
+					.Padding(0.0f, 3.0f, 0.0f, 3.0f)
 					.VAlign(VAlign_Center)
 					[
 						SNew(STextBlock)
 						.Font(IDetailLayoutBuilder::GetDetailFontBold())
-						.Text(TAttribute<FText>::Create([ModulatorHandle, StructPropertyHandle]()
+						.Text(TAttribute<FText>::Create([ModSet = ModulatorsHandle->AsSet(), StructPropertyHandle]()
 						{
-							UObject* ModObject = nullptr;
-							ModulatorHandle->GetValue(ModObject);
-
-							UObject* DestObject = nullptr;
-							
-							FName ModName;
-							FName DestName;
-							if (ModDestinationLayoutUtils::IsParamMismatched(ModulatorHandle, StructPropertyHandle, &ModName, &DestName))
+							FText MismatchedText = FText::GetEmpty();
+							ModSet->IterateElements([&MismatchedText, &StructPropertyHandle](TSharedRef<IPropertyHandle> ModulatorHandle)
 							{
-								return FText::Format(LOCTEXT("ModulationDestinationLayout_UnitMismatchFormat", "Parameter Mismatch: Modulator Output = {0}, Destination Input = {1}"),
-									FText::FromName(ModName),
-									FText::FromName(DestName));
-							}
-
-							return FText::GetEmpty();
-
+								FName ModParamName;
+								FName DestName;
+								FName ModName;
+								if (ModDestinationLayoutUtils::IsParamMismatched(ModulatorHandle, StructPropertyHandle, &ModParamName, &DestName, &ModName))
+								{
+									MismatchedText = FText::Format(
+										LOCTEXT("ModulationDestinationLayout_UnitMismatchesFormat", "{0}{1} ({2})\n"),
+										MismatchedText,
+										FText::FromName(ModName),
+										FText::FromName(ModParamName)
+									);
+								}
+							});
+
+							return MismatchedText;
 						}))
 					]
 			]
-			.Visibility(TAttribute<EVisibility>::Create([ModulatorHandle, StructPropertyHandle]()
+			.Visibility(TAttribute<EVisibility>::Create([ModSet = ModulatorsHandle->AsSet(), ModEnabled, StructPropertyHandle]()
 			{
-				return ModDestinationLayoutUtils::IsParamMismatched(ModulatorHandle, StructPropertyHandle)
-					? EVisibility::Visible
-					: EVisibility::Hidden;
+				EVisibility Visibility = ModEnabled();
+				if (Visibility == EVisibility::Collapsed)
+				{
+					return Visibility;
+				}
+
+				bool bIsMismatched = false;
+				ModSet->IterateElements([&bIsMismatched, &StructPropertyHandle](TSharedPtr<IPropertyHandle> ModulatorHandle)
+				{
+					if (ModDestinationLayoutUtils::IsParamMismatched(ModulatorHandle->AsShared(), StructPropertyHandle))
+					{
+						bIsMismatched = true;
+					}
+				});
+				return bIsMismatched ? EVisibility::Visible : EVisibility::Hidden;
 			}));
 	}
 
@@ -436,13 +454,13 @@ void FSoundModulationDestinationLayoutCustomization::CustomizeChildren(TSharedRe
 	}
 
 	TSharedRef<IPropertyHandle>EnablementHandle = PropertyHandles.FindChecked(GET_MEMBER_NAME_CHECKED(FSoundModulationDestinationSettings, bEnableModulation)).ToSharedRef();
-	TSharedRef<IPropertyHandle>ModulatorHandle = PropertyHandles.FindChecked(GET_MEMBER_NAME_CHECKED(FSoundModulationDestinationSettings, Modulator)).ToSharedRef();
+	TSharedRef<IPropertyHandle>ModulatorsHandle = PropertyHandles.FindChecked(GET_MEMBER_NAME_CHECKED(FSoundModulationDestinationSettings, Modulators)).ToSharedRef();
 	TSharedRef<IPropertyHandle>ValueHandle = PropertyHandles.FindChecked(GET_MEMBER_NAME_CHECKED(FSoundModulationDestinationSettings, Value)).ToSharedRef();
 
 	if (ModDestinationLayoutUtils::IsModulationEnabled())
 	{
-		ModDestinationLayoutUtils::CustomizeChildren_AddValueRow(ChildBuilder, StructPropertyHandle, ValueHandle, ModulatorHandle, EnablementHandle);
-		ModDestinationLayoutUtils::CustomizeChildren_AddModulatorRow(ChildBuilder, StructPropertyHandle, ModulatorHandle, EnablementHandle);
+		ModDestinationLayoutUtils::CustomizeChildren_AddValueRow(ChildBuilder, StructPropertyHandle, ValueHandle, ModulatorsHandle, EnablementHandle);
+		ModDestinationLayoutUtils::CustomizeChildren_AddModulatorRow(ChildBuilder, StructPropertyHandle, ModulatorsHandle, EnablementHandle);
 	}
 	else
 	{
@@ -519,15 +537,25 @@ void FSoundModulationDefaultRoutingSettingsLayoutCustomization::CustomizeChildre
 				uint8 RoutingValue = 0;
 				if (RoutingHandle->GetValue(RoutingValue) != FPropertyAccess::Success)
 				{
-					return EVisibility::Hidden;
+					return EVisibility::Collapsed;
 				}
 
-				if (static_cast<EModulationRouting>(RoutingValue) != EModulationRouting::Override)
+				switch (static_cast<EModulationRouting>(RoutingValue))
 				{
-					return EVisibility::Hidden;
-				}
+					case EModulationRouting::Disable:
+					case EModulationRouting::Inherit:
+					{
+						return EVisibility::Collapsed;
+					}
 
-				return EVisibility::Visible;
+					case EModulationRouting::Override:
+					case EModulationRouting::Union:
+					default:
+					{
+						return EVisibility::Visible;
+					}
+					break;
+				}
 			});
 		};
 
diff --git a/Engine/Source/Editor/PropertyEditor/Private/PropertyHandleImpl.cpp b/Engine/Source/Editor/PropertyEditor/Private/PropertyHandleImpl.cpp
--- a/Engine/Source/Editor/PropertyEditor/Private/PropertyHandleImpl.cpp
+++ b/Engine/Source/Editor/PropertyEditor/Private/PropertyHandleImpl.cpp
@@ -30,6 +30,7 @@
 #include "PropertyPathHelpers.h"
 #include "PropertyTextUtilities.h"
 #include "HAL/PlatformApplicationMisc.h"
+#include "Templates/Function.h"
 
 #define LOCTEXT_NAMESPACE "PropertyHandleImplementation"
 
@@ -4861,6 +4862,19 @@ FPropertyAccess::Result FPropertyHandleSet::GetNumElements(uint32& OutNumChildre
 	return FPropertyAccess::Success;
 }
 
+void FPropertyHandleSet::IterateElements(TFunctionRef<void(TSharedRef<IPropertyHandle>)> Func)
+{
+	uint32 NumChildren = 0;
+	GetNumChildren(NumChildren);
+
+	for (uint32 Index = 0; Index < NumChildren; ++Index)
+	{
+		TSharedPtr<FPropertyNode> PropertyNode = Implementation->GetChildNode(Index);
+		TSharedRef<IPropertyHandle> ChildHandle = PropertyEditorHelpers::GetPropertyHandle(PropertyNode.ToSharedRef(), Implementation->GetNotifyHook(), Implementation->GetPropertyUtilities()).ToSharedRef();
+		Func(ChildHandle);
+	}
+}
+
 void FPropertyHandleSet::SetOnNumElementsChanged(FSimpleDelegate& OnChildrenChanged)
 {
 	Implementation->SetOnRebuildChildren(OnChildrenChanged);
diff --git a/Engine/Source/Editor/PropertyEditor/Private/PropertyHandleImpl.h b/Engine/Source/Editor/PropertyEditor/Private/PropertyHandleImpl.h
--- a/Engine/Source/Editor/PropertyEditor/Private/PropertyHandleImpl.h
+++ b/Engine/Source/Editor/PropertyEditor/Private/PropertyHandleImpl.h
@@ -706,6 +706,8 @@ public:
 	virtual FPropertyAccess::Result Empty() override;
 	virtual FPropertyAccess::Result DeleteItem(int32 Index) override;
 	virtual FPropertyAccess::Result GetNumElements(uint32& OutNumElements) override;
+	virtual void IterateElements(TFunctionRef<void(TSharedRef<IPropertyHandle>)> Func) override;
+
 	virtual void SetOnNumElementsChanged(FSimpleDelegate& InOnNumElementsChanged) override;
 	virtual bool HasDocumentation() override { return true; }
 	virtual FString GetDocumentationLink() override { return FString("Engine/UI/LevelEditor/Details/Properties/Set/"); }
diff --git a/Engine/Source/Editor/PropertyEditor/Public/PropertyHandle.h b/Engine/Source/Editor/PropertyEditor/Public/PropertyHandle.h
--- a/Engine/Source/Editor/PropertyEditor/Public/PropertyHandle.h
+++ b/Engine/Source/Editor/PropertyEditor/Public/PropertyHandle.h
@@ -7,6 +7,7 @@
 #include "Widgets/SWidget.h"
 #include "PropertyEditorModule.h"
 #include "UObject/PropertyPortFlags.h"
+#include "Templates/Function.h"
 
 struct FAssetData;
 class IPropertyHandleArray;
@@ -795,6 +796,12 @@ public:
 	 */
 	virtual FPropertyAccess::Result GetNumElements(uint32& OutNumElements) = 0;
 
+	/**
+	 * Iterates over all set elements using the provided function
+	 */
+	virtual void IterateElements(TFunctionRef<void(TSharedRef<IPropertyHandle>)> Func) = 0;
+
+
 	/**
 	 * Sets a delegate to call when the number of elements changes
 	 */
diff --git a/Engine/Source/Runtime/AudioExtensions/Public/IAudioModulation.h b/Engine/Source/Runtime/AudioExtensions/Public/IAudioModulation.h
--- a/Engine/Source/Runtime/AudioExtensions/Public/IAudioModulation.h
+++ b/Engine/Source/Runtime/AudioExtensions/Public/IAudioModulation.h
@@ -136,6 +136,21 @@ namespace Audio
 		bool GetValueThreadSafe(float& OutValue) const;
 		bool IsValid() const;
 
+		friend FORCEINLINE uint32 GetTypeHash(const FModulatorHandle& InModulatorHandle)
+		{
+			return HashCombineFast(InModulatorHandle.HandleId, InModulatorHandle.ModulatorId);
+		}
+
+		FORCEINLINE bool operator==(const FModulatorHandle& Other) const
+		{
+			return HandleId == Other.HandleId && ModulatorId == Other.ModulatorId;
+		}
+
+		FORCEINLINE bool operator!=(const FModulatorHandle& Other) const
+		{
+			return !(*this == Other);
+		}
+
 	private:
 		FModulationParameter Parameter;
 		FModulatorHandleId HandleId = INDEX_NONE;
@@ -149,43 +164,40 @@ class AUDIOEXTENSIONS_API IAudioModulationManager : public TSharedFromThis<IAudi
 {
 public:
 	/** Virtual destructor */
-	virtual ~IAudioModulationManager() { }
+	virtual ~IAudioModulationManager() = default;
 
 	/** Initialize the modulation plugin with the same rate and number of sources */
-	virtual void Initialize(const FAudioPluginInitializationParams& InitializationParams) { }
+	virtual void Initialize(const FAudioPluginInitializationParams& InitializationParams) = 0;
 
-	virtual void OnAuditionEnd() { }
+	virtual void OnAuditionEnd() = 0;
 
 #if !UE_BUILD_SHIPPING
 	/** Request to post help from active plugin (non-shipping builds only) */
-	virtual bool OnPostHelp(FCommonViewportClient* ViewportClient, const TCHAR* Stream) { return false; };
+	virtual bool OnPostHelp(FCommonViewportClient* ViewportClient, const TCHAR* Stream) = 0; 
 
 	/** Render stats pertaining to modulation (non-shipping builds only) */
-	virtual int32 OnRenderStat(FViewport* Viewport, FCanvas* Canvas, int32 X, int32 Y, const UFont& Font, const FVector* ViewLocation, const FRotator* ViewRotation) { return Y; }
+	virtual int32 OnRenderStat(FViewport* Viewport, FCanvas* Canvas, int32 X, int32 Y, const UFont& Font, const FVector* ViewLocation, const FRotator* ViewRotation) = 0;
 
 	/** Toggle showing render stats pertaining to modulation (non-shipping builds only) */
-	virtual bool OnToggleStat(FCommonViewportClient* ViewportClient, const TCHAR* Stream) { return false; }
+	virtual bool OnToggleStat(FCommonViewportClient* ViewportClient, const TCHAR* Stream) = 0;
 #endif //!UE_BUILD_SHIPPING
 
-	/** Processes audio with the given input and output data structs.*/
-	virtual void ProcessAudio(const FAudioPluginSourceInputData& InputData, FAudioPluginSourceOutputData& OutputData) { }
-
 	/** Processes all modulators Run on the audio render thread prior to processing audio */
-	virtual void ProcessModulators(const double InElapsed) { }
+	virtual void ProcessModulators(const double InElapsed) = 0;
 
 	/** Updates modulator definition on the AudioRender Thread with that provided by the UObject representation */
-	virtual void UpdateModulator(const USoundModulatorBase& InModulator) { }
+	virtual void UpdateModulator(const USoundModulatorBase& InModulator) = 0;
 
 protected:
-	virtual void RegisterModulator(uint32 InHandleId, Audio::FModulatorId InModulatorId) { }
+	virtual void RegisterModulator(uint32 InHandleId, Audio::FModulatorId InModulatorId) = 0;
 
 	// Get the modulator value from the AudioRender Thread
-	virtual bool GetModulatorValue(const Audio::FModulatorHandle& ModulatorHandle, float& OutValue) const { return false; }
+	virtual bool GetModulatorValue(const Audio::FModulatorHandle& ModulatorHandle, float& OutValue) const = 0;
 
 	// Get the modulator value from any thread.
-	virtual bool GetModulatorValueThreadSafe(const Audio::FModulatorHandle& ModulatorHandle, float& OutValue) const { return false; }
+	virtual bool GetModulatorValueThreadSafe(const Audio::FModulatorHandle& ModulatorHandle, float& OutValue) const = 0;
 
-	virtual void UnregisterModulator(const Audio::FModulatorHandle& InHandle) { }
+	virtual void UnregisterModulator(const Audio::FModulatorHandle& InHandle) = 0;
 
 	friend Audio::FModulatorHandle;
 };
diff --git a/Engine/Source/Runtime/AudioMixer/Private/AudioMixerDevice.cpp b/Engine/Source/Runtime/AudioMixer/Private/AudioMixerDevice.cpp
--- a/Engine/Source/Runtime/AudioMixer/Private/AudioMixerDevice.cpp
+++ b/Engine/Source/Runtime/AudioMixer/Private/AudioMixerDevice.cpp
@@ -1453,15 +1453,15 @@ namespace Audio
 		}
 	}
 
-	void FMixerDevice::UpdateSubmixModulationSettings(USoundSubmix* InSoundSubmix, USoundModulatorBase* InOutputModulation, USoundModulatorBase* InWetLevelModulation, USoundModulatorBase* InDryLevelModulation)
+	void FMixerDevice::UpdateSubmixModulationSettings(USoundSubmix* InSoundSubmix, const TSet<TObjectPtr<USoundModulatorBase>>& InOutputModulation, const TSet<TObjectPtr<USoundModulatorBase>>& InWetLevelModulation, const TSet<TObjectPtr<USoundModulatorBase>>& InDryLevelModulation)
 	{
 		if (!IsInAudioThread())
 		{
 			FMixerDevice* MixerDevice = this;
 
-			FAudioThread::RunCommandOnAudioThread([MixerDevice, InSoundSubmix, InOutputModulation, InWetLevelModulation, InDryLevelModulation]() 
+			FAudioThread::RunCommandOnAudioThread([MixerDevice, InSoundSubmix, OutMod = InOutputModulation, WetMod = InWetLevelModulation, DryMod = InDryLevelModulation]()
 			{
-				MixerDevice->UpdateSubmixModulationSettings(InSoundSubmix, InOutputModulation, InWetLevelModulation, InDryLevelModulation);
+				MixerDevice->UpdateSubmixModulationSettings(InSoundSubmix, OutMod, WetMod, DryMod);
 			});
 			return;
 		}
@@ -1472,11 +1472,7 @@ namespace Audio
 
 			if (MixerSubmixPtr.IsValid())
 			{
-				USoundModulatorBase* VolumeMod = InOutputModulation;
-				USoundModulatorBase* WetMod = InOutputModulation;
-				USoundModulatorBase* DryMod = InOutputModulation;
-
-				AudioRenderThreadCommand([MixerSubmixPtr, VolumeMod, WetMod, DryMod]()
+				AudioRenderThreadCommand([MixerSubmixPtr, VolumeMod = InOutputModulation, WetMod = InOutputModulation, DryMod = InOutputModulation]()
 				{
 					MixerSubmixPtr->UpdateModulationSettings(VolumeMod, WetMod, DryMod);
 				});
diff --git a/Engine/Source/Runtime/AudioMixer/Private/AudioMixerSource.cpp b/Engine/Source/Runtime/AudioMixer/Private/AudioMixerSource.cpp
--- a/Engine/Source/Runtime/AudioMixer/Private/AudioMixerSource.cpp
+++ b/Engine/Source/Runtime/AudioMixer/Private/AudioMixerSource.cpp
@@ -9,10 +9,12 @@
 #include "AudioMixerTrace.h"
 #include "ContentStreaming.h"
 #include "IAudioExtensionPlugin.h"
+#include "IAudioModulation.h"
 #include "ProfilingDebugging/CsvProfiler.h"
 #include "Sound/AudioSettings.h"
 #include "Sound/SoundModulationDestination.h"
 #include "Misc/ScopeRWLock.h"
+#include "Templates/Function.h"
 
 // Link to "Audio" profiling category
 CSV_DECLARE_CATEGORY_MODULE_EXTERN(AUDIOMIXERCORE_API, Audio);
@@ -57,37 +59,79 @@ namespace Audio
 					return EMixerSourceSubmixSendStage::PostDistanceAttenuation;
 			}
 		}
+
+		const USoundClass* GetFallbackSoundClass(const FActiveSound& InActiveSound, const FWaveInstance& InWaveInstance)
+		{
+			const USoundClass* SoundClass = InActiveSound.GetSoundClass();
+			if (InWaveInstance.SoundClass)
+			{
+				SoundClass = InWaveInstance.SoundClass;
+			}
+
+			return SoundClass;
+		}
 	}
 	namespace ModulationUtils
 	{
-		static const FSoundModulationDestinationSettings DefaultDestination;
+		void MixInRoutedValue(const FModulationParameter& InParam, float& InOutValueA, float InValueB)
+		{
+			if (InParam.bRequiresConversion)
+			{
+				InParam.NormalizedFunction(InOutValueA);
+				InParam.NormalizedFunction(InValueB);
+			}
+			InParam.MixFunction(InOutValueA, InValueB);
+			if (InParam.bRequiresConversion)
+			{
+				InParam.UnitFunction(InOutValueA);
+			}
+		}
 
-		const FSoundModulationDestinationSettings& GetRoutedVolumeModulation(const FWaveInstance& InWaveInstance, const USoundWave& InWaveData, FActiveSound* InActiveSound)
+		FSoundModulationDestinationSettings InitRoutedDestinationSettings(
+			const EModulationRouting& InActiveSoundRouting,
+			const FSoundModulationDestinationSettings& InActiveSoundSettings,
+			const EModulationRouting& InWaveRouting,
+			const FSoundModulationDestinationSettings& InWaveSettings,
+			const USoundClass* InSoundClass,
+			const FModulationParameter& InParam,
+			TFunctionRef<const FSoundModulationDestinationSettings* (const USoundClass&)> InGetSoundClassDestinationFunction)
 		{
-			const FSoundModulationDefaultRoutingSettings& RoutingSettings = InActiveSound->ModulationRouting;
-			switch (RoutingSettings.VolumeRouting)
+			auto UnionSoundClassSettings = [&](FSoundModulationDestinationSettings& InOutSettings)
 			{
-				case EModulationRouting::Inherit:
+				if (InSoundClass)
+				{
+					const FSoundModulationDestinationSettings& ClassSettings = *InGetSoundClassDestinationFunction(*InSoundClass);
+					MixInRoutedValue(InParam, InOutSettings.Value, ClassSettings.Value);
+					InOutSettings.Modulators = InOutSettings.Modulators.Union(ClassSettings.Modulators);
+				}
+			};
+
+			switch (InActiveSoundRouting)
+			{
+				case EModulationRouting::Union:
 				{
-					switch (InWaveData.ModulationSettings.VolumeRouting)
+					FSoundModulationDestinationSettings UnionSettings = InActiveSoundSettings;
+					switch (InWaveRouting)
 					{
+						case EModulationRouting::Union:
+						{
+							MixInRoutedValue(InParam, UnionSettings.Value, InWaveSettings.Value);
+							UnionSettings.Modulators = UnionSettings.Modulators.Union(InWaveSettings.Modulators);
+							UnionSoundClassSettings(UnionSettings);
+							return UnionSettings;
+						}
+						break;
+
 						case EModulationRouting::Inherit:
 						{
-							USoundClass* SoundClass = InActiveSound->GetSoundClass();
-							if (InWaveInstance.SoundClass)
-							{
-								SoundClass = InWaveInstance.SoundClass;
-							}
-							if (SoundClass)
-							{
-								return SoundClass->Properties.ModulationSettings.VolumeModulationDestination;
-							}
+							UnionSoundClassSettings(UnionSettings);
 						}
 						break;
 
 						case EModulationRouting::Override:
 						{
-							return InWaveData.ModulationSettings.VolumeModulationDestination;
+							MixInRoutedValue(InParam, UnionSettings.Value, InWaveSettings.Value);
+							UnionSettings.Modulators = UnionSettings.Modulators.Union(InWaveSettings.Modulators);
 						}
 						break;
 
@@ -95,49 +139,35 @@ namespace Audio
 						default:
 						break;
 					}
-				}
-				break;
 
-				case EModulationRouting::Override:
-				{
-					return RoutingSettings.VolumeModulationDestination;
+					return UnionSettings;
 				}
 				break;
 
-				case EModulationRouting::Disable:
-				default:
-				break;
-			}
-
-			return DefaultDestination;
-		}
-
-		const FSoundModulationDestinationSettings& GetRoutedPitchModulation(const FWaveInstance& InWaveInstance, const USoundWave& InWaveData, FActiveSound* InActiveSound)
-		{
-			const FSoundModulationDefaultRoutingSettings& RoutingSettings = InActiveSound->ModulationRouting;
-			switch (RoutingSettings.PitchRouting)
-			{
 				case EModulationRouting::Inherit:
 				{
-					switch (InWaveData.ModulationSettings.PitchRouting)
+					switch (InWaveRouting)
 					{
+						case EModulationRouting::Union:
+						{
+							FSoundModulationDestinationSettings UnionSettings = InWaveSettings;
+							UnionSoundClassSettings(UnionSettings);
+							return UnionSettings;
+						}
+						break;
+
 						case EModulationRouting::Inherit:
 						{
-							USoundClass* SoundClass = InActiveSound->GetSoundClass();
-							if (InWaveInstance.SoundClass)
+							if (InSoundClass)
 							{
-								SoundClass = InWaveInstance.SoundClass;
-							}
-							if (SoundClass)
-							{
-								return SoundClass->Properties.ModulationSettings.PitchModulationDestination;
+								return *InGetSoundClassDestinationFunction(*InSoundClass);
 							}
 						}
 						break;
 
 						case EModulationRouting::Override:
 						{
-							return InWaveData.ModulationSettings.PitchModulationDestination;
+							return InWaveSettings;
 						}
 						break;
 
@@ -150,43 +180,60 @@ namespace Audio
 
 				case EModulationRouting::Override:
 				{
-					return RoutingSettings.PitchModulationDestination;
+					return InActiveSoundSettings;
 				}
 				break;
+
 				case EModulationRouting::Disable:
 				default:
 				break;
 			}
 
-			return DefaultDestination;
+			return { };
 		}
 
-		const FSoundModulationDestinationSettings& GetRoutedHighpassModulation(const FWaveInstance& InWaveInstance, const USoundWave& InWaveData, FActiveSound* InActiveSound)
+		float GetRoutedDestinationValue(
+			const EModulationRouting& InActiveSoundRouting,
+			const FSoundModulationDestinationSettings& InActiveSoundSettings,
+			const EModulationRouting& InWaveRouting,
+			const FSoundModulationDestinationSettings& InWaveSettings,
+			const USoundClass* InSoundClass,
+			const FModulationParameter& InParam,
+			TFunctionRef<const FSoundModulationDestinationSettings* (const USoundClass&)> InGetSoundClassDestinationFunction)
 		{
-			const FSoundModulationDefaultRoutingSettings& RoutingSettings = InActiveSound->ModulationRouting;
-			switch (RoutingSettings.HighpassRouting)
+			auto MixInSoundClassValue = [&](float& InOutValue)
 			{
-				case EModulationRouting::Inherit:
+				if (InSoundClass)
 				{
-					switch (InWaveData.ModulationSettings.HighpassRouting)
+					const FSoundModulationDestinationSettings& ClassSettings = *InGetSoundClassDestinationFunction(*InSoundClass);
+					MixInRoutedValue(InParam, InOutValue, ClassSettings.Value);
+				}
+			};
+
+			switch (InActiveSoundRouting)
+			{
+				case EModulationRouting::Union:
+				{
+					float UnionValue = InActiveSoundSettings.Value;
+					switch (InWaveRouting)
 					{
+						case EModulationRouting::Union:
+						{
+							MixInRoutedValue(InParam, UnionValue, InWaveSettings.Value);
+							MixInSoundClassValue(UnionValue);
+							return UnionValue;
+						}
+						break;
+
 						case EModulationRouting::Inherit:
 						{
-							USoundClass* SoundClass = InActiveSound->GetSoundClass();
-							if (InWaveInstance.SoundClass)
-							{
-								SoundClass = InWaveInstance.SoundClass;
-							}
-							if (SoundClass)
-							{
-								return SoundClass->Properties.ModulationSettings.HighpassModulationDestination;
-							}
+							MixInSoundClassValue(UnionValue);
 						}
 						break;
 
 						case EModulationRouting::Override:
 						{
-							return InWaveData.ModulationSettings.HighpassModulationDestination;
+							MixInRoutedValue(InParam, UnionValue, InWaveSettings.Value);
 						}
 						break;
 
@@ -194,49 +241,35 @@ namespace Audio
 						default:
 						break;
 					}
-				}
-				break;
 
-				case EModulationRouting::Override:
-				{
-					return RoutingSettings.HighpassModulationDestination;
+					return UnionValue;
 				}
 				break;
 
-				case EModulationRouting::Disable:
-				default:
-				break;
-			}
-
-			return DefaultDestination;
-		}
-
-		const FSoundModulationDestinationSettings& GetRoutedLowpassModulation(const FWaveInstance& InWaveInstance, const USoundWave& InWaveData, FActiveSound* InActiveSound)
-		{
-			const FSoundModulationDefaultRoutingSettings& RoutingSettings = InActiveSound->ModulationRouting;
-			switch (RoutingSettings.LowpassRouting)
-			{
 				case EModulationRouting::Inherit:
 				{
-					switch (InWaveData.ModulationSettings.LowpassRouting)
+					switch (InWaveRouting)
 					{
+						case EModulationRouting::Union:
+						{
+							float UnionValue = InWaveSettings.Value;
+							MixInSoundClassValue(UnionValue);
+							return UnionValue;
+						}
+						break;
+
 						case EModulationRouting::Inherit:
 						{
-							USoundClass* SoundClass = InActiveSound->GetSoundClass();
-							if (InWaveInstance.SoundClass)
-							{
-								SoundClass = InWaveInstance.SoundClass;
-							}
-							if (SoundClass)
+							if (InSoundClass)
 							{
-								return SoundClass->Properties.ModulationSettings.LowpassModulationDestination;
+								return InGetSoundClassDestinationFunction(*InSoundClass)->Value;
 							}
 						}
 						break;
 
 						case EModulationRouting::Override:
 						{
-							return InWaveData.ModulationSettings.LowpassModulationDestination;
+							return InWaveSettings.Value;
 						}
 						break;
 
@@ -249,7 +282,7 @@ namespace Audio
 
 				case EModulationRouting::Override:
 				{
-					return RoutingSettings.LowpassModulationDestination;
+					return InActiveSoundSettings.Value;
 				}
 				break;
 
@@ -258,18 +291,170 @@ namespace Audio
 				break;
 			}
 
-			return DefaultDestination;
+			return 1.0f;
+		}
+
+		FSoundModulationDestinationSettings InitRoutedVolumeModulation(const FWaveInstance& InWaveInstance, const USoundWave& InWaveData, const FActiveSound& InActiveSound)
+		{
+			const EModulationRouting& ActiveSoundRouting = InActiveSound.ModulationRouting.VolumeRouting;
+			const FSoundModulationDestinationSettings& ActiveSoundSettings = InActiveSound.ModulationRouting.VolumeModulationDestination;
+
+			const EModulationRouting& WaveRouting = InWaveData.ModulationSettings.VolumeRouting;
+			const FSoundModulationDestinationSettings& WaveSettings = InWaveData.ModulationSettings.VolumeModulationDestination;
+
+			return InitRoutedDestinationSettings(
+				ActiveSoundRouting,
+				ActiveSoundSettings,
+				WaveRouting,
+				WaveSettings,
+				MixerSourcePrivate::GetFallbackSoundClass(InActiveSound, InWaveInstance),
+				Audio::GetModulationParameter("Volume"),
+				[](const USoundClass& InSoundClass) { return &InSoundClass.Properties.ModulationSettings.VolumeModulationDestination; }
+			);
+		}
+
+		float GetRoutedVolume(const FWaveInstance& InWaveInstance, const USoundWave& InWaveData, const FActiveSound& InActiveSound)
+		{
+			const EModulationRouting& ActiveSoundRouting = InActiveSound.ModulationRouting.VolumeRouting;
+			const FSoundModulationDestinationSettings& ActiveSoundSettings = InActiveSound.ModulationRouting.VolumeModulationDestination;
+
+			const EModulationRouting& WaveRouting = InWaveData.ModulationSettings.VolumeRouting;
+			const FSoundModulationDestinationSettings& WaveSettings = InWaveData.ModulationSettings.VolumeModulationDestination;
+
+			return GetRoutedDestinationValue(
+				ActiveSoundRouting,
+				ActiveSoundSettings,
+				WaveRouting,
+				WaveSettings,
+				MixerSourcePrivate::GetFallbackSoundClass(InActiveSound, InWaveInstance),
+				Audio::GetModulationParameter("Volume"),
+				[](const USoundClass& InSoundClass) { return &InSoundClass.Properties.ModulationSettings.VolumeModulationDestination; }
+			);
+		}
+
+		FSoundModulationDestinationSettings InitRoutedPitchModulation(const FWaveInstance& InWaveInstance, const USoundWave& InWaveData, const FActiveSound& InActiveSound)
+		{
+			const EModulationRouting& ActiveSoundRouting = InActiveSound.ModulationRouting.PitchRouting;
+			const FSoundModulationDestinationSettings& ActiveSoundSettings = InActiveSound.ModulationRouting.PitchModulationDestination;
+
+			const EModulationRouting& WaveRouting = InWaveData.ModulationSettings.PitchRouting;
+			const FSoundModulationDestinationSettings& WaveSettings = InWaveData.ModulationSettings.PitchModulationDestination;
+
+			return InitRoutedDestinationSettings(
+				ActiveSoundRouting,
+				ActiveSoundSettings,
+				WaveRouting,
+				WaveSettings,
+				MixerSourcePrivate::GetFallbackSoundClass(InActiveSound, InWaveInstance),
+				Audio::GetModulationParameter("Pitch"),
+				[](const USoundClass& InSoundClass) { return &InSoundClass.Properties.ModulationSettings.PitchModulationDestination; }
+			);
+		}
+
+		float GetRoutedPitch(const FWaveInstance& InWaveInstance, const USoundWave& InWaveData, const FActiveSound& InActiveSound)
+		{
+			const EModulationRouting& ActiveSoundRouting = InActiveSound.ModulationRouting.PitchRouting;
+			const FSoundModulationDestinationSettings& ActiveSoundSettings = InActiveSound.ModulationRouting.PitchModulationDestination;
+
+			const EModulationRouting& WaveRouting = InWaveData.ModulationSettings.PitchRouting;
+			const FSoundModulationDestinationSettings& WaveSettings = InWaveData.ModulationSettings.PitchModulationDestination;
+
+			return GetRoutedDestinationValue(
+				ActiveSoundRouting,
+				ActiveSoundSettings,
+				WaveRouting,
+				WaveSettings,
+				MixerSourcePrivate::GetFallbackSoundClass(InActiveSound, InWaveInstance),
+				Audio::GetModulationParameter("Pitch"),
+				[](const USoundClass& InSoundClass) { return &InSoundClass.Properties.ModulationSettings.PitchModulationDestination; }
+			);
+		}
+
+		FSoundModulationDestinationSettings InitRoutedHighpassModulation(const FWaveInstance& InWaveInstance, const USoundWave& InWaveData, const FActiveSound& InActiveSound)
+		{
+			const EModulationRouting& ActiveSoundRouting = InActiveSound.ModulationRouting.HighpassRouting;
+			const FSoundModulationDestinationSettings& ActiveSoundSettings = InActiveSound.ModulationRouting.HighpassModulationDestination;
+
+			const EModulationRouting& WaveRouting = InWaveData.ModulationSettings.HighpassRouting;
+			const FSoundModulationDestinationSettings& WaveSettings = InWaveData.ModulationSettings.HighpassModulationDestination;
+
+			return InitRoutedDestinationSettings(
+				ActiveSoundRouting,
+				ActiveSoundSettings,
+				WaveRouting,
+				WaveSettings,
+				MixerSourcePrivate::GetFallbackSoundClass(InActiveSound, InWaveInstance),
+				Audio::GetModulationParameter("HPFCutoffFrequency"),
+				[](const USoundClass& InSoundClass) { return &InSoundClass.Properties.ModulationSettings.HighpassModulationDestination; }
+			);
+		}
+
+		float GetRoutedHighpass(const FWaveInstance& InWaveInstance, const USoundWave& InWaveData, const FActiveSound& InActiveSound)
+		{
+			const EModulationRouting& ActiveSoundRouting = InActiveSound.ModulationRouting.HighpassRouting;
+			const FSoundModulationDestinationSettings& ActiveSoundSettings = InActiveSound.ModulationRouting.HighpassModulationDestination;
+
+			const EModulationRouting& WaveRouting = InWaveData.ModulationSettings.HighpassRouting;
+			const FSoundModulationDestinationSettings& WaveSettings = InWaveData.ModulationSettings.HighpassModulationDestination;
+
+			return GetRoutedDestinationValue(
+				ActiveSoundRouting,
+				ActiveSoundSettings,
+				WaveRouting,
+				WaveSettings,
+				MixerSourcePrivate::GetFallbackSoundClass(InActiveSound, InWaveInstance),
+				Audio::GetModulationParameter("HPFCutoffFrequency"),
+				[](const USoundClass& InSoundClass) { return &InSoundClass.Properties.ModulationSettings.HighpassModulationDestination; }
+			);
+		}
+
+		FSoundModulationDestinationSettings InitRoutedLowpassModulation(const FWaveInstance& InWaveInstance, const USoundWave& InWaveData, const FActiveSound& InActiveSound)
+		{
+			const EModulationRouting& ActiveSoundRouting = InActiveSound.ModulationRouting.LowpassRouting;
+			const FSoundModulationDestinationSettings& ActiveSoundSettings = InActiveSound.ModulationRouting.LowpassModulationDestination;
+
+			const EModulationRouting& WaveRouting = InWaveData.ModulationSettings.LowpassRouting;
+			const FSoundModulationDestinationSettings& WaveSettings = InWaveData.ModulationSettings.LowpassModulationDestination;
+
+			return InitRoutedDestinationSettings(
+				ActiveSoundRouting,
+				ActiveSoundSettings,
+				WaveRouting,
+				WaveSettings,
+				MixerSourcePrivate::GetFallbackSoundClass(InActiveSound, InWaveInstance),
+				Audio::GetModulationParameter("LPFCutoffFrequency"),
+				[](const USoundClass& InSoundClass) { return &InSoundClass.Properties.ModulationSettings.LowpassModulationDestination; }
+			);
+		}
+
+		float GetRoutedLowpass(const FWaveInstance& InWaveInstance, const USoundWave& InWaveData, const FActiveSound& InActiveSound)
+		{
+			const EModulationRouting& ActiveSoundRouting = InActiveSound.ModulationRouting.LowpassRouting;
+			const FSoundModulationDestinationSettings& ActiveSoundSettings = InActiveSound.ModulationRouting.LowpassModulationDestination;
+
+			const EModulationRouting& WaveRouting = InWaveData.ModulationSettings.LowpassRouting;
+			const FSoundModulationDestinationSettings& WaveSettings = InWaveData.ModulationSettings.LowpassModulationDestination;
+
+			return GetRoutedDestinationValue(
+				ActiveSoundRouting,
+				ActiveSoundSettings,
+				WaveRouting,
+				WaveSettings,
+				MixerSourcePrivate::GetFallbackSoundClass(InActiveSound, InWaveInstance),
+				Audio::GetModulationParameter("LPFCutoffFrequency"),
+				[](const USoundClass& InSoundClass) { return &InSoundClass.Properties.ModulationSettings.LowpassModulationDestination; }
+			);
 		}
 
-		FSoundModulationDefaultSettings GetRoutedModulation(const FWaveInstance& InWaveInstance, const USoundWave& InWaveData, FActiveSound* InActiveSound)
+		FSoundModulationDefaultSettings InitRoutedModulation(const FWaveInstance& InWaveInstance, const USoundWave& InWaveData, FActiveSound* InActiveSound)
 		{
 			FSoundModulationDefaultSettings Settings;
 			if (InActiveSound)
 			{
-				Settings.VolumeModulationDestination = GetRoutedVolumeModulation(InWaveInstance, InWaveData, InActiveSound);
-				Settings.PitchModulationDestination = GetRoutedPitchModulation(InWaveInstance, InWaveData, InActiveSound);
-				Settings.HighpassModulationDestination = GetRoutedHighpassModulation(InWaveInstance, InWaveData, InActiveSound);
-				Settings.LowpassModulationDestination = GetRoutedLowpassModulation(InWaveInstance, InWaveData, InActiveSound);
+				Settings.VolumeModulationDestination = InitRoutedVolumeModulation(InWaveInstance, InWaveData, *InActiveSound);
+				Settings.PitchModulationDestination = InitRoutedPitchModulation(InWaveInstance, InWaveData, *InActiveSound);
+				Settings.HighpassModulationDestination = InitRoutedHighpassModulation(InWaveInstance, InWaveData, *InActiveSound);
+				Settings.LowpassModulationDestination = InitRoutedLowpassModulation(InWaveInstance, InWaveData, *InActiveSound);
 			}
 
 			return Settings;
@@ -385,7 +570,7 @@ namespace Audio
 			InitParams.bIsSoundfield = WaveInstance->bIsAmbisonics && (WaveData->NumChannels == 4);
 
 			FActiveSound* ActiveSound = WaveInstance->ActiveSound;
-			InitParams.ModulationSettings = ModulationUtils::GetRoutedModulation(*WaveInstance, *WaveData, ActiveSound);
+			InitParams.ModulationSettings = ModulationUtils::InitRoutedModulation(*WaveInstance, *WaveData, ActiveSound);
 
 			// Copy quantization request data
 			if (WaveInstance->QuantizedRequestData)
@@ -1290,8 +1475,8 @@ namespace Audio
 
 		USoundWave* WaveData = WaveInstance->WaveData;
 		check(WaveData);
-		const FSoundModulationDestinationSettings& PitchSettings = ModulationUtils::GetRoutedPitchModulation(*WaveInstance, *WaveData, ActiveSound);
-		MixerSourceVoice->SetModPitch(PitchSettings.Value);
+		const float ModPitchBase = ModulationUtils::GetRoutedPitch(*WaveInstance, *WaveData, *ActiveSound);
+		MixerSourceVoice->SetModPitch(ModPitchBase);
 	}
 
 	void FMixerSource::UpdateVolume()
@@ -1318,8 +1503,8 @@ namespace Audio
 
 			USoundWave* WaveData = WaveInstance->WaveData;
 			check(WaveData);
-			const FSoundModulationDestinationSettings& VolumeSettings = ModulationUtils::GetRoutedVolumeModulation(*WaveInstance, *WaveData, ActiveSound);
-			MixerSourceVoice->SetModVolume(VolumeSettings.Value);
+			const float ModVolumeBase = ModulationUtils::GetRoutedVolume(*WaveInstance, *WaveData, *ActiveSound);
+			MixerSourceVoice->SetModVolume(ModVolumeBase);
 		}
 		MixerSourceVoice->SetVolume(CurrentVolume);
 	}
@@ -1360,11 +1545,11 @@ namespace Audio
 		USoundWave* WaveData = WaveInstance->WaveData;
 		check(WaveData);
 
-		const FSoundModulationDestinationSettings& HighpassSettings = ModulationUtils::GetRoutedHighpassModulation(*WaveInstance, *WaveData, ActiveSound);
-		MixerSourceVoice->SetModHPFFrequency(HighpassSettings.Value);
+		float ModHighpassBase = ModulationUtils::GetRoutedHighpass(*WaveInstance, *WaveData, *ActiveSound);
+		MixerSourceVoice->SetModHPFFrequency(ModHighpassBase);
 
-		const FSoundModulationDestinationSettings& LowpassSettings = ModulationUtils::GetRoutedLowpassModulation(*WaveInstance, *WaveData, ActiveSound);
-		MixerSourceVoice->SetModLPFFrequency(LowpassSettings.Value);
+		float ModLowpassBase = ModulationUtils::GetRoutedLowpass(*WaveInstance, *WaveData, *ActiveSound);
+		MixerSourceVoice->SetModLPFFrequency(ModLowpassBase);
 
 		// If reverb is applied, figure out how of the source to "send" to the reverb.
 		if (WaveInstance->bReverb)
diff --git a/Engine/Source/Runtime/AudioMixer/Private/AudioMixerSourceManager.cpp b/Engine/Source/Runtime/AudioMixer/Private/AudioMixerSourceManager.cpp
--- a/Engine/Source/Runtime/AudioMixer/Private/AudioMixerSourceManager.cpp
+++ b/Engine/Source/Runtime/AudioMixer/Private/AudioMixerSourceManager.cpp
@@ -746,19 +746,19 @@ namespace Audio
 
 		FModulationDestination VolumeMod;
 		VolumeMod.Init(MixerDevice->DeviceID, FName("Volume"), false /* bInIsBuffered */, true /* bInValueLinear */);
-		VolumeMod.UpdateModulator(InitParams.ModulationSettings.VolumeModulationDestination.Modulator);
+		VolumeMod.UpdateModulators(InitParams.ModulationSettings.VolumeModulationDestination.Modulators);
 
 		FModulationDestination PitchMod;
 		PitchMod.Init(MixerDevice->DeviceID, FName("Pitch"), false /* bInIsBuffered */);
-		PitchMod.UpdateModulator(InitParams.ModulationSettings.PitchModulationDestination.Modulator);
+		PitchMod.UpdateModulators(InitParams.ModulationSettings.PitchModulationDestination.Modulators);
 
 		FModulationDestination HighpassMod;
 		HighpassMod.Init(MixerDevice->DeviceID, FName("HPFCutoffFrequency"), false /* bInIsBuffered */);
-		HighpassMod.UpdateModulator(InitParams.ModulationSettings.HighpassModulationDestination.Modulator);
+		HighpassMod.UpdateModulators(InitParams.ModulationSettings.HighpassModulationDestination.Modulators);
 
 		FModulationDestination LowpassMod;
 		LowpassMod.Init(MixerDevice->DeviceID, FName("LPFCutoffFrequency"), false /* bInIsBuffered */);
-		LowpassMod.UpdateModulator(InitParams.ModulationSettings.LowpassModulationDestination.Modulator);
+		LowpassMod.UpdateModulators(InitParams.ModulationSettings.LowpassModulationDestination.Modulators);
 
 		AudioMixerThreadCommand([
 			this,
diff --git a/Engine/Source/Runtime/AudioMixer/Private/AudioMixerSubmix.cpp b/Engine/Source/Runtime/AudioMixer/Private/AudioMixerSubmix.cpp
--- a/Engine/Source/Runtime/AudioMixer/Private/AudioMixerSubmix.cpp
+++ b/Engine/Source/Runtime/AudioMixer/Private/AudioMixerSubmix.cpp
@@ -274,14 +274,14 @@ namespace Audio
 
 			if (MixerDevice->IsModulationPluginEnabled() && MixerDevice->ModulationInterface.IsValid())
 			{
-				USoundModulatorBase* VolumeModulator = SoundSubmix->OutputVolumeModulation.Modulator;
-				USoundModulatorBase* WetLevelModulator = SoundSubmix->WetLevelModulation.Modulator;
-				USoundModulatorBase* DryLevelModulator = SoundSubmix->DryLevelModulation.Modulator;
+				TSet<TObjectPtr<USoundModulatorBase>> VolumeModulator = SoundSubmix->OutputVolumeModulation.Modulators;
+				TSet<TObjectPtr<USoundModulatorBase>> WetLevelModulator = SoundSubmix->WetLevelModulation.Modulators;
+				TSet<TObjectPtr<USoundModulatorBase>> DryLevelModulator = SoundSubmix->DryLevelModulation.Modulators;
 
 				// Queue this up to happen after submix init, when the mixer device has finished being added to the device manager
-				SubmixCommand([this, VolumeModulator, WetLevelModulator, DryLevelModulator]()
+				SubmixCommand([this, VolMod = MoveTemp(VolumeModulator), WetMod = MoveTemp(WetLevelModulator), DryMod = MoveTemp(DryLevelModulator)]() mutable
 				{
-					UpdateModulationSettings(VolumeModulator, WetLevelModulator, DryLevelModulator);
+					UpdateModulationSettings(VolMod, WetMod, DryMod);
 				});
 			}
 			if (SoundSubmix->AudioLinkSettings)
@@ -2310,11 +2310,11 @@ namespace Audio
 		WetLevelModifier = FMath::Clamp(InWetLevel, 0.0f, 1.0f);
 	}
 
-	void FMixerSubmix::UpdateModulationSettings(USoundModulatorBase* InOutputModulator, USoundModulatorBase* InWetLevelModulator, USoundModulatorBase* InDryLevelModulator)
+	void FMixerSubmix::UpdateModulationSettings(const TSet<TObjectPtr<USoundModulatorBase>>& InOutputModulators, const TSet<TObjectPtr<USoundModulatorBase>>& InWetLevelModulators, const TSet<TObjectPtr<USoundModulatorBase>>& InDryLevelModulators)
 	{
-		VolumeMod.UpdateModulator(InOutputModulator);
-		WetLevelMod.UpdateModulator(InWetLevelModulator);
-		DryLevelMod.UpdateModulator(InDryLevelModulator);
+		VolumeMod.UpdateModulators(InOutputModulators);
+		WetLevelMod.UpdateModulators(InWetLevelModulators);
+		DryLevelMod.UpdateModulators(InDryLevelModulators);
 	}
 
 	void FMixerSubmix::SetModulationBaseLevels(float InVolumeModBaseDb, float InWetModBaseDb, float InDryModBaseDb)
diff --git a/Engine/Source/Runtime/AudioMixer/Public/AudioMixerDevice.h b/Engine/Source/Runtime/AudioMixer/Public/AudioMixerDevice.h
--- a/Engine/Source/Runtime/AudioMixer/Public/AudioMixerDevice.h
+++ b/Engine/Source/Runtime/AudioMixer/Public/AudioMixerDevice.h
@@ -144,7 +144,7 @@ namespace Audio
 		virtual void SetSubmixAutoDisableTime(USoundSubmix* InSoundSubmix, float InDisableTime) override;
 
 		// Submix Modulation Settings
-		virtual void UpdateSubmixModulationSettings(USoundSubmix* InSoundSubmix, USoundModulatorBase* InOutputModulation, USoundModulatorBase* InWetLevelModulation, USoundModulatorBase* InDryLevelModulation) override;
+		virtual void UpdateSubmixModulationSettings(USoundSubmix* InSoundSubmix, const TSet<TObjectPtr<USoundModulatorBase>>& InOutputModulation, const TSet<TObjectPtr<USoundModulatorBase>>& InWetLevelModulation, const TSet<TObjectPtr<USoundModulatorBase>>& InDryLevelModulation) override;
 		virtual void SetSubmixModulationBaseLevels(USoundSubmix* InSoundSubmix, float InVolumeModBase, float InWetModBase, float InDryModBase) override;
 
 		// Submix effect chain override settings
diff --git a/Engine/Source/Runtime/AudioMixer/Public/AudioMixerSubmix.h b/Engine/Source/Runtime/AudioMixer/Public/AudioMixerSubmix.h
--- a/Engine/Source/Runtime/AudioMixer/Public/AudioMixerSubmix.h
+++ b/Engine/Source/Runtime/AudioMixer/Public/AudioMixerSubmix.h
@@ -142,7 +142,7 @@ namespace Audio
 		void SetWetLevel(float InWetLevel);
 
 		// Update modulation settings of the submix
-		void UpdateModulationSettings(USoundModulatorBase* InOutputModulator, USoundModulatorBase* InWetLevelModulator, USoundModulatorBase* InDryLevelModulator);
+		void UpdateModulationSettings(const TSet<TObjectPtr<USoundModulatorBase>>& InOutputModulators, const TSet<TObjectPtr<USoundModulatorBase>>& InWetLevelModulators, const TSet<TObjectPtr<USoundModulatorBase>>& InDryLevelModulators);
 
 		// Update modulation settings of the submix with Decibel values
 		void SetModulationBaseLevels(float InVolumeModBaseDb, float InWetModeBaseDb, float InDryModBaseDb);
diff --git a/Engine/Source/Runtime/Engine/Classes/Sound/SoundModulationDestination.h b/Engine/Source/Runtime/Engine/Classes/Sound/SoundModulationDestination.h
--- a/Engine/Source/Runtime/Engine/Classes/Sound/SoundModulationDestination.h
+++ b/Engine/Source/Runtime/Engine/Classes/Sound/SoundModulationDestination.h
@@ -3,6 +3,7 @@
 
 #include "Audio.h"
 #include "Containers/ArrayView.h"
+#include "Containers/Set.h"
 #include "DSP/BufferVectorOperations.h"
 #include "HAL/CriticalSection.h"
 #include "IAudioModulation.h"
@@ -16,6 +17,11 @@
 // Forward Declarations
 class USoundModulatorBase;
 
+namespace Audio
+{
+	struct FModulationParameter;
+}
+
 
 UENUM(BlueprintType)
 enum class EModulationRouting : uint8
@@ -27,7 +33,10 @@ enum class EModulationRouting : uint8
 	Inherit,
 
 	/* Ignores inherited settings and uses modulation settings on this object */
-	Override
+	Override,
+
+	/* Performs set union on local modulation sources with those inherited (AudioComponent inherits from Sound, Sound inherits from SoundClass) */
+	Union
 };
 
 /** Parameter destination settings allowing modulation control override for parameter destinations opting in to the Modulation System. */
@@ -41,14 +50,23 @@ struct ENGINE_API FSoundModulationDestinationSettings
 	float Value = 1.0f;
 
 #if WITH_EDITORONLY_DATA
-	/** Base value of parameter */
+	/** Whether or not to enable modulation */
 	UPROPERTY(EditAnywhere, Category = Modulation, meta = (DisplayName = "Modulate"))
 	bool bEnableModulation = false;
+
+	/** Deprecated in favor of Modulators set. */
+	UPROPERTY(meta = (DeprecatedProperty))
+	TObjectPtr<USoundModulatorBase> Modulator = nullptr;
 #endif // WITH_EDITORONLY_DATA
 
-	/** Modulation source, which provides value to mix with base value. */
+	/** Set of modulation sources, which provides values to mix with base value. */
 	UPROPERTY(EditAnywhere, BlueprintReadWrite, Category = Modulation)
-	TObjectPtr<USoundModulatorBase> Modulator = nullptr;
+	TSet<TObjectPtr<USoundModulatorBase>> Modulators;
+
+#if WITH_EDITORONLY_DATA
+	/** Versioning utility function to upgrade single modulator field to set of modulators */
+	void VersionModulators();
+#endif // WITH_EDITORONLY_DATA
 };
 
 /** Default parameter destination settings for source audio object. */
@@ -74,6 +92,10 @@ struct ENGINE_API FSoundModulationDefaultSettings
 	/** Lowpass modulation */
 	UPROPERTY(EditAnywhere, BlueprintReadWrite, Category = "Modulation", meta = (DisplayName = "Lowpass", AudioParam = "LPFCutoffFrequency"))
 	FSoundModulationDestinationSettings LowpassModulationDestination;
+
+#if WITH_EDITORONLY_DATA
+	void VersionModulators();
+#endif // WITH_EDITORONLY_DATA
 };
 
 /** Default parameter destination settings for source audio object. */
@@ -137,12 +159,18 @@ namespace Audio
 			 */
 			bool ProcessControl(float InValueUnitBase, int32 InNumSamples = 0);
 
+
+			UE_DEPRECATED(5.1, "Deprecated in favor of supporting multiple modulators per destination. Use 'UpdateModulators' instead.")
 			void UpdateModulator(const USoundModulatorBase* InModulator);
 
-		private:
-			void ResetHandle();
-			void SetHandle(const FModulatorHandle& InHandle);
-			void SetHandle(FModulatorHandle&& InHandle);
+			void UpdateModulators(const TSet<TObjectPtr<USoundModulatorBase>>& InModulators);
+			void UpdateModulators(const TSet<USoundModulatorBase*>& InModulators);
+			void UpdateModulators(const TSet<const USoundModulatorBase*>& InModulators);
+	
+	private:
+			void UpdateModulatorsInternal(TArray<TUniquePtr<Audio::IModulatorSettings>>&& ProxySettings);
+
+			void ResetHandles();
 
 			FDeviceId DeviceId = INDEX_NONE;
 
@@ -153,9 +181,10 @@ namespace Audio
 			bool bHasProcessed = false;
 
 			FAlignedFloatBuffer OutputBuffer;
-			FModulatorHandle Handle;
 
-			FName ParameterName;
+			TSet<FModulatorHandle> Handles;
+
+			FModulationParameter Parameter;
 
 			mutable FCriticalSection HandleCritSection;
 
diff --git a/Engine/Source/Runtime/Engine/Private/Components/AudioComponent.cpp b/Engine/Source/Runtime/Engine/Private/Components/AudioComponent.cpp
--- a/Engine/Source/Runtime/Engine/Private/Components/AudioComponent.cpp
+++ b/Engine/Source/Runtime/Engine/Private/Components/AudioComponent.cpp
@@ -186,6 +186,8 @@ void UAudioComponent::Serialize(FArchive& Ar)
 			ConcurrencySet.Add(ConcurrencySettings_DEPRECATED);
 			ConcurrencySettings_DEPRECATED = nullptr;
 		}
+
+		ModulationRouting.VersionModulators();
 	}
 #endif // WITH_EDITORONLY_DATA
 }
diff --git a/Engine/Source/Runtime/Engine/Private/SoundClass.cpp b/Engine/Source/Runtime/Engine/Private/SoundClass.cpp
--- a/Engine/Source/Runtime/Engine/Private/SoundClass.cpp
+++ b/Engine/Source/Runtime/Engine/Private/SoundClass.cpp
@@ -283,6 +283,13 @@ void USoundClass::Serialize( FArchive& Ar )
 		TMap<USoundClass*, FSoundClassEditorData>	EditorData_DEPRECATED;
 		Ar << EditorData_DEPRECATED;
 	}
+
+#if WITH_EDITORONLY_DATA
+	if (Ar.IsLoading())
+	{
+		Properties.ModulationSettings.VersionModulators();
+	}
+#endif // WITH_EDITORONLY_DATA
 }
 
 void USoundClass::BeginDestroy()
diff --git a/Engine/Source/Runtime/Engine/Private/SoundModulationDestination.cpp b/Engine/Source/Runtime/Engine/Private/SoundModulationDestination.cpp
--- a/Engine/Source/Runtime/Engine/Private/SoundModulationDestination.cpp
+++ b/Engine/Source/Runtime/Engine/Private/SoundModulationDestination.cpp
@@ -2,11 +2,15 @@
 
 #include "Sound/SoundModulationDestination.h"
 
+#include "Algo/NoneOf.h"
+#include "Algo/Transform.h"
 #include "Async/Async.h"
 #include "AudioDevice.h"
 #include "DSP/FloatArrayMath.h"
+#include "IAudioModulation.h"
 #include "Math/TransformCalculus.h"
 #include "UObject/Object.h"
+#include "GenericPlatform/GenericPlatformCompilerPreSetup.h"
 
 
 FSoundModulationDefaultSettings::FSoundModulationDefaultSettings()
@@ -17,11 +21,36 @@ FSoundModulationDefaultSettings::FSoundModulationDefaultSettings()
 	LowpassModulationDestination.Value = MAX_FILTER_FREQUENCY;
 }
 
+#if WITH_EDITORONLY_DATA
+void FSoundModulationDefaultSettings::VersionModulators()
+{
+PRAGMA_DISABLE_DEPRECATION_WARNINGS
+	VolumeModulationDestination.VersionModulators();
+	PitchModulationDestination.VersionModulators();
+	HighpassModulationDestination.VersionModulators();
+	LowpassModulationDestination.VersionModulators();
+PRAGMA_ENABLE_DEPRECATION_WARNINGS
+}
+
+void FSoundModulationDestinationSettings::VersionModulators()
+{
+PRAGMA_DISABLE_DEPRECATION_WARNINGS
+	if (Modulator)
+	{
+		Modulators.Add(Modulator);
+		Modulator = nullptr;
+	}
+PRAGMA_ENABLE_DEPRECATION_WARNINGS
+}
+#endif // WITH_EDITORONLY_DATA
+
+
 FSoundModulationDefaultRoutingSettings::FSoundModulationDefaultRoutingSettings()
 	: FSoundModulationDefaultSettings()
 {
 }
 
+
 namespace Audio
 {
 	FModulationDestination::FModulationDestination(const FModulationDestination& InModulationDestination)
@@ -30,10 +59,10 @@ namespace Audio
 		, bIsBuffered(InModulationDestination.bIsBuffered)
 		, bValueNormalized(InModulationDestination.bValueNormalized)
 		, OutputBuffer(InModulationDestination.OutputBuffer)
-		, ParameterName(InModulationDestination.ParameterName)
+		, Parameter(InModulationDestination.Parameter)
 	{
 		FScopeLock OtherLock(&InModulationDestination.HandleCritSection);
-		Handle = InModulationDestination.Handle;
+		Handles = InModulationDestination.Handles;
 	}
 
 	FModulationDestination::FModulationDestination(FModulationDestination&& InModulationDestination)
@@ -42,10 +71,10 @@ namespace Audio
 		, bIsBuffered(MoveTemp(InModulationDestination.bIsBuffered))
 		, bValueNormalized(MoveTemp(InModulationDestination.bValueNormalized))
 		, OutputBuffer(MoveTemp(InModulationDestination.OutputBuffer))
-		, ParameterName(MoveTemp(InModulationDestination.ParameterName))
+		, Parameter(MoveTemp(InModulationDestination.Parameter))
 	{
 		FScopeLock OtherLock(&InModulationDestination.HandleCritSection);
-		Handle = MoveTemp(InModulationDestination.Handle);
+		Handles = MoveTemp(InModulationDestination.Handles);
 	}
 
 	FModulationDestination& FModulationDestination::operator=(const FModulationDestination& InModulationDestination)
@@ -58,11 +87,11 @@ namespace Audio
 
 		{
 			FScopeLock OtherLock(&InModulationDestination.HandleCritSection);
-			SetHandle(InModulationDestination.Handle);
+			FScopeLock Lock(&HandleCritSection);
+			Handles = InModulationDestination.Handles;
 		}
-			
-		ParameterName = InModulationDestination.ParameterName;
 
+		Parameter = InModulationDestination.Parameter;
 
 		return *this;
 	}
@@ -78,21 +107,21 @@ namespace Audio
 		{
 			FScopeLock OtherLock(&InModulationDestination.HandleCritSection);
 			FScopeLock Lock(&HandleCritSection);
-			Handle = MoveTemp(InModulationDestination.Handle);
+			Handles = MoveTemp(InModulationDestination.Handles);
 		}
 
-
-		ParameterName = MoveTemp(InModulationDestination.ParameterName);
+		Parameter = MoveTemp(InModulationDestination.Parameter);
 
 		return *this;
 	}
 
-	void FModulationDestination::ResetHandle()
+	void FModulationDestination::ResetHandles()
 	{
-		Audio::FModulationParameter Parameter = Audio::GetModulationParameter(ParameterName);
+		Audio::FModulationParameter ParameterCopy = Parameter;
 
 		FScopeLock Lock(&HandleCritSection);
-		Handle = FModulatorHandle(MoveTemp(Parameter));
+		Handles.Reset();
+		Handles.Add(FModulatorHandle { MoveTemp(ParameterCopy) });
 	}
 
 	void FModulationDestination::Init(FDeviceId InDeviceId, bool bInIsBuffered, bool bInValueNormalized)
@@ -107,15 +136,15 @@ namespace Audio
 		bValueNormalized = bInValueNormalized;
 
 		OutputBuffer.Reset();
-		ParameterName = InParameterName;
+		Parameter = Audio::GetModulationParameter(InParameterName);
 
-		ResetHandle();
+		ResetHandles();
 	}
 
 	bool FModulationDestination::IsActive()
 	{
 		FScopeLock Lock(&HandleCritSection);
-		return Handle.IsValid();
+		return Algo::NoneOf(Handles, [](const FModulatorHandle& Handle) { return Handle.IsValid(); });
 	}
 
 	bool FModulationDestination::ProcessControl(float InValueUnitBase, int32 InNumSamples)
@@ -123,36 +152,40 @@ namespace Audio
 		bHasProcessed = true;
 		float LastTarget = ValueTarget;
 
-		FScopeLock Lock(&HandleCritSection);
+		float NewTargetNormalized = Parameter.DefaultValue;
+		if (Parameter.bRequiresConversion)
 		{
-			const FModulationParameter& Parameter = Handle.GetParameter();
-			float NewTargetNormalized = Parameter.DefaultValue;
-			if (Parameter.bRequiresConversion)
-			{
-				Parameter.NormalizedFunction(NewTargetNormalized);
-			}
+			Parameter.NormalizedFunction(NewTargetNormalized);
+		}
 
-			if (Handle.IsValid())
+		FScopeLock Lock(&HandleCritSection);
+		{
+			for (const FModulatorHandle& Handle : Handles)
 			{
-				Handle.GetValue(NewTargetNormalized);
+				if (Handle.IsValid())
+				{
+					float NewHandleValue = 1.0f;
+					Handle.GetValue(NewHandleValue);
+					Parameter.MixFunction(NewTargetNormalized, NewHandleValue);
+				}
 			}
+		}
 
-			// Convert base to linear space
-			float InValueBaseNormalized = InValueUnitBase;
-			if (Parameter.bRequiresConversion)
-			{
-				Parameter.NormalizedFunction(InValueBaseNormalized);
-			}
+		// Convert base to linear space
+		float InValueBaseNormalized = InValueUnitBase;
+		if (Parameter.bRequiresConversion)
+		{
+			Parameter.NormalizedFunction(InValueBaseNormalized);
+		}
 
-			// Mix in base value
-			Parameter.MixFunction(NewTargetNormalized, InValueBaseNormalized);
-			ValueTarget = NewTargetNormalized;
+		// Mix in base value
+		Parameter.MixFunction(NewTargetNormalized, InValueBaseNormalized);
+		ValueTarget = NewTargetNormalized;
 
-			// Convert target to unit space if required
-			if (Parameter.bRequiresConversion && !bValueNormalized)
-			{
-				Parameter.UnitFunction(ValueTarget);
-			}
+		// Convert target to unit space if required
+		if (Parameter.bRequiresConversion && !bValueNormalized)
+		{
+			Parameter.UnitFunction(ValueTarget);
 		}
 
 		if (bIsBuffered)
@@ -201,27 +234,53 @@ namespace Audio
 		return !FMath::IsNearlyEqual(LastTarget, ValueTarget);
 	}
 
-	void FModulationDestination::SetHandle(const FModulatorHandle& InHandle)
+	void FModulationDestination::UpdateModulator(const USoundModulatorBase* InModulator)
 	{
-		FScopeLock Lock(&HandleCritSection);
-		Handle = InHandle;
+		UpdateModulators({ InModulator });
 	}
 
-	void FModulationDestination::SetHandle(FModulatorHandle&& InHandle)
+	void FModulationDestination::UpdateModulators(const TSet<TObjectPtr<USoundModulatorBase>>& InModulators)
 	{
-		FScopeLock Lock(&HandleCritSection);
-		Handle = MoveTemp(InHandle);
+		TArray<TUniquePtr<Audio::IModulatorSettings>> ProxySettings;
+		Algo::TransformIf(
+			InModulators,
+			ProxySettings,
+			[](const USoundModulatorBase* Mod) { return Mod != nullptr; },
+			[](const USoundModulatorBase* Mod) { return Mod->CreateProxySettings(); }
+		);
+
+		UpdateModulatorsInternal(MoveTemp(ProxySettings));
 	}
 
-	void FModulationDestination::UpdateModulator(const USoundModulatorBase* InModulator)
+	void FModulationDestination::UpdateModulators(const TSet<USoundModulatorBase*>& InModulators)
 	{
-		if (!InModulator)
-		{
-			return;
-		}
+		TArray<TUniquePtr<Audio::IModulatorSettings>> ProxySettings;
+		Algo::TransformIf(
+			InModulators,
+			ProxySettings,
+			[](const USoundModulatorBase* Mod) { return Mod != nullptr; },
+			[](const USoundModulatorBase* Mod) { return Mod->CreateProxySettings(); }
+		);
+
+		UpdateModulatorsInternal(MoveTemp(ProxySettings));
+	}
 
-		Audio::FModulationParameter OutputParam = Audio::GetModulationParameter(ParameterName);
-		auto UpdateHandleLambda = [this, ModSettings = InModulator->CreateProxySettings(), Parameter = MoveTemp(OutputParam)]() mutable
+	void FModulationDestination::UpdateModulators(const TSet<const USoundModulatorBase*>& InModulators)
+	{
+		TArray<TUniquePtr<Audio::IModulatorSettings>> ProxySettings;
+		Algo::TransformIf(
+			InModulators,
+			ProxySettings,
+			[](const USoundModulatorBase* Mod) { return Mod != nullptr; },
+			[](const USoundModulatorBase* Mod) { return Mod->CreateProxySettings(); }
+		);
+
+		UpdateModulatorsInternal(MoveTemp(ProxySettings));
+	}
+
+	void FModulationDestination::UpdateModulatorsInternal(TArray<TUniquePtr<Audio::IModulatorSettings>>&& ProxySettings)
+	{
+		auto UpdateHandleLambda = [this, ModSettings = MoveTemp(ProxySettings)]() mutable
 		{
 			if (FAudioDevice* AudioDevice = FAudioDeviceManager::Get()->GetAudioDeviceRaw(DeviceId))
 			{
@@ -229,14 +288,21 @@ namespace Audio
 				{
 					if (IAudioModulationManager* Modulation = AudioDevice->ModulationInterface.Get())
 					{
-						check(ModSettings.IsValid());
-						SetHandle(FModulatorHandle(*Modulation, *ModSettings.Get(), MoveTemp(Parameter)));
+						TSet<FModulatorHandle> NewHandles;
+						for (TUniquePtr<Audio::IModulatorSettings>& ModSetting : ModSettings)
+						{
+							Audio::FModulationParameter HandleParam = Parameter;
+							NewHandles.Add(FModulatorHandle{ *Modulation, *ModSetting.Get(), MoveTemp(HandleParam) });
+						}
+
+						FScopeLock Lock(&HandleCritSection);
+						Handles = MoveTemp(NewHandles);
 					}
 					return;
 				}
 			}
 
-			ResetHandle();
+			ResetHandles();
 		};
 
 		FAudioThread::RunCommandOnAudioThread(MoveTemp(UpdateHandleLambda));
diff --git a/Engine/Source/Runtime/Engine/Private/SoundSubmix.cpp b/Engine/Source/Runtime/Engine/Private/SoundSubmix.cpp
--- a/Engine/Source/Runtime/Engine/Private/SoundSubmix.cpp
+++ b/Engine/Source/Runtime/Engine/Private/SoundSubmix.cpp
@@ -117,6 +117,10 @@ void USoundSubmix::Serialize(FArchive& Ar)
 		{
 			DryLevelModulation.Value = Audio::ConvertToDecibels(DryLevelModulation.Value);
 		}
+
+		OutputVolumeModulation.VersionModulators();
+		WetLevelModulation.VersionModulators();
+		DryLevelModulation.VersionModulators();
 	}
 #endif // WITH_EDITORONLY_DATA
 }
@@ -491,13 +495,13 @@ void USoundSubmix::PostEditChangeProperty(struct FPropertyChangedEvent& Property
 			{
 				USoundSubmix* SoundSubmix = this;
 
-				USoundModulatorBase* NewVolumeMod = OutputVolumeModulation.Modulator;
-				USoundModulatorBase* NewWetLevelMod = WetLevelModulation.Modulator;
-				USoundModulatorBase* NewDryLevelMod = DryLevelModulation.Modulator;
+				TSet<TObjectPtr<USoundModulatorBase>> NewVolumeMod = OutputVolumeModulation.Modulators;
+				TSet<TObjectPtr<USoundModulatorBase>> NewWetLevelMod = WetLevelModulation.Modulators;
+				TSet<TObjectPtr<USoundModulatorBase>> NewDryLevelMod = DryLevelModulation.Modulators;
 
-				AudioDeviceManager->IterateOverAllDevices([SoundSubmix, NewVolumeMod, NewWetLevelMod, NewDryLevelMod](Audio::FDeviceId Id, FAudioDevice* Device)
+				AudioDeviceManager->IterateOverAllDevices([SoundSubmix, VolMod = MoveTemp(NewVolumeMod), WetMod = MoveTemp(NewWetLevelMod), DryMod = MoveTemp(NewDryLevelMod)](Audio::FDeviceId Id, FAudioDevice* Device) mutable
 				{
-					Device->UpdateSubmixModulationSettings(SoundSubmix, NewVolumeMod, NewWetLevelMod, NewDryLevelMod);
+					Device->UpdateSubmixModulationSettings(SoundSubmix, VolMod, WetMod, DryMod);
 				});
 
 				float NewVolumeModBase = OutputVolumeModulation.Value;
diff --git a/Engine/Source/Runtime/Engine/Private/SoundWave.cpp b/Engine/Source/Runtime/Engine/Private/SoundWave.cpp
--- a/Engine/Source/Runtime/Engine/Private/SoundWave.cpp
+++ b/Engine/Source/Runtime/Engine/Private/SoundWave.cpp
@@ -918,6 +918,11 @@ void USoundWave::Serialize( FArchive& Ar )
 		VirtualizationMode = EVirtualizationMode::PlayWhenSilent;
 	}
 
+	if (Ar.IsLoading())
+	{
+		ModulationSettings.VersionModulators();
+	}
+
 	// If we have deprecated properties with data, migrate the data to the new enum
 	if (bUseBinkAudio || bSeekableStreaming)
 	{
diff --git a/Engine/Source/Runtime/Engine/Public/AudioDevice.h b/Engine/Source/Runtime/Engine/Public/AudioDevice.h
--- a/Engine/Source/Runtime/Engine/Public/AudioDevice.h
+++ b/Engine/Source/Runtime/Engine/Public/AudioDevice.h
@@ -1485,7 +1485,13 @@ public:
 		UE_LOG(LogAudio, Error, TEXT("Submixes are only supported in audio mixer."));
 	}
 
+	UE_DEPRECATED(5.1, "UpdateSubmixModulationSettings taking single modulators is deprecated.  Use the overload that allows for modulator sets")
 	virtual void UpdateSubmixModulationSettings(USoundSubmix* InSoundSubmix, USoundModulatorBase* InOutputModulation, USoundModulatorBase* InWetLevelModulation, USoundModulatorBase* InDryLevelModulation)
+	{
+		UE_LOG(LogAudio, Error, TEXT("Submixes are only supported in audio mixer & 'UpdateSubmixModulationSettings' function taking single modulators no longer supported."));
+	}
+
+	virtual void UpdateSubmixModulationSettings(USoundSubmix* InSoundSubmix, const TSet<TObjectPtr<USoundModulatorBase>>& InOutputModulation, const TSet<TObjectPtr<USoundModulatorBase>>& InWetLevelModulation, const TSet<TObjectPtr<USoundModulatorBase>>& InDryLevelModulation)
 	{
 		UE_LOG(LogAudio, Error, TEXT("Submixes are only supported in audio mixer."));
 	}
