Tooltip fixup

diff --git a/Engine/Source/Runtime/AudioMixer/Public/AudioMixerBlueprintLibrary.h b/Engine/Source/Runtime/AudioMixer/Public/AudioMixerBlueprintLibrary.h
--- a/Engine/Source/Runtime/AudioMixer/Public/AudioMixerBlueprintLibrary.h
+++ b/Engine/Source/Runtime/AudioMixer/Public/AudioMixerBlueprintLibrary.h
@@ -343,7 +343,7 @@ public:
 	UFUNCTION(BlueprintCallable, Category = "Audio|Effects", meta = (WorldContext = "WorldContextObject"))
 	static void AddSourceEffectToPresetChain(const UObject* WorldContextObject, USoundEffectSourcePresetChain* PresetChain, FSourceEffectChainEntry Entry);
 
-	/** Adds source effect entry to preset chain. Only affects the instance of preset chain. */
+	/** Removes source effect entry from preset chain. Only affects the instance of preset chain. */
 	UFUNCTION(BlueprintCallable, Category = "Audio|Effects", meta = (WorldContext = "WorldContextObject"))
 	static void RemoveSourceEffectFromPresetChain(const UObject* WorldContextObject, USoundEffectSourcePresetChain* PresetChain, int32 EntryIndex);
 
