USynthComponent: exposing UAudioComponent FadeIn, FadeOut, AdjustVolume methods

diff --git a/Engine/Source/Runtime/AudioMixer/Private/Components/SynthComponent.cpp b/Engine/Source/Runtime/AudioMixer/Private/Components/SynthComponent.cpp
--- a/Engine/Source/Runtime/AudioMixer/Private/Components/SynthComponent.cpp
+++ b/Engine/Source/Runtime/AudioMixer/Private/Components/SynthComponent.cpp
@@ -576,6 +576,29 @@ void USynthComponent::SetOutputToBusOnly(bool bInOutputToBusOnly)
 	}
 }
 
+void USynthComponent::FadeIn(float FadeInDuration, float FadeVolumeLevel/* = 1.0f*/, float StartTime/* = 0.0f*/, const EAudioFaderCurve FadeCurve/* = EAudioFaderCurve::Linear*/) const
+{
+	if(AudioComponent)
+	{
+		AudioComponent->FadeIn(FadeInDuration, FadeVolumeLevel, StartTime, FadeCurve);
+	}
+}
+
+void USynthComponent::FadeOut(float FadeOutDuration, float FadeVolumeLevel, const EAudioFaderCurve FadeCurve/* = EAudioFaderCurve::Linear*/) const
+{
+	if(AudioComponent)
+	{
+		AudioComponent->FadeOut(FadeOutDuration, FadeVolumeLevel, FadeCurve);
+	}
+}
+
+void USynthComponent::AdjustVolume(float AdjustVolumeDuration, float AdjustVolumeLevel, const EAudioFaderCurve FadeCurve/* = EAudioFaderCurve::Linear*/) const
+{
+	if(AudioComponent)
+	{
+		AudioComponent->AdjustVolume(AdjustVolumeDuration, AdjustVolumeLevel, FadeCurve);
+	}
+}
 
 void USynthComponent::SynthCommand(TFunction<void()> Command)
 {
diff --git a/Engine/Source/Runtime/AudioMixer/Public/Components/SynthComponent.h b/Engine/Source/Runtime/AudioMixer/Public/Components/SynthComponent.h
--- a/Engine/Source/Runtime/AudioMixer/Public/Components/SynthComponent.h
+++ b/Engine/Source/Runtime/AudioMixer/Public/Components/SynthComponent.h
@@ -131,6 +131,36 @@ public:
 	UFUNCTION(BlueprintCallable, Category = "Audio|Components|Audio")
 	void SetOutputToBusOnly(bool bInOutputToBusOnly);
 
+	/**
+	 * This function allows designers to call Play on an Audio Component instance while applying a volume curve over time. 
+	 * Parameters allow designers to indicate the duration of the fade, the curve shape, and the start time if seeking into the sound.
+	 *
+	 * @param FadeInDuration How long it should take to reach the FadeVolumeLevel
+	 * @param FadeVolumeLevel The percentage of the AudioComponents's calculated volume to fade to
+	 * @param FadeCurve The curve to use when interpolating between the old and new volume
+	 */
+	UFUNCTION(BlueprintCallable, Category = "Audio|Components|Audio")
+	void FadeIn(float FadeInDuration, float FadeVolumeLevel = 1.0f, float StartTime = 0.0f, const EAudioFaderCurve FadeCurve = EAudioFaderCurve::Linear) const;
+
+	/**
+	 * This function allows designers to call a delayed Stop on an Audio Component instance while applying a
+	 * volume curve over time. Parameters allow designers to indicate the duration of the fade and the curve shape.
+	 *
+	 * @param FadeOutDuration how long it should take to reach the FadeVolumeLevel
+	 * @param FadeVolumeLevel the percentage of the AudioComponents's calculated volume in which to fade to
+	 * @param FadeCurve The curve to use when interpolating between the old and new volume
+	 */
+	UFUNCTION(BlueprintCallable, Category = "Audio|Components|Audio")
+	void FadeOut(float FadeOutDuration, float FadeVolumeLevel, const EAudioFaderCurve FadeCurve = EAudioFaderCurve::Linear) const;
+
+	/** This function allows designers to trigger an adjustment to the sound instance’s playback Volume with options for smoothly applying a curve over time.
+     * @param AdjustVolumeDuration The length of time in which to interpolate between the initial volume and the new volume.
+     * @param AdjustVolumeLevel The new volume to set the Audio Component to.
+     * @param FadeCurve The curve used when interpolating between the old and new volume.
+     */
+	UFUNCTION(BlueprintCallable, Category = "Audio|Components|Audio")
+	void AdjustVolume(float AdjustVolumeDuration, float AdjustVolumeLevel, const EAudioFaderCurve FadeCurve = EAudioFaderCurve::Linear) const;
+
 	/** Auto destroy this component on completion */
 	UPROPERTY()
 	uint8 bAutoDestroy : 1;
