Adding debug info to AudioMixerThreadCommand to breadcrumb hangs.

diff --git a/Engine/Source/Runtime/Android/AudioMixerAndroid/AudioMixerAndroid.Build.cs b/Engine/Source/Runtime/Android/AudioMixerAndroid/AudioMixerAndroid.Build.cs
--- a/Engine/Source/Runtime/Android/AudioMixerAndroid/AudioMixerAndroid.Build.cs
+++ b/Engine/Source/Runtime/Android/AudioMixerAndroid/AudioMixerAndroid.Build.cs
@@ -31,6 +31,8 @@ public class AudioMixerAndroid : ModuleRules
                     "Engine"
                 }
             );
+
+            PublicDefinitions.Add("WITH_AUDIO_MIXER_THREAD_COMMAND_DEBUG=0");
         }
 	}
 }
diff --git a/Engine/Source/Runtime/Apple/AudioMixerCoreAudio/AudioMixerCoreAudio.Build.cs b/Engine/Source/Runtime/Apple/AudioMixerCoreAudio/AudioMixerCoreAudio.Build.cs
--- a/Engine/Source/Runtime/Apple/AudioMixerCoreAudio/AudioMixerCoreAudio.Build.cs
+++ b/Engine/Source/Runtime/Apple/AudioMixerCoreAudio/AudioMixerCoreAudio.Build.cs
@@ -58,5 +58,7 @@ public class AudioMixerCoreAudio : ModuleRules
 		{
 			PrecompileForTargets = PrecompileTargetsType.Any;
 		}
+
+		PublicDefinitions.Add("WITH_AUDIO_MIXER_THREAD_COMMAND_DEBUG=0");
 	}
 }
diff --git a/Engine/Source/Runtime/AudioMixer/Private/AudioMixerSourceManager.cpp b/Engine/Source/Runtime/AudioMixer/Private/AudioMixerSourceManager.cpp
--- a/Engine/Source/Runtime/AudioMixer/Private/AudioMixerSourceManager.cpp
+++ b/Engine/Source/Runtime/AudioMixer/Private/AudioMixerSourceManager.cpp
@@ -3,7 +3,6 @@
 #include "AudioMixerSourceManager.h"
 
 #include "AudioMixerSourceBuffer.h"
-#include "AudioMixerSource.h"
 #include "AudioMixerDevice.h"
 #include "AudioMixerSourceVoice.h"
 #include "AudioMixerSubmix.h"
@@ -15,6 +14,7 @@
 #include "SoundFieldRendering.h"
 #include "ProfilingDebugging/CsvProfiler.h"
 #include "Async/Async.h"
+#include "ProfilingDebugging/CountersTrace.h"
 #include "Stats/Stats.h"
 
 // Link to "Audio" profiling category
@@ -97,6 +97,20 @@ FAutoConsoleVariableRef CVarCommandBufferMaxSizeMb(
 	TEXT("How big to allow the command buffer to grow before ignoring more commands"),
 	ECVF_Default);
 
+static float AudioCommandExecTimeMsWarningThresholdCvar = 1.f;
+FAutoConsoleVariableRef CVarAudioCommandExecTimeMsWarningThreshold(
+	TEXT("au.AudioThreadCommand.ExecutionTimeWarningThresholdInMs"),
+	AudioCommandExecTimeMsWarningThresholdCvar,
+	TEXT("If a command took longer to execute than this number (in milliseconds) then we log a warning"),
+	ECVF_Default);
+
+static int32 LogEveryAudioThreadCommandCvar = 0;
+FAutoConsoleVariableRef LogEveryAudioThreadCommand(
+	TEXT("au.AudioThreadCommand.LogEveryExecution"),
+	LogEveryAudioThreadCommandCvar,
+	TEXT("Extremely verbose logging of each Audio Thread command caller and it's execution time"),
+	ECVF_Default);
+
 // +/- 4 Octaves (default)
 static float MaxModulationPitchRangeFreqCVar = 16.0f;
 static float MinModulationPitchRangeFreqCVar = 0.0625f;
@@ -353,7 +367,7 @@ namespace Audio
 
 		// Spam command queue with nops.
 		static FAutoConsoleCommand SpamNopsCmd(
-			TEXT("au.SpamCommandQueue"),
+			TEXT("au.AudioThreadCommand.SpamCommandQueue"),
 			TEXT(""),
 			FConsoleCommandDelegate::CreateLambda([this]() 
 			{				
@@ -363,11 +377,27 @@ namespace Audio
 				} Payload;
 				for (int32 i = 0; i < 65536; ++i)
 				{
-					AudioMixerThreadCommand([Payload] {});
+					AudioMixerThreadCommand([Payload] {}, AUDIO_MIXER_THREAD_COMMAND_STRING("SpamNopsCmd() -- Console command"));
 				}
 			})
 		);
 
+
+		// submit a command that has an endless loop
+		static FAutoConsoleCommand SpamEndlessCmd(
+			TEXT("au.AudioThreadCommand.ChokeCommandQueue"),
+			TEXT(""),
+			FConsoleCommandDelegate::CreateLambda([this]()
+			{
+				struct FSpamPayload
+				{
+					uint8 JunkBytes[1024];
+				} Payload;
+
+				AudioMixerThreadCommand([Payload] {while(true){}}, AUDIO_MIXER_THREAD_COMMAND_STRING("SpamEndlessCmd() -- Console command"));
+			})
+		);
+
 		bInitialized = true;
 		bPumpQueue = false;
 	}
@@ -408,8 +438,9 @@ namespace Audio
 					// Pop and execute all the commands that came since last update tick
 					for (int32 Id = 0; Id < NextCommandBuffer.SourceCommandQueue.Num(); ++Id)
 					{
-						TFunction<void()>& CommandFunction = NextCommandBuffer.SourceCommandQueue[Id];
-						CommandFunction();
+						FAudioMixerThreadCommand AudioCommand = NextCommandBuffer.SourceCommandQueue[Id];
+
+						AudioCommand();
 						NumCommands.Decrement();
 					}
 
@@ -1031,7 +1062,7 @@ namespace Audio
 #endif 
 
 			AUDIO_MIXER_DEBUG_LOG(SourceId, TEXT("Is initializing"));
-		});
+		}, AUDIO_MIXER_THREAD_COMMAND_STRING("InitSource"));
 	}
 
 	void FMixerSourceManager::ReleaseSourceId(const int32 SourceId)
@@ -1057,7 +1088,7 @@ namespace Audio
 			AUDIO_MIXER_CHECK_AUDIO_PLAT_THREAD(MixerDevice);
 
 			ReleaseSource(SourceId);
-		});
+		}, AUDIO_MIXER_THREAD_COMMAND_STRING("ReleaseSourceId"));
 	}
 
 	void FMixerSourceManager::StartAudioBus(uint32 InAudioBusId, int32 InNumChannels, bool bInIsAutomatic)
@@ -1098,7 +1129,7 @@ namespace Audio
 					SourceInfo.bIsStopping = false;
 				}
 			}
-		});
+		}, AUDIO_MIXER_THREAD_COMMAND_STRING("StartAudioBus"));
 	}
 
 	void FMixerSourceManager::StopAudioBus(uint32 InAudioBusId)
@@ -1131,7 +1162,7 @@ namespace Audio
 					AudioBuses.Remove(InAudioBusId);
 				}
 			}
-		});
+		}, AUDIO_MIXER_THREAD_COMMAND_STRING("StartAudioBus"));
 	}
 
 	bool FMixerSourceManager::IsAudioBusActive(uint32 InAudioBusId) const
@@ -1178,7 +1209,7 @@ namespace Audio
 		AudioMixerThreadCommand([this, InAudioBusId, NewPatchPtr = InPatchOutputStrongPtr]() mutable
 		{
 			AddPatchOutputForAudioBus(InAudioBusId, NewPatchPtr);
-		});
+		}, AUDIO_MIXER_THREAD_COMMAND_STRING("AddPatchOutputForAudioBus_AudioThread()"));
 	}
 
 	void FMixerSourceManager::Play(const int32 SourceId)
@@ -1203,7 +1234,7 @@ namespace Audio
 			SourceInfo.StartTime = StartTime;
 
 			AUDIO_MIXER_DEBUG_LOG(SourceId, TEXT("Is playing"));
-		});
+		}, AUDIO_MIXER_THREAD_COMMAND_STRING("Play()"));
 	}
 
 	void FMixerSourceManager::CancelQuantizedSound(const int32 SourceId)
@@ -1252,7 +1283,7 @@ namespace Audio
 		AudioMixerThreadCommand([this, SourceId]()
 		{
 			StopInternal(SourceId);
-		});
+		}, AUDIO_MIXER_THREAD_COMMAND_STRING("Stop()"));
 	}
 
 	void FMixerSourceManager::StopInternal(const int32 SourceId)
@@ -1319,7 +1350,7 @@ namespace Audio
 			}
 
 			AUDIO_MIXER_DEBUG_LOG(SourceId, TEXT("Is stopping with fade"));
-		});
+		}, AUDIO_MIXER_THREAD_COMMAND_STRING("StopFade()"));
 	}
 
 
@@ -1337,7 +1368,7 @@ namespace Audio
 
 			SourceInfo.bIsPaused = true;
 			SourceInfo.bIsActive = false;
-		});
+		}, AUDIO_MIXER_THREAD_COMMAND_STRING("Pause()"));
 	}
 
 	void FMixerSourceManager::SetPitch(const int32 SourceId, const float Pitch)
@@ -1351,7 +1382,7 @@ namespace Audio
 			check(NumOutputFrames > 0);
 
 			SourceInfos[SourceId].PitchSourceParam.SetValue(Pitch, NumOutputFrames);
-		});
+		}, AUDIO_MIXER_THREAD_COMMAND_STRING("SetPitch()"));
 	}
 
 	void FMixerSourceManager::SetVolume(const int32 SourceId, const float Volume)
@@ -1378,7 +1409,7 @@ namespace Audio
 
 				SourceInfos[SourceId].VolumeSourceDestination = Volume;
 			}
-		});
+		}, AUDIO_MIXER_THREAD_COMMAND_STRING("SetVolume()"));
 	}
 
 	void FMixerSourceManager::SetDistanceAttenuation(const int32 SourceId, const float DistanceAttenuation)
@@ -1399,7 +1430,7 @@ namespace Audio
 			}
 
 			SourceInfos[SourceId].DistanceAttenuationSourceDestination = DistanceAttenuation;
-		});
+		}, AUDIO_MIXER_THREAD_COMMAND_STRING("SetDistanceAttenuation()"));
 	}
 
 	void FMixerSourceManager::SetSpatializationParams(const int32 SourceId, const FSpatializationParams& InParams)
@@ -1413,7 +1444,7 @@ namespace Audio
 			AUDIO_MIXER_CHECK_AUDIO_PLAT_THREAD(MixerDevice);
 
 			SourceInfos[SourceId].SpatParams = InParams;
-		});
+		}, AUDIO_MIXER_THREAD_COMMAND_STRING("SetSpatializationParams()"));
 	}
 
 	void FMixerSourceManager::SetChannelMap(const int32 SourceId, const uint32 NumInputChannels, const Audio::FAlignedFloatBuffer& ChannelMap, const bool bInIs3D, const bool bInIsCenterChannelOnly)
@@ -1444,7 +1475,7 @@ namespace Audio
 
 			bool bNeedsSpeakerMap = SourceSubmixOutput.SetChannelMap(ChannelMap, bInIsCenterChannelOnly);
 			GameThreadInfo.bNeedsSpeakerMap[SourceId] = bNeedsSpeakerMap;
-		});
+		}, AUDIO_MIXER_THREAD_COMMAND_STRING("SetChannelMap()"));
 	}
 
 	void FMixerSourceManager::SetLPFFrequency(const int32 SourceId, const float InLPFFrequency)
@@ -1464,7 +1495,7 @@ namespace Audio
 			// On callback, the client version may be overridden with the modulation LPF value depending on which is more aggressive.  
 			SourceInfo.LowPassFreq = InLPFFrequency;
 			SourceInfo.LowPassFilter.StartFrequencyInterpolation(InLPFFrequency, NumOutputFrames);
-		});
+		}, AUDIO_MIXER_THREAD_COMMAND_STRING("SetLPFFrequency()"));
 	}
 
 	void FMixerSourceManager::SetHPFFrequency(const int32 SourceId, const float InHPFFrequency)
@@ -1483,7 +1514,7 @@ namespace Audio
 			// On callback, the client version may be overridden with the modulation HPF value depending on which is more aggressive.  
 			SourceInfo.HighPassFreq = InHPFFrequency;
 			SourceInfo.HighPassFilter.StartFrequencyInterpolation(InHPFFrequency, NumOutputFrames);
-		});
+		}, AUDIO_MIXER_THREAD_COMMAND_STRING("SetHPFFrequency()"));
 	}
 
 	void FMixerSourceManager::SetModLPFFrequency(const int32 SourceId, const float InLPFFrequency)
@@ -1499,7 +1530,7 @@ namespace Audio
 			FSourceInfo& SourceInfo = SourceInfos[SourceId];
 			SourceInfo.LowpassModulationBase = InLPFFrequency;
 			SourceInfo.bModFiltersUpdated = true;
-		});
+		}, AUDIO_MIXER_THREAD_COMMAND_STRING("SetModLPFFrequency()"));
 	}
 
 	void FMixerSourceManager::SetModHPFFrequency(const int32 SourceId, const float InHPFFrequency)
@@ -1515,7 +1546,7 @@ namespace Audio
 			FSourceInfo& SourceInfo = SourceInfos[SourceId];
 			SourceInfo.HighpassModulationBase = InHPFFrequency;
 			SourceInfo.bModFiltersUpdated = true;
-		});
+		}, AUDIO_MIXER_THREAD_COMMAND_STRING("SetModHPFFrequency()"));
 	}
 
 	void FMixerSourceManager::SetModVolume(const int32 SourceId, const float InModVolume)
@@ -1530,7 +1561,7 @@ namespace Audio
 
 			FSourceInfo& SourceInfo = SourceInfos[SourceId];
 			SourceInfo.VolumeModulationBase = InModVolume;
-		});
+		}, AUDIO_MIXER_THREAD_COMMAND_STRING("SetModVolume()"));
 	}
 
 	void FMixerSourceManager::SetModPitch(const int32 SourceId, const float InModPitch)
@@ -1545,7 +1576,7 @@ namespace Audio
 
 			FSourceInfo& SourceInfo = SourceInfos[SourceId];
 			SourceInfo.PitchModulationBase = InModPitch;
-		});
+		}, AUDIO_MIXER_THREAD_COMMAND_STRING("SetModPitch()"));
 	}
 
 	void FMixerSourceManager::SetSubmixSendInfo(const int32 SourceId, const FMixerSourceSubmixSend& InSubmixSend)
@@ -1604,7 +1635,7 @@ namespace Audio
 					InSubmixPtr->AddOrSetSourceVoice(SourceVoice, InSubmixSend.SendLevel, InSubmixSend.SubmixSendStage);
 				}
 			}
-		});
+		}, AUDIO_MIXER_THREAD_COMMAND_STRING("SetSubmixSendInfo()"));
 	}
 
 	void FMixerSourceManager::ClearSubmixSendInfo(const int32 SourceId, const FMixerSourceSubmixSend& InSubmixSend)
@@ -1652,7 +1683,7 @@ namespace Audio
 				// Now remove the source voice from the submix send list
 				InSubmixPtr->RemoveSourceVoice(MixerSources[SourceId]);
 			}
-		});
+		}, AUDIO_MIXER_THREAD_COMMAND_STRING("ClearSubmixSendInfo()"));
 	}
 
 	void FMixerSourceManager::SetBusSendInfo(const int32 SourceId, EBusSendType InAudioBusSendType, uint32 AudioBusId, float BusSendLevel)
@@ -1707,7 +1738,7 @@ namespace Audio
 			{
 				SourceInfo.AudioBusSends[(int32)InAudioBusSendType].Add(AudioBusId);
 			}
-		});
+		}, AUDIO_MIXER_THREAD_COMMAND_STRING("SetBusSendInfo()"));
 	}
 
 	void FMixerSourceManager::SetListenerTransforms(const TArray<FTransform>& InListenerTransforms)
@@ -1715,7 +1746,7 @@ namespace Audio
 		AudioMixerThreadCommand([this, InListenerTransforms]()
 		{
 			ListenerTransforms = InListenerTransforms;
-		});
+		}, AUDIO_MIXER_THREAD_COMMAND_STRING("SetListenerTransforms()"));
 	}
 
 	const TArray<FTransform>* FMixerSourceManager::GetListenerTransforms() const
@@ -2111,6 +2142,22 @@ namespace Audio
 		}
 	}
 
+	// ctor
+	FMixerSourceManager::FAudioMixerThreadCommand::FAudioMixerThreadCommand(TFunction<void()> InFunction, FName InCallerDebugInfo)
+#if WITH_AUDIO_MIXER_THREAD_COMMAND_DEBUG
+		: CallerDebugInfo(InCallerDebugInfo)
+		, Function(MoveTemp(InFunction))
+#else
+		: Function(MoveTemp(InFunction))
+#endif // #if WITH_AUDIO_MIXER_THREAD_COMMAND_DEBUG
+	{
+	}
+
+	void FMixerSourceManager::FAudioMixerThreadCommand::operator()() const
+	{
+		Function();
+	}
+
 	void FMixerSourceManager::ApplyDistanceAttenuation(FSourceInfo& SourceInfo, int32 NumSamples)
 	{
 		if (DisableDistanceAttenuationCvar)
@@ -2715,7 +2762,7 @@ namespace Audio
 
 				SourceSubmixOutputBuffer.SetChannelMap(SourceInfo.ScratchChannelMap, SourceInfo.bIsCenterChannelOnly);
 			}
-		});
+		}, AUDIO_MIXER_THREAD_COMMAND_STRING("UpdateDeviceChannelCount()"));
 	}
 
 	void FMixerSourceManager::UpdateSourceEffectChain(const uint32 InSourceEffectChainId, const TArray<FSourceEffectChainEntry>& InSourceEffectChain, const bool bPlayEffectChainTails)
@@ -2780,7 +2827,7 @@ namespace Audio
 					}
 				}
 			}
-		});
+		}, AUDIO_MIXER_THREAD_COMMAND_STRING("UpdateSourceEffectChain()"));
 	}
 
 	void FMixerSourceManager::PauseSoundForQuantizationCommand(const int32 SourceId)
@@ -2926,7 +2973,7 @@ namespace Audio
 			SourceDataOverridePlugin->OnAllSourcesProcessed();
 		}
 
-		// Update the game thread copy of source doneness
+		// Update the game thread copy of source donenessf
 		for (int32 SourceId = 0; SourceId < NumTotalSources; ++SourceId)
 		{		
 			FSourceInfo& SourceInfo = SourceInfos[SourceId];
@@ -2966,8 +3013,53 @@ namespace Audio
 		MpscCommandQueue.Enqueue(MoveTemp(InCommand));
 	}
 
-	void FMixerSourceManager::AudioMixerThreadCommand(TFunction<void()> InFunction)
+	void FMixerSourceManager::FCurrentAudioMixerThreadCommandInfo::LogWarning() const
+	{
+#if WITH_AUDIO_MIXER_THREAD_COMMAND_DEBUG
+		if(CallerDebugInfo.IsValid())
+		{
+
+			UE_LOG(LogAudioMixer, Warning, TEXT("The AudioThreadCommand currently executing was called by %s, and started execution %f milliseconds ago")
+				, *CallerDebugInfo.ToString(), FPlatformTime::ToMilliseconds64(FPlatformTime::Cycles64() - QueueTimeInCycles.GetValue()));
+		}
+		else
+		{
+			UE_LOG(LogAudioMixer, Warning, TEXT("The AudioThreadCommand has not yet executed a command"));
+		}
+#endif // #if WITH_AUDIO_MIXER_THREAD_COMMAND_DEBUG
+	}
+
+	void FMixerSourceManager::FCurrentAudioMixerThreadCommandInfo::LogError() const
 	{
+#if WITH_AUDIO_MIXER_THREAD_COMMAND_DEBUG
+		if(CallerDebugInfo.IsValid())
+		{
+
+			UE_LOG(LogAudioMixer, Error, TEXT("The AudioThreadCommand currently executing was called by %s, and started execution %f milliseconds ago")
+				, *CallerDebugInfo.ToString(), FPlatformTime::ToMilliseconds64(FPlatformTime::Cycles64() - QueueTimeInCycles.GetValue()));
+		}
+		else
+		{
+			UE_LOG(LogAudioMixer, Error, TEXT("The AudioThreadCommand has not yet executed a command"));
+		}
+#endif // #if WITH_AUDIO_MIXER_THREAD_COMMAND_DEBUG
+	}
+
+	void FMixerSourceManager::FCurrentAudioMixerThreadCommandInfo::Reset()
+	{
+#if WITH_AUDIO_MIXER_THREAD_COMMAND_DEBUG
+		CallerDebugInfo = FName();
+#endif // #if WITH_AUDIO_MIXER_THREAD_COMMAND_DEBUG
+		QueueTimeInCycles.Reset();
+	}
+
+	void FMixerSourceManager::AudioMixerThreadCommand(TFunction<void()> InFunction, FName CallingFunction)
+	{
+		FAudioMixerThreadCommand AudioCommand(InFunction, CallingFunction);
+
+		auto test1 = AUDIO_MIXER_THREAD_COMMAND_STRING("My Name");
+		auto test2 = FName("My Name");
+
 		// Here, we make sure that we don't flip our command double buffer while we are executing this function.
 		FScopeLock ScopeLock(&CommandBufferIndexCriticalSection);
 		AUDIO_MIXER_CHECK_GAME_THREAD(MixerDevice);
@@ -2975,31 +3067,39 @@ namespace Audio
 		// Add the function to the command queue:
 		int32 AudioThreadCommandIndex = !RenderThreadCommandBufferIndex.GetValue();
 		SIZE_T CurrentBufferSizeInBytes = CommandBuffers[AudioThreadCommandIndex].SourceCommandQueue.GetAllocatedSize();
-		
-#if !NO_LOGGING
+
+		TRACE_INT_VALUE(TEXT("AudioMixerThreadCommands::CurrentBufferSizeInKb"), CurrentBufferSizeInBytes >> 10);
+
 		static SIZE_T WarnSize = 1024 * 1024;
 		if (CurrentBufferSizeInBytes > WarnSize )
 		{		
+			CurrentAudioMixerThreadCommandInfo.LogError();
+
 			SIZE_T Num = CommandBuffers[AudioThreadCommandIndex].SourceCommandQueue.Num();
-			// NOTE: Although not really and error we want this to show up in shipping builds.
-			UE_LOG(LogAudioMixer, Error, TEXT("Command Queue has grown to %ukb, containing %d cmds, last pump was %.2fms ago."), 
+			UE_LOG(LogAudioMixer, Error, TEXT("Command Queue has grown to %ukb, containing %d cmds, last pump was %.2fms ago."),
 				CurrentBufferSizeInBytes >> 10, Num, FPlatformTime::ToMilliseconds64(FPlatformTime::Cycles64() - LastPumpTimeInCycles));
 			WarnSize *= 2;
+
 		}
-#endif //!NO_LOGGING
 
 		// Before adding further commands, ensure we're not growing outside any sensible size for these buffers.
 		// On shipping builds, this will just stop us crashing from growing out of control and OOMing the machine.		
 		const SIZE_T MaxBufferSizeInBytes = ((SIZE_T)CommandBufferMaxSizeInMbCvar) << 20;
-		if (ensureMsgf(CurrentBufferSizeInBytes < MaxBufferSizeInBytes, TEXT("Command buffer grown to %umb, preventing any more adds! Likely cause, the h/w has stopped consuming data."), CurrentBufferSizeInBytes >>20))
+		if (CurrentBufferSizeInBytes < MaxBufferSizeInBytes)
 		{
-			CommandBuffers[AudioThreadCommandIndex].SourceCommandQueue.Add(MoveTemp(InFunction));
+			CommandBuffers[AudioThreadCommandIndex].SourceCommandQueue.Add(AudioCommand);
 			NumCommands.Increment();
 		}
+		else
+		{
+			CurrentAudioMixerThreadCommandInfo.LogError();
+			UE_LOG(LogAudioMixer, Error,  TEXT("Command buffer grown to %umb, preventing any more adds! Likely cause, the h/w has stopped consuming data."), CurrentBufferSizeInBytes >>20);
+		}
 	}
 
 	void FMixerSourceManager::PumpCommandQueue()
 	{
+		TRACE_CPUPROFILER_EVENT_SCOPE(AudioMixerThreadCommands::PumpCommandQueue)
 		// If we're already triggered, we need to wait for the audio thread to reset it before pumping
 		if (FPlatformProcess::SupportsMultithreading())
 		{
@@ -3009,6 +3109,8 @@ namespace Audio
 			}
 		}
 
+		const int32 CurrentRenderThreadIndex = RenderThreadCommandBufferIndex.GetValue();
+
 		// Pump the MPSC command queue
 		TOptional Opt{ MpscCommandQueue.Dequeue() };
 		while (Opt.IsSet())
@@ -3023,14 +3125,35 @@ namespace Audio
 
 		FCommands& Commands = CommandBuffers[CurrentRenderThreadIndex];
 
+		const int32 NumCommandsToExecute = Commands.SourceCommandQueue.Num();
+		TRACE_INT_VALUE(TEXT("AudioMixerThreadCommands::NumCommandsToExecute"), NumCommandsToExecute);
+
 		// Pop and execute all the commands that came since last update tick
-		for (int32 Id = 0; Id < Commands.SourceCommandQueue.Num(); ++Id)
+		for (int32 Id = 0; Id < NumCommandsToExecute; ++Id)
 		{
-			TFunction<void()>& CommandFunction = Commands.SourceCommandQueue[Id];
-			CommandFunction();
+			FAudioMixerThreadCommand& AudioCommand = Commands.SourceCommandQueue[Id];
+#if WITH_AUDIO_MIXER_THREAD_COMMAND_DEBUG
+			// cache debug info about the command we are about to execute (for safe access from the producer thread)
+			CurrentAudioMixerThreadCommandInfo.CallerDebugInfo = AudioCommand.CallerDebugInfo;
+			CurrentAudioMixerThreadCommandInfo.QueueTimeInCycles = FPlatformTime::Cycles64();
+#endif // #if WITH_AUDIO_MIXER_THREAD_COMMAND_DEBUG
+
+			AudioCommand(); // execute
+
+#if WITH_AUDIO_MIXER_THREAD_COMMAND_DEBUG
+			const float ExecutionTimeInMs = static_cast<float>(FPlatformTime::ToMilliseconds64(FPlatformTime::Cycles64() - CurrentAudioMixerThreadCommandInfo.QueueTimeInCycles.GetValue()));
+			if(LogEveryAudioThreadCommandCvar || (ExecutionTimeInMs >= AudioCommandExecTimeMsWarningThresholdCvar))
+			{
+				UE_LOG(LogAudioMixer, Warning, TEXT("Audio Thread Command from %s() took %f milleseconds to execute"),
+				*CurrentAudioMixerThreadCommandInfo.CallerDebugInfo.ToString(), ExecutionTimeInMs);
+			}
+#endif // #if WITH_AUDIO_MIXER_THREAD_COMMAND_DEBUG
+
+			CurrentAudioMixerThreadCommandInfo.Reset();
 			NumCommands.Decrement();
 		}
 
+
 		LastPumpTimeInCycles = FPlatformTime::Cycles64();
 		Commands.SourceCommandQueue.Reset();
 
@@ -3043,7 +3166,6 @@ namespace Audio
 		{
 			RenderThreadCommandBufferIndex.Set(!CurrentRenderThreadIndex);
 		}
-
 	}
 
 	void FMixerSourceManager::FlushCommandQueue(bool bPumpInCommand)
diff --git a/Engine/Source/Runtime/AudioMixer/Private/AudioMixerSourceManager.h b/Engine/Source/Runtime/AudioMixer/Private/AudioMixerSourceManager.h
--- a/Engine/Source/Runtime/AudioMixer/Private/AudioMixerSourceManager.h
+++ b/Engine/Source/Runtime/AudioMixer/Private/AudioMixerSourceManager.h
@@ -19,6 +19,19 @@
 #include "Sound/QuartzQuantizationUtilities.h"
 #include "Stats/Stats.h"
 
+// defines for debug data/logging in AudioMixerThreadCommand execution
+// if this hasn't been explicitly defined, assume it is enabled
+#ifndef WITH_AUDIO_MIXER_THREAD_COMMAND_DEBUG
+#define WITH_AUDIO_MIXER_THREAD_COMMAND_DEBUG 1
+#endif // #ifndef WITH_AUDIO_MIXER_THREAD_COMMAND_DEBUG
+
+#if WITH_AUDIO_MIXER_THREAD_COMMAND_DEBUG
+#define AUDIO_MIXER_THREAD_COMMAND_STRING(X) ( FName(X) )
+#else //WITH_AUDIO_MIXER_THREAD_COMMAND_DEBUG
+static const FName NAME_NONE;
+#define AUDIO_MIXER_THREAD_COMMAND_STRING(X) ( NAME_NONE )
+#endif //WITH_AUDIO_MIXER_THREAD_COMMAND_DEBUG
+
 
 // Tracks the time it takes to up the source manager (computes source buffers, source effects, sample rate conversion)
 DECLARE_CYCLE_STAT_EXTERN(TEXT("Source Manager Update"), STAT_AudioMixerSourceManagerUpdate, STATGROUP_AudioMixer, AUDIOMIXER_API);
@@ -61,6 +74,8 @@ namespace Audio
 	class ISourceListener
 	{
 	public:
+		virtual ~ISourceListener() = default;
+
 		// Called before a source begins to generate audio. 
 		virtual void OnBeginGenerate() = 0;
 
@@ -284,7 +299,34 @@ namespace Audio
 		void ComputeBuses();
 		void UpdateBuses();
 
-		void AudioMixerThreadCommand(TFunction<void()> InFunction);
+		struct FAudioMixerThreadCommand
+		{
+			// ctor
+			FAudioMixerThreadCommand(TFunction<void()> InFunction, FName InCallerDebugInfo);
+
+			// function-call operator
+			void operator()() const;
+
+			// data
+#if WITH_AUDIO_MIXER_THREAD_COMMAND_DEBUG
+			FName CallerDebugInfo;
+#endif // #if WITH_AUDIO_MIXER_THREAD_COMMAND_DEBUG
+			TFunction<void()> Function;
+		};
+
+		struct FCurrentAudioMixerThreadCommandInfo
+		{
+			FName CallerDebugInfo;
+			FThreadSafeCounter64 QueueTimeInCycles;
+
+			void LogWarning() const;
+			void LogError() const;
+			void Reset();
+		};
+
+		FCurrentAudioMixerThreadCommandInfo CurrentAudioMixerThreadCommandInfo;
+
+		void AudioMixerThreadCommand(TFunction<void()> InFunction, FName CallingFunction = {});
 
 		
 		static const int32 NUM_BYTES_PER_SAMPLE = 2;
@@ -339,7 +381,7 @@ namespace Audio
 		// A command queue to execute commands from audio thread (or game thread) to audio mixer device thread.
 		struct FCommands
 		{
-			TArray<TFunction<void()>> SourceCommandQueue;
+			TArray<FAudioMixerThreadCommand> SourceCommandQueue;
 		};
 		
 		FCommands CommandBuffers[2];
