Remove duplicate line

diff --git a/Engine/Source/Runtime/AudioMixer/Private/AudioMixerSourceManager.cpp b/Engine/Source/Runtime/AudioMixer/Private/AudioMixerSourceManager.cpp
--- a/Engine/Source/Runtime/AudioMixer/Private/AudioMixerSourceManager.cpp
+++ b/Engine/Source/Runtime/AudioMixer/Private/AudioMixerSourceManager.cpp
@@ -3109,8 +3109,6 @@ namespace Audio
 			}
 		}
 
-		const int32 CurrentRenderThreadIndex = RenderThreadCommandBufferIndex.GetValue();
-
 		// Pump the MPSC command queue
 		TOptional Opt{ MpscCommandQueue.Dequeue() };
 		while (Opt.IsSet())
