GC leak fixes for SynthComponent.
Stop SynthComponent playing audio when streamed out or when the world is destroyed.
Create Synth object without the component as the outer to prevent the component being leaked along with the synth sound.

#localization none
#tests Tested on console

diff --git a/Engine/Source/Runtime/AudioMixer/Private/Components/SynthComponent.cpp b/Engine/Source/Runtime/AudioMixer/Private/Components/SynthComponent.cpp
--- a/Engine/Source/Runtime/AudioMixer/Private/Components/SynthComponent.cpp
+++ b/Engine/Source/Runtime/AudioMixer/Private/Components/SynthComponent.cpp
@@ -233,7 +233,7 @@ void USynthComponent::Initialize(int32 SampleRateOverride)
 
 		if (!Synth)
 		{
-			Synth = NewObject<USynthSound>(this, TEXT("Synth"));
+			Synth = NewObject<USynthSound>();
 		}
 
 		// Copy sound base data to the sound
@@ -350,6 +350,17 @@ void USynthComponent::OnUnregister()
 	}
 }
 
+void USynthComponent::EndPlay(const EEndPlayReason::Type Reason) 
+{	
+	Super::EndPlay(Reason);
+
+	if (GetOwner() && (Reason == EEndPlayReason::LevelTransition || Reason == EEndPlayReason::RemovedFromWorld))
+	{
+		// If our world or sublevel is going away, stop immediately to prevent the containing world/level from being leaked via hard references from the audio device.
+		Stop();
+	}
+}
+
 USoundClass* USynthComponent::GetSoundClass()
 {
 	if (SoundClass)
diff --git a/Engine/Source/Runtime/AudioMixer/Public/Components/SynthComponent.h b/Engine/Source/Runtime/AudioMixer/Public/Components/SynthComponent.h
--- a/Engine/Source/Runtime/AudioMixer/Public/Components/SynthComponent.h
+++ b/Engine/Source/Runtime/AudioMixer/Public/Components/SynthComponent.h
@@ -85,6 +85,7 @@ public:
 	virtual void OnRegister() override;
 	virtual void OnUnregister() override;
 	virtual bool IsReadyForOwnerToAutoDestroy() const override;
+	virtual void EndPlay(const EEndPlayReason::Type Reason) override;
 	//~ End ActorComponent Interface.
 
 	//~ Begin UObject Interface.
