Fix for bypassing the a source effect chain entry in the same game frame as the sound is played failing (because the bypass request is executed before the target sound is initialized. )

diff --git a/Engine/Source/Runtime/AudioMixer/Private/AudioMixerSourceManager.cpp b/Engine/Source/Runtime/AudioMixer/Private/AudioMixerSourceManager.cpp
--- a/Engine/Source/Runtime/AudioMixer/Private/AudioMixerSourceManager.cpp
+++ b/Engine/Source/Runtime/AudioMixer/Private/AudioMixerSourceManager.cpp
@@ -2143,10 +2143,11 @@ namespace Audio
 	}
 
 	// ctor
-	FMixerSourceManager::FAudioMixerThreadCommand::FAudioMixerThreadCommand(TFunction<void()> InFunction, FName InCallerDebugInfo)
+	FMixerSourceManager::FAudioMixerThreadCommand::FAudioMixerThreadCommand(TFunction<void()> InFunction, FName InCallerDebugInfo, bool bInDeferExecution)
 #if WITH_AUDIO_MIXER_THREAD_COMMAND_DEBUG
 		: CallerDebugInfo(InCallerDebugInfo)
 		, Function(MoveTemp(InFunction))
+		, bDeferExecution(bInDeferExecution)
 #else
 		: Function(MoveTemp(InFunction))
 #endif // #if WITH_AUDIO_MIXER_THREAD_COMMAND_DEBUG
@@ -2827,7 +2828,7 @@ namespace Audio
 					}
 				}
 			}
-		}, AUDIO_MIXER_THREAD_COMMAND_STRING("UpdateSourceEffectChain()"));
+		}, AUDIO_MIXER_THREAD_COMMAND_STRING("UpdateSourceEffectChain()"), /*bDeferExecution*/true);
 	}
 
 	void FMixerSourceManager::PauseSoundForQuantizationCommand(const int32 SourceId)
@@ -3053,12 +3054,9 @@ namespace Audio
 		QueueTimeInCycles.Reset();
 	}
 
-	void FMixerSourceManager::AudioMixerThreadCommand(TFunction<void()> InFunction, FName CallingFunction)
+	void FMixerSourceManager::AudioMixerThreadCommand(TFunction<void()> InFunction, FName CallingFunction, bool bDeferExecution)
 	{
-		FAudioMixerThreadCommand AudioCommand(InFunction, CallingFunction);
-
-		auto test1 = AUDIO_MIXER_THREAD_COMMAND_STRING("My Name");
-		auto test2 = FName("My Name");
+		FAudioMixerThreadCommand AudioCommand(InFunction, CallingFunction, bDeferExecution);
 
 		// Here, we make sure that we don't flip our command double buffer while we are executing this function.
 		FScopeLock ScopeLock(&CommandBufferIndexCriticalSection);
@@ -3127,6 +3125,8 @@ namespace Audio
 		TRACE_INT_VALUE(TEXT("AudioMixerThreadCommands::NumCommandsToExecute"), NumCommandsToExecute);
 
 		// Pop and execute all the commands that came since last update tick
+		TArray<FAudioMixerThreadCommand> DelayedCommands;
+
 		for (int32 Id = 0; Id < NumCommandsToExecute; ++Id)
 		{
 			FAudioMixerThreadCommand& AudioCommand = Commands.SourceCommandQueue[Id];
@@ -3135,8 +3135,15 @@ namespace Audio
 			CurrentAudioMixerThreadCommandInfo.CallerDebugInfo = AudioCommand.CallerDebugInfo;
 			CurrentAudioMixerThreadCommandInfo.QueueTimeInCycles = FPlatformTime::Cycles64();
 #endif // #if WITH_AUDIO_MIXER_THREAD_COMMAND_DEBUG
-
-			AudioCommand(); // execute
+			if(AudioCommand.bDeferExecution)
+			{
+				AudioCommand.bDeferExecution = false;
+				DelayedCommands.Emplace(MoveTemp(AudioCommand));
+			}
+			else
+			{
+				AudioCommand(); // execute
+			}
 
 #if WITH_AUDIO_MIXER_THREAD_COMMAND_DEBUG
 			const float ExecutionTimeInMs = static_cast<float>(FPlatformTime::ToMilliseconds64(FPlatformTime::Cycles64() - CurrentAudioMixerThreadCommandInfo.QueueTimeInCycles.GetValue()));
@@ -3153,7 +3160,7 @@ namespace Audio
 
 
 		LastPumpTimeInCycles = FPlatformTime::Cycles64();
-		Commands.SourceCommandQueue.Reset();
+		Commands.SourceCommandQueue = DelayedCommands;
 
 		if (FPlatformProcess::SupportsMultithreading())
 		{
diff --git a/Engine/Source/Runtime/AudioMixer/Private/AudioMixerSourceManager.h b/Engine/Source/Runtime/AudioMixer/Private/AudioMixerSourceManager.h
--- a/Engine/Source/Runtime/AudioMixer/Private/AudioMixerSourceManager.h
+++ b/Engine/Source/Runtime/AudioMixer/Private/AudioMixerSourceManager.h
@@ -302,7 +302,7 @@ namespace Audio
 		struct FAudioMixerThreadCommand
 		{
 			// ctor
-			FAudioMixerThreadCommand(TFunction<void()> InFunction, FName InCallerDebugInfo);
+			FAudioMixerThreadCommand(TFunction<void()> InFunction, FName InCallerDebugInfo, bool bInDeferExecution = false);
 
 			// function-call operator
 			void operator()() const;
@@ -312,6 +312,11 @@ namespace Audio
 			FName CallerDebugInfo;
 #endif // #if WITH_AUDIO_MIXER_THREAD_COMMAND_DEBUG
 			TFunction<void()> Function;
+
+			// Defers the execution by a single call to PumpCommandQueue()
+			// (used for commands that affect a playing source,
+			// and that source gets initialized after the command executes
+			bool bDeferExecution;
 		};
 
 		struct FCurrentAudioMixerThreadCommandInfo
@@ -326,7 +331,7 @@ namespace Audio
 
 		FCurrentAudioMixerThreadCommandInfo CurrentAudioMixerThreadCommandInfo;
 
-		void AudioMixerThreadCommand(TFunction<void()> InFunction, FName CallingFunction = {});
+		void AudioMixerThreadCommand(TFunction<void()> InFunction, FName CallingFunction = {}, bool bInDeferExecution = false);
 
 		
 		static const int32 NUM_BYTES_PER_SAMPLE = 2;
