Allow source buffer listeners to be set on the source manager from source data override plugins.

diff --git a/Engine/Source/Runtime/AudioMixer/Private/AudioMixerSource.cpp b/Engine/Source/Runtime/AudioMixer/Private/AudioMixerSource.cpp
--- a/Engine/Source/Runtime/AudioMixer/Private/AudioMixerSource.cpp
+++ b/Engine/Source/Runtime/AudioMixer/Private/AudioMixerSource.cpp
@@ -1703,6 +1703,8 @@ namespace Audio
 		}
  		
 		MixerSourceVoice->SetEnablement(WaveInstance->bEnableBusSends, WaveInstance->bEnableBaseSubmix, WaveInstance->bEnableSubmixSends);
+
+		MixerSourceVoice->SetSourceBufferListener(WaveInstance->SourceBufferListener, WaveInstance->bShouldSourceBufferListenerZeroBuffer);
 	}
 
 	void FMixerSource::UpdateSourceBusSends()
diff --git a/Engine/Source/Runtime/AudioMixer/Private/AudioMixerSourceManager.cpp b/Engine/Source/Runtime/AudioMixer/Private/AudioMixerSourceManager.cpp
--- a/Engine/Source/Runtime/AudioMixer/Private/AudioMixerSourceManager.cpp
+++ b/Engine/Source/Runtime/AudioMixer/Private/AudioMixerSourceManager.cpp
@@ -1549,6 +1549,22 @@ namespace Audio
 		}, AUDIO_MIXER_THREAD_COMMAND_STRING("SetModHPFFrequency()"));
 	}
 
+	void FMixerSourceManager::SetSourceBufferListener(const int32 SourceId, FSharedISourceBufferListenerPtr& InSourceBufferListener, bool InShouldSourceBufferListenerZeroBuffer)
+	{
+		AUDIO_MIXER_CHECK(SourceId < NumTotalSources);
+		AUDIO_MIXER_CHECK(GameThreadInfo.bIsBusy[SourceId]);
+		AUDIO_MIXER_CHECK_GAME_THREAD(MixerDevice);
+
+		AudioMixerThreadCommand([this, SourceId, InSourceBufferListener, InShouldSourceBufferListenerZeroBuffer]()
+		{
+			AUDIO_MIXER_CHECK_AUDIO_PLAT_THREAD(MixerDevice);
+
+			FSourceInfo& SourceInfo = SourceInfos[SourceId];
+			SourceInfo.SourceBufferListener = InSourceBufferListener;
+			SourceInfo.bShouldSourceBufferListenerZeroBuffer = InShouldSourceBufferListenerZeroBuffer;
+		});
+	}
+
 	void FMixerSourceManager::SetModVolume(const int32 SourceId, const float InModVolume)
 	{
 		AUDIO_MIXER_CHECK(SourceId < NumTotalSources);
diff --git a/Engine/Source/Runtime/AudioMixer/Private/AudioMixerSourceManager.h b/Engine/Source/Runtime/AudioMixer/Private/AudioMixerSourceManager.h
--- a/Engine/Source/Runtime/AudioMixer/Private/AudioMixerSourceManager.h
+++ b/Engine/Source/Runtime/AudioMixer/Private/AudioMixerSourceManager.h
@@ -228,6 +228,8 @@ namespace Audio
 		void SetModLPFFrequency(const int32 SourceId, const float InModFrequency);
 		void SetModHPFFrequency(const int32 SourceId, const float InModFrequency);
 
+		void SetSourceBufferListener(const int32 SourceId, FSharedISourceBufferListenerPtr& InSourceBufferListener, bool InShouldSourceBufferListenerZeroBuffer);
+
 
 		void SetListenerTransforms(const TArray<FTransform>& ListenerTransforms);
 		const TArray<FTransform>* GetListenerTransforms() const;
diff --git a/Engine/Source/Runtime/AudioMixer/Private/AudioMixerSourceVoice.cpp b/Engine/Source/Runtime/AudioMixer/Private/AudioMixerSourceVoice.cpp
--- a/Engine/Source/Runtime/AudioMixer/Private/AudioMixerSourceVoice.cpp
+++ b/Engine/Source/Runtime/AudioMixer/Private/AudioMixerSourceVoice.cpp
@@ -196,6 +196,13 @@ namespace Audio
 		}
 	}
 
+	void FMixerSourceVoice::SetSourceBufferListener(FSharedISourceBufferListenerPtr& InSourceBufferListener, bool InShouldSourceBufferListenerZeroBuffer)
+	{
+		AUDIO_MIXER_CHECK_GAME_THREAD(MixerDevice);
+
+		SourceManager->SetSourceBufferListener(SourceId, InSourceBufferListener, InShouldSourceBufferListenerZeroBuffer);
+	}
+
 	void FMixerSourceVoice::SetChannelMap(const uint32 NumInputChannels, const Audio::FAlignedFloatBuffer& InChannelMap, const bool bInIs3D, const bool bInIsCenterChannelOnly)
 	{
 		AUDIO_MIXER_CHECK_GAME_THREAD(MixerDevice);
diff --git a/Engine/Source/Runtime/AudioMixer/Private/AudioMixerSourceVoice.h b/Engine/Source/Runtime/AudioMixer/Private/AudioMixerSourceVoice.h
--- a/Engine/Source/Runtime/AudioMixer/Private/AudioMixerSourceVoice.h
+++ b/Engine/Source/Runtime/AudioMixer/Private/AudioMixerSourceVoice.h
@@ -59,6 +59,9 @@ namespace Audio
 		// Sets the source voice's HPF filter modulation base frequency.
 		void SetModHPFFrequency(const float InFrequency);
 
+		// Set the source voice's SourceBufferListener and associated boolean.
+		void SetSourceBufferListener(FSharedISourceBufferListenerPtr& InSourceBufferListener, bool InShouldSourceBufferListenerZeroBuffer);
+
 		// Sets the source voice's channel map (2d or 3d).
 		void SetChannelMap(const uint32 NumInputChannels, const Audio::FAlignedFloatBuffer& InChannelMap, const bool bInIs3D, const bool bInIsCenterChannelOnly);
 
