Fix for issues with audio bus registry from the audio mixer thread.

diff --git a/Engine/Source/Runtime/AudioMixer/Private/AudioMixerDevice.cpp b/Engine/Source/Runtime/AudioMixer/Private/AudioMixerDevice.cpp
--- a/Engine/Source/Runtime/AudioMixer/Private/AudioMixerDevice.cpp
+++ b/Engine/Source/Runtime/AudioMixer/Private/AudioMixerDevice.cpp
@@ -2606,6 +2606,13 @@ namespace Audio
 
 	FPatchOutputStrongPtr FMixerDevice::AddPatchForAudioBus(uint32 InAudioBusId, float InPatchGain)
 	{
+		// This function is supporting adding audio bus patches from multiple threads (AT, ART, GT, and tasks) and is currently
+		// depending on a number of places where data lives, which accounts for the complexity here.
+		// The key idea here is to create and return a strong patch output ptr at roughly the size we expect to need then 
+		// pass the strong ptr down to the audio render thread that then is registered to the audio bus to start feeding audio to.
+		// This code needs a clean up to refactor everything into a true MPSC model, along with an MPSC refactor of the source manager
+		// and our command queues. Once we do that we can remove the code which branches based on the thread the request is coming from. 
+
 		if (IsInGameThread())
 		{
 			if (ActiveAudioBuses_GameThread.Find(InAudioBusId))
@@ -2621,6 +2628,29 @@ namespace Audio
 			UE_LOG(LogAudioMixer, Warning, TEXT("Unable to add a patch output for audio bus because audio bus id '%d' is not active."), InAudioBusId);
 			return nullptr;
 		}
+		else if (IsInAudioThread())
+		{
+			int32 NumOutputFrames = SourceManager->GetNumOutputFrames();
+			FPatchOutputStrongPtr StrongOutputPtr = MakeShareable(new FPatchOutput(2 * NumOutputFrames, InPatchGain));
+
+			SourceManager->AddPatchOutputForAudioBus_AudioThread(InAudioBusId, StrongOutputPtr);
+			return StrongOutputPtr;
+		}
+		else if (IsAudioRenderingThread())
+		{
+			check(SourceManager);
+
+			const int32 NumChannels = SourceManager->GetAudioBusNumChannels(InAudioBusId);
+			if (NumChannels > 0)
+			{
+				const int32 NumOutputFrames = SourceManager->GetNumOutputFrames();
+				FPatchOutputStrongPtr StrongOutputPtr = MakeShareable(new FPatchOutput(NumOutputFrames * NumChannels, InPatchGain));
+				SourceManager->AddPatchOutputForAudioBus(InAudioBusId, StrongOutputPtr);
+				return StrongOutputPtr;
+			}
+
+			return nullptr;
+		}
 		else
 		{
 			// Need to make a strong output patch even if this is not going to ever write to it since the bus may not be running
