Fix scope lock affecting multiple submix spectral analysis delegates

diff --git a/Engine/Source/Runtime/AudioMixer/Private/AudioMixerSubmix.cpp b/Engine/Source/Runtime/AudioMixer/Private/AudioMixerSubmix.cpp
--- a/Engine/Source/Runtime/AudioMixer/Private/AudioMixerSubmix.cpp
+++ b/Engine/Source/Runtime/AudioMixer/Private/AudioMixerSubmix.cpp
@@ -2373,28 +2373,26 @@ namespace Audio
 			{
 				if (ensureMsgf(SpectrumAnalyzer.IsValid(), TEXT("Analyzing spectrum with invalid spectrum analyzer")))
 				{
-					// New results array
-					TArray<float> SpectralResults;
-
-					//const TArray<float>& InFrequencies, TArray<float>& OutMagnitudes
-					for (FSpectrumAnalysisDelegateInfo& DelegateInfo : SpectralAnalysisDelegates)
+					TArray<TPair<FSpectrumAnalysisDelegateInfo*, TArray<float>>> ResultsPerDelegate;
+					
 					{
-						const float CurrentTime = FPlatformTime::ToSeconds64(FPlatformTime::Cycles64());
+						// This lock ensures that the spectrum analyzer's analysis buffer doesn't
+						// change in this scope. 
+						Audio::FAsyncSpectrumAnalyzerScopeLock AnalyzerLock(SpectrumAnalyzer.Get());
 
-						// Don't update the spectral band until it's time since the last tick.
-						if (DelegateInfo.LastUpdateTime > 0.0f && ((CurrentTime - DelegateInfo.LastUpdateTime) < DelegateInfo.UpdateDelta))
+						for (FSpectrumAnalysisDelegateInfo& DelegateInfo : SpectralAnalysisDelegates)
 						{
-							continue;
-						}
+							TArray<float> SpectralResults;
+							SpectralResults.Reset();
+							const float CurrentTime = FPlatformTime::ToSeconds64(FPlatformTime::Cycles64());
 
-						DelegateInfo.LastUpdateTime = CurrentTime;
-
-						SpectralResults.Reset();
+							// Don't update the spectral band until it's time since the last tick.
+							if (DelegateInfo.LastUpdateTime > 0.0f && ((CurrentTime - DelegateInfo.LastUpdateTime) < DelegateInfo.UpdateDelta))
+							{
+								continue;
+							}
 
-						{
-							// This lock ensures that the spectrum analyzer's analysis buffer doesn't
-							// change in this scope. 
-							Audio::FAsyncSpectrumAnalyzerScopeLock AnalyzerLock(SpectrumAnalyzer.Get());
+							DelegateInfo.LastUpdateTime = CurrentTime;
 
 							if (ensure(DelegateInfo.SpectrumBandExtractor.IsValid()))
 							{
@@ -2402,22 +2400,27 @@ namespace Audio
 
 								SpectrumAnalyzer->GetBands(*Extractor, SpectralResults);
 							}
+							ResultsPerDelegate.Emplace(&DelegateInfo, MoveTemp(SpectralResults));
 						}
+					}
 
+					for (TPair<FSpectrumAnalysisDelegateInfo*, TArray<float>> Results : ResultsPerDelegate)
+					{
+						FSpectrumAnalysisDelegateInfo* DelegateInfo = Results.Key;
 						// Feed the results through the band envelope followers
-						for (int32 ResultIndex = 0; ResultIndex < SpectralResults.Num(); ++ResultIndex)
+						for (int32 ResultIndex = 0; ResultIndex < Results.Value.Num(); ++ResultIndex)
 						{
-							if (ensure(ResultIndex < DelegateInfo.SpectralBands.Num()))
+							if (ensure(ResultIndex < DelegateInfo->SpectralBands.Num()))
 							{
-								FSpectralAnalysisBandInfo& BandInfo = DelegateInfo.SpectralBands[ResultIndex];
+								FSpectralAnalysisBandInfo& BandInfo = DelegateInfo->SpectralBands[ResultIndex];
 
-								SpectralResults[ResultIndex] = BandInfo.EnvelopeFollower.ProcessSample(SpectralResults[ResultIndex]);
+								Results.Value[ResultIndex] = BandInfo.EnvelopeFollower.ProcessSample(Results.Value[ResultIndex]);
 							}
 						}
 
-						if (DelegateInfo.OnSubmixSpectralAnalysis.IsBound())
+						if (DelegateInfo->OnSubmixSpectralAnalysis.IsBound())
 						{
-							DelegateInfo.OnSubmixSpectralAnalysis.Broadcast(SpectralResults);
+							DelegateInfo->OnSubmixSpectralAnalysis.Broadcast(MoveTemp(Results.Value));
 						}
 					}
 				}
