Added missing critical section around iterating through SubmixBufferListeners.

diff --git a/Engine/Source/Runtime/AudioMixer/Private/AudioMixerSubmix.cpp b/Engine/Source/Runtime/AudioMixer/Private/AudioMixerSubmix.cpp
--- a/Engine/Source/Runtime/AudioMixer/Private/AudioMixerSubmix.cpp
+++ b/Engine/Source/Runtime/AudioMixer/Private/AudioMixerSubmix.cpp
@@ -1078,12 +1078,15 @@ namespace Audio
 			return true;
 		}
 
-		// query the SubmixBufferListeners to see if they plan to render audio into this buffer
-		for(const ISubmixBufferListener* listener : BufferListeners)
 		{
-			if(listener->IsRenderingAudio())
+			// query the SubmixBufferListeners to see if they plan to render audio into this buffer
+			FScopeLock Lock(&BufferListenerCriticalSection);
+			for(const ISubmixBufferListener* BufferListener : BufferListeners)
 			{
-				return true;
+				if(BufferListener->IsRenderingAudio())
+				{
+					return true;
+				}
 			}
 		}
 
diff --git a/Engine/Source/Runtime/AudioMixer/Public/AudioMixerSubmix.h b/Engine/Source/Runtime/AudioMixer/Public/AudioMixerSubmix.h
--- a/Engine/Source/Runtime/AudioMixer/Public/AudioMixerSubmix.h
+++ b/Engine/Source/Runtime/AudioMixer/Public/AudioMixerSubmix.h
@@ -565,7 +565,7 @@ namespace Audio
 		TArray<ISubmixBufferListener*> BufferListeners;
 
 		// Critical section used for modifying and interacting with buffer listeners
-		FCriticalSection BufferListenerCriticalSection;
+		mutable FCriticalSection BufferListenerCriticalSection;
 
 		// This buffer is used for recorded output of the submix.
 		FAlignedFloatBuffer RecordingData;
