Reverting fix 21707031 to prevent gc leak

diff --git a/Engine/Source/Runtime/AudioMixer/Private/Components/SynthComponent.cpp b/Engine/Source/Runtime/AudioMixer/Private/Components/SynthComponent.cpp
--- a/Engine/Source/Runtime/AudioMixer/Private/Components/SynthComponent.cpp
+++ b/Engine/Source/Runtime/AudioMixer/Private/Components/SynthComponent.cpp
@@ -11,7 +11,6 @@
 
 USynthSound::USynthSound(const FObjectInitializer& ObjectInitializer)
 	: Super(ObjectInitializer)
-	, OwningSynthComponent(nullptr)
 {
 }
 
@@ -49,7 +48,7 @@ void USynthSound::StartOnAudioDevice(FAudioDevice* InAudioDevice)
 
 void USynthSound::OnBeginGenerate()
 {
-	if (ensure(OwningSynthComponent))
+	if (ensure(OwningSynthComponent.IsValid()))
 	{
 		OwningSynthComponent->OnBeginGenerate();
 	}
@@ -67,7 +66,7 @@ int32 USynthSound::OnGeneratePCMAudio(TArray<uint8>& OutAudio, int32 NumSamples)
 		OutAudio.AddZeroed(NumSamples * sizeof(float));
 
 		// Mark pending kill can null this out on the game thread in rare cases.
-		if (!OwningSynthComponent)
+		if (!OwningSynthComponent.IsValid())
 		{
 			return 0;
 		}
@@ -81,7 +80,7 @@ int32 USynthSound::OnGeneratePCMAudio(TArray<uint8>& OutAudio, int32 NumSamples)
 		FloatBuffer.AddZeroed(NumSamples * sizeof(float));
 
 		// Mark pending kill can null this out on the game thread in rare cases.
-		if (!OwningSynthComponent)
+		if (!OwningSynthComponent.IsValid())
 		{
 			return 0;
 		}
@@ -105,7 +104,7 @@ int32 USynthSound::OnGeneratePCMAudio(TArray<uint8>& OutAudio, int32 NumSamples)
 void USynthSound::OnEndGenerate()
 {
 	// Mark pending kill can null this out on the game thread in rare cases.
-	if (OwningSynthComponent)
+	if (OwningSynthComponent.IsValid())
 	{
 		OwningSynthComponent->OnEndGenerate();
 	}
@@ -113,7 +112,7 @@ void USynthSound::OnEndGenerate()
 
 ISoundGeneratorPtr USynthSound::CreateSoundGenerator(const FSoundGeneratorInitParams& InParams)
 {
-	if (OwningSynthComponent)
+	if (OwningSynthComponent.IsValid())
 	{
 		return OwningSynthComponent->CreateSoundGeneratorInternal(InParams);
 	}
diff --git a/Engine/Source/Runtime/AudioMixer/Public/Components/SynthComponent.h b/Engine/Source/Runtime/AudioMixer/Public/Components/SynthComponent.h
--- a/Engine/Source/Runtime/AudioMixer/Public/Components/SynthComponent.h
+++ b/Engine/Source/Runtime/AudioMixer/Public/Components/SynthComponent.h
@@ -56,7 +56,7 @@ class AUDIOMIXER_API USynthSound : public USoundWaveProcedural
 
 protected:
 	UPROPERTY()
-	TObjectPtr<USynthComponent> OwningSynthComponent;
+	TWeakObjectPtr<USynthComponent> OwningSynthComponent = nullptr;
 
 	TArray<float> FloatBuffer;
 	bool bAudioMixer;
