Clamping submix effect chain entry index in Audio Mixer Blueprint Library
#jria UE-165936, UE-165938

diff --git a/Engine/Source/Runtime/AudioMixer/Private/AudioMixerBlueprintLibrary.cpp b/Engine/Source/Runtime/AudioMixer/Private/AudioMixerBlueprintLibrary.cpp
--- a/Engine/Source/Runtime/AudioMixer/Private/AudioMixerBlueprintLibrary.cpp
+++ b/Engine/Source/Runtime/AudioMixer/Private/AudioMixerBlueprintLibrary.cpp
@@ -536,12 +536,11 @@ void UAudioMixerBlueprintLibrary::RemoveSourceEffectFromPresetChain(const UObjec
 			Chain = PresetChain->Chain;
 		}
 
-		if (EntryIndex < Chain.Num())
+		if (EntryIndex >= 0 && EntryIndex < Chain.Num())
 		{
 			Chain.RemoveAt(EntryIndex);
+			MixerDevice->UpdateSourceEffectChain(PresetChainId, Chain, PresetChain->bPlayEffectChainTails);
 		}
-
-		MixerDevice->UpdateSourceEffectChain(PresetChainId, Chain, PresetChain->bPlayEffectChainTails);
 	}
 
 }
@@ -565,12 +564,11 @@ void UAudioMixerBlueprintLibrary::SetBypassSourceEffectChainEntry(const UObject*
 			Chain = PresetChain->Chain;
 		}
 
-		if (EntryIndex < Chain.Num())
+		if (EntryIndex >= 0 && EntryIndex < Chain.Num())
 		{
 			Chain[EntryIndex].bBypass = bBypassed;
+			MixerDevice->UpdateSourceEffectChain(PresetChainId, Chain, PresetChain->bPlayEffectChainTails);
 		}
-
-		MixerDevice->UpdateSourceEffectChain(PresetChainId, Chain, PresetChain->bPlayEffectChainTails);
 	}
 }
 
