Fix for clock handle/manager associating breaking on level change.

diff --git a/Engine/Source/Runtime/AudioMixer/Private/Quartz/QuartzSubsystem.cpp b/Engine/Source/Runtime/AudioMixer/Private/Quartz/QuartzSubsystem.cpp
--- a/Engine/Source/Runtime/AudioMixer/Private/Quartz/QuartzSubsystem.cpp
+++ b/Engine/Source/Runtime/AudioMixer/Private/Quartz/QuartzSubsystem.cpp
@@ -29,48 +29,72 @@ FAutoConsoleVariableRef CVarSimulateNoAudioDevice(
 	TEXT("0: Not Enabled, 1: Enabled"),
 	ECVF_Default);
 
-
-static FAudioDevice* GetAudioDeviceUsingWorldContext(const UObject* WorldContextObject)
+static Audio::FMixerDevice* GetAudioMixerDevice(const UWorld* InWorld)
 {
-	if (nullptr == WorldContextObject)
+	if(!InWorld || SimulateNoAudioDeviceCvar)
 	{
 		return nullptr;
 	}
 
-	const UWorld* ThisWorld = GEngine->GetWorldFromContextObject(WorldContextObject, EGetWorldErrorMode::LogAndReturnNull);
-	if (!ThisWorld || !ThisWorld->bAllowAudioPlayback || ThisWorld->GetNetMode() == NM_DedicatedServer)
+	FAudioDevice* AudioDevicePtr = InWorld->GetAudioDeviceRaw();
+	if(AudioDevicePtr && AudioDevicePtr->IsAudioMixerEnabled())
 	{
-		return nullptr;
+		return static_cast<Audio::FMixerDevice*>(AudioDevicePtr);
 	}
 
-	return ThisWorld->GetAudioDeviceRaw();
+	return nullptr;
 }
 
+static TUniquePtr<FScopeLock> GetPersistentStateScopeLock(const UWorld* InWorld)
+{
+	if(Audio::FMixerDevice* MixerDevicePtr = GetAudioMixerDevice(InWorld))
+	{
+		return MakeUnique<FScopeLock>(&MixerDevicePtr->QuartzPersistentStateCritSec);
+	}
 
-static Audio::FMixerDevice* GetAudioMixerDeviceUsingWorldContext(const UObject* WorldContextObject)
+	return {}; // Empty lock for AudioMixer-less runtimes (server)
+}
+
+
+void UQuartzSubsystem::Initialize(FSubsystemCollectionBase& Collection)
 {
-	if (FAudioDevice* AudioDevice = GetAudioDeviceUsingWorldContext(WorldContextObject))
+	Super::Initialize(Collection);
+	TUniquePtr<FScopeLock> Lock(GetPersistentStateScopeLock(GetWorld()));
+
+	if(const UWorld* WorldPtr = GetWorld())
 	{
-		if (!AudioDevice->IsAudioMixerEnabled())
+		if(Audio::FMixerDevice const* MixerDevice = static_cast<Audio::FMixerDevice*>(WorldPtr->GetAudioDeviceRaw()))
 		{
-			return nullptr;
-		}
-		else
-		{
-			return static_cast<Audio::FMixerDevice*>(AudioDevice);
+			ClockManagerDataPtr = MixerDevice->QuartzSubsystemData;
 		}
 	}
-	return nullptr;
+
+	if(!ClockManagerDataPtr)
+	{
+		ClockManagerDataPtr = MakeShared<Audio::FPersistentQuartzSubsystemData>();
+	}
 }
 
+void UQuartzSubsystem::Deinitialize()
+{
+	Super::Deinitialize();
+	TUniquePtr<FScopeLock> Lock(GetPersistentStateScopeLock(GetWorld()));
+
+	if(Audio::FMixerDevice* MixerDevice = GetAudioMixerDevice(GetWorld()))
+	{
+		MixerDevice->QuartzSubsystemData = ClockManagerDataPtr;
+	}
+}
 
 void UQuartzSubsystem::BeginDestroy()
 {
 	Super::BeginDestroy();
 
 	// force un-subscribe all Quartz tickable objects
-	SubsystemClockManager.Shutdown();
-	SubsystemClockManager.Flush();
+	if (ClockManagerDataPtr)
+	{
+		ClockManagerDataPtr->SubsystemClockManager.Flush();
+	}
 }
 
 void FQuartzTickableObjectsManager::Tick(float DeltaTime)
@@ -159,6 +183,7 @@ void FQuartzTickableObjectsManager::UnsubscribeFromQuartzTick(FQuartzTickableObj
 	QuartzTickSubscribers.RemoveSingleSwap(InObjectToTick);
 }
 
+
 bool UQuartzSubsystem::DoesSupportWorldType(EWorldType::Type WorldType) const
 {
 	return Super::DoesSupportWorldType(WorldType) || WorldType == EWorldType::EditorPreview;
@@ -170,18 +195,22 @@ void UQuartzSubsystem::Tick(float DeltaTime)
 	Super::Tick(DeltaTime);
 	TRACE_CPUPROFILER_EVENT_SCOPE(QuartzSubsystem::Tick);
 
-	ensure(TickableObjectManagerPtr.IsValid());
+	check(TickableObjectManagerPtr);
+	check(ClockManagerDataPtr);
+	TUniquePtr<FScopeLock> Lock(GetPersistentStateScopeLock(GetWorld()));
 
 	PruneStaleProxies();
-	SubsystemClockManager.LowResoultionUpdate(DeltaTime);
+	ClockManagerDataPtr->SubsystemClockManager.LowResoultionUpdate(DeltaTime);
 	TickableObjectManagerPtr->Tick(DeltaTime);
 }
 
 bool UQuartzSubsystem::IsTickable() const
 {
-	ensure(TickableObjectManagerPtr.IsValid());
+	check(TickableObjectManagerPtr);
+	check(ClockManagerDataPtr);
+	TUniquePtr<FScopeLock> Lock(GetPersistentStateScopeLock(GetWorld()));
 
-	const int32 NumClocks = SubsystemClockManager.GetNumClocks();
+	const int32 NumClocks = ClockManagerDataPtr->SubsystemClockManager.GetNumClocks();
 	TRACE_INT_VALUE(TEXT("QuartzSubsystem::NumClocks"), NumClocks);
 
 	// IsTickable() updates unreal insights values
@@ -198,14 +227,14 @@ TStatId UQuartzSubsystem::GetStatId() const
 
  void UQuartzSubsystem::SubscribeToQuartzTick(FQuartzTickableObject * InObjectToTick)
  {
-	ensure(TickableObjectManagerPtr.IsValid());
+	check(TickableObjectManagerPtr.IsValid());
  	TickableObjectManagerPtr->SubscribeToQuartzTick(InObjectToTick);
  }
 
 
  void UQuartzSubsystem::UnsubscribeFromQuartzTick(FQuartzTickableObject * InObjectToTick)
  {
-	ensure(TickableObjectManagerPtr.IsValid());
+	check(TickableObjectManagerPtr.IsValid());
  	TickableObjectManagerPtr->UnsubscribeFromQuartzTick(InObjectToTick);
  }
 
@@ -360,10 +389,12 @@ Audio::FQuartzClockManager* UQuartzSubsystem::GetClockManager(const UObject* Wor
 {
 	// decide if the clock should be managed by the AudioDevice (audio engine) or the Subsystem (this object)
 	Audio::FQuartzClockManager* ClockManager;
-	Audio::FMixerDevice* MixerDevice = GetAudioMixerDeviceUsingWorldContext(WorldContextObject);
-	if(!bUseAudioEngineClockManager || SimulateNoAudioDeviceCvar || !MixerDevice)
+	Audio::FMixerDevice* MixerDevice = GetAudioMixerDevice(WorldContextObject->GetWorld());
+	if(!bUseAudioEngineClockManager || !MixerDevice)
 	{
-		ClockManager = &SubsystemClockManager;
+		check(ClockManagerDataPtr);
+		TUniquePtr<FScopeLock> Lock(GetPersistentStateScopeLock(GetWorld()));
+		ClockManager = &ClockManagerDataPtr->SubsystemClockManager;
 	}
 	else
 	{
@@ -382,7 +413,7 @@ UQuartzClockHandle* UQuartzSubsystem::CreateNewClock(const UObject* WorldContext
 	}
 
 	Audio::FQuartzClockManager* ClockManager = GetClockManager(WorldContextObject, bUseAudioEngineClockManager);
-	ensure(ClockManager); // should have at least fallen back to "this" object as a manager
+	check(ClockManager); // should have at least fallen back to "this" object as a manager
 
 	// numerator of time signature must be >= 1
 	if (InSettings.TimeSignature.NumBeats < 1)
@@ -397,7 +428,9 @@ UQuartzClockHandle* UQuartzSubsystem::CreateNewClock(const UObject* WorldContext
 
 	// if we are not the manager for the clock, it means the FAudioDevice is,
 	// so we hold onto our own copy of the proxy
-	ActiveAudioMixerClockProxies.Add(ClockProxy);
+	check(ClockManagerDataPtr);
+	TUniquePtr<FScopeLock> Lock(GetPersistentStateScopeLock(GetWorld()));
+	ClockManagerDataPtr->ActiveAudioMixerClockProxies.Add(ClockProxy);
 
 	return ClockHandle;
 }
@@ -551,21 +584,21 @@ float UQuartzSubsystem::GetGameThreadToAudioRenderThreadMaxLatency(const UObject
 
 float UQuartzSubsystem::GetAudioRenderThreadToGameThreadAverageLatency()
 {
-	ensure(TickableObjectManagerPtr.IsValid());
+	check(TickableObjectManagerPtr.IsValid());
 	return TickableObjectManagerPtr->GetLifetimeAverageLatency();
 }
 
 
 float UQuartzSubsystem::GetAudioRenderThreadToGameThreadMinLatency()
 {
-	ensure(TickableObjectManagerPtr.IsValid());
+	check(TickableObjectManagerPtr.IsValid());
 	return TickableObjectManagerPtr->GetMinLatency();
 }
 
 
 float UQuartzSubsystem::GetAudioRenderThreadToGameThreadMaxLatency()
 {
-	ensure(TickableObjectManagerPtr.IsValid());
+	check(TickableObjectManagerPtr.IsValid());
 	return TickableObjectManagerPtr->GetMaxLatency();
 }
 
@@ -605,8 +638,10 @@ TWeakPtr<FQuartzTickableObjectsManager> UQuartzSubsystem::GetTickableObjectManag
 
 void UQuartzSubsystem::PruneStaleProxies()
 {
-	PruneStaleProxiesInternal(ActiveExternalClockProxies);
-	PruneStaleProxiesInternal(ActiveAudioMixerClockProxies);
+	check(ClockManagerDataPtr);
+	TUniquePtr<FScopeLock> Lock(GetPersistentStateScopeLock(GetWorld()));
+	PruneStaleProxiesInternal(ClockManagerDataPtr->ActiveExternalClockProxies);
+	PruneStaleProxiesInternal(ClockManagerDataPtr->ActiveAudioMixerClockProxies);
 }
 
 void UQuartzSubsystem::PruneStaleProxiesInternal(TArray<Audio::FQuartzClockProxy>& ContainerToPrune)
@@ -622,12 +657,14 @@ void UQuartzSubsystem::PruneStaleProxiesInternal(TArray<Audio::FQuartzClockProxy
 
 Audio::FQuartzClockProxy* UQuartzSubsystem::FindProxyByName(const FName& ClockName)
 {
-	Audio::FQuartzClockProxy* Result = ActiveAudioMixerClockProxies.FindByKey(ClockName);
+	check(ClockManagerDataPtr);
+	TUniquePtr<FScopeLock> Lock(GetPersistentStateScopeLock(GetWorld()));
+	Audio::FQuartzClockProxy* Result = ClockManagerDataPtr->ActiveAudioMixerClockProxies.FindByKey(ClockName);
 
 	// if the subsystem doesn't have a match, check the externally-registered clock proxies
 	if(!Result)
 	{
-		Result = ActiveExternalClockProxies.FindByKey(ClockName);
+		Result = ClockManagerDataPtr->ActiveExternalClockProxies.FindByKey(ClockName);
 	}
 
 	return Result;
@@ -635,12 +672,15 @@ Audio::FQuartzClockProxy* UQuartzSubsystem::FindProxyByName(const FName& ClockNa
 
 Audio::FQuartzClockProxy const* UQuartzSubsystem::FindProxyByName(const FName& ClockName) const
 {
-	Audio::FQuartzClockProxy const* Result = ActiveAudioMixerClockProxies.FindByKey(ClockName);
+	check(ClockManagerDataPtr);
+	TUniquePtr<FScopeLock> Lock(GetPersistentStateScopeLock(GetWorld()));
+
+	Audio::FQuartzClockProxy const* Result = ClockManagerDataPtr->ActiveAudioMixerClockProxies.FindByKey(ClockName);
 
 	// if the subsystem doesn't have a match, check the externally-registered clock proxies
 	if(!Result)
 	{
-		Result = ActiveExternalClockProxies.FindByKey(ClockName);
+		Result = ClockManagerDataPtr->ActiveExternalClockProxies.FindByKey(ClockName);
 	}
 
 	return Result;
diff --git a/Engine/Source/Runtime/AudioMixer/Public/AudioMixerDevice.h b/Engine/Source/Runtime/AudioMixer/Public/AudioMixerDevice.h
--- a/Engine/Source/Runtime/AudioMixer/Public/AudioMixerDevice.h
+++ b/Engine/Source/Runtime/AudioMixer/Public/AudioMixerDevice.h
@@ -305,6 +305,14 @@ namespace Audio
 		// Clock Manager for quantized event handling on Audio Render Thread
 		FQuartzClockManager QuantizedEventClockManager;
 
+		// Keep a reference alive to UQuartzSubsystem state that needs to persist across level transitions (UWorld destruction
+		TSharedPtr<FPersistentQuartzSubsystemData, ESPMode::ThreadSafe> QuartzSubsystemData { nullptr };
+
+		// Technically, in editor, multiple UQuartz(World)Subsystem's will reference the same FMixerDevice object.
+		// We need to protect around mutation/access of the "shared" state.
+		// (in practice this should be low/zero contention)
+		FCriticalSection QuartzPersistentStateCritSec;
+
 		// Pushes the command to a audio render thread command queue to be executed on render thread
 		void AudioRenderThreadCommand(TFunction<void()> Command);
 
diff --git a/Engine/Source/Runtime/AudioMixer/Public/Quartz/AudioMixerClockManager.h b/Engine/Source/Runtime/AudioMixer/Public/Quartz/AudioMixerClockManager.h
--- a/Engine/Source/Runtime/AudioMixer/Public/Quartz/AudioMixerClockManager.h
+++ b/Engine/Source/Runtime/AudioMixer/Public/Quartz/AudioMixerClockManager.h
@@ -10,6 +10,7 @@ namespace Audio
 	// forwards
 	class FMixerDevice;
 	class FQuartzClock;
+	class FQuartzClockManager;
 
 	// Class that owns, updates, and provides access to all active clocks
 	// All methods are thread-safe. The method locks if it returns a value, and stages a command if it returns void
@@ -131,4 +132,15 @@ namespace Audio
 		// allow a clock that is queuing a command directly use FindClock() to retrieve the TSharedPtr<FQuartzClock>
 		friend void FQuartzClock::AddQuantizedCommand(FQuartzQuantizedCommandInitInfo& InQuantizationCommandInitInfo);
 	};
+
+	// data that the UQuartzSubsystem needs to persist on the AudioDevice across UWorld shutdown/startup
+	struct AUDIOMIXER_API FPersistentQuartzSubsystemData
+	{
+		// internal clock manager for game-thread-ticked clocks
+		FQuartzClockManager SubsystemClockManager;
+
+		// array of active clock handles (update FindProxyByName() if more are added later)
+		TArray<Audio::FQuartzClockProxy> ActiveExternalClockProxies;
+		TArray<Audio::FQuartzClockProxy> ActiveAudioMixerClockProxies;
+	};
 } // namespace Audio
diff --git a/Engine/Source/Runtime/AudioMixer/Public/Quartz/QuartzSubsystem.h b/Engine/Source/Runtime/AudioMixer/Public/Quartz/QuartzSubsystem.h
--- a/Engine/Source/Runtime/AudioMixer/Public/Quartz/QuartzSubsystem.h
+++ b/Engine/Source/Runtime/AudioMixer/Public/Quartz/QuartzSubsystem.h
@@ -4,8 +4,6 @@
 
 #include "CoreMinimal.h"
 #include "UObject/ObjectMacros.h"
-#include "UObject/Object.h"
-#include "UObject/Class.h"
 #include "Subsystems/WorldSubsystem.h"
 #include "Quartz/AudioMixerClockManager.h"
 #include "Sound/QuartzQuantizationUtilities.h"
@@ -21,6 +19,7 @@ namespace Audio
 
 	template<class ListenerType>
 	class TQuartzShareableCommandQueue;
+
 }
 
 class FQuartzTickableObject;
@@ -57,6 +56,8 @@ public:
 	virtual ~UQuartzSubsystem() override = default;
 
 	//~ Begin UWorldSubsystem Interface
+	virtual void Initialize(FSubsystemCollectionBase& Collection) override;
+	virtual void Deinitialize() override;
 	virtual bool DoesSupportWorldType(EWorldType::Type WorldType) const override;
 	void virtual BeginDestroy() override;
 	//~ End UWorldSubsystem Interface
@@ -180,15 +181,12 @@ private:
 	void PruneStaleProxies();
 	static void PruneStaleProxiesInternal(TArray<Audio::FQuartzClockProxy>& ContainerToPrune);
 
-	// internal clock manager for game-thread-ticked clocks
-	Audio::FQuartzClockManager SubsystemClockManager;
 
 	// sharable tickable object manager to allow for non-UObject subscription / un-subscription
 	TSharedPtr<FQuartzTickableObjectsManager> TickableObjectManagerPtr { MakeShared<FQuartzTickableObjectsManager>() };
 
-	// array of active clock handles (update FindProxyByName() if more are added later)
-	TArray<Audio::FQuartzClockProxy> ActiveExternalClockProxies;
-	TArray<Audio::FQuartzClockProxy> ActiveAudioMixerClockProxies;
+	// Clock manager/proxy-related data that lives on the AudioDevice for persistence.
+	TSharedPtr<Audio::FPersistentQuartzSubsystemData> ClockManagerDataPtr { nullptr };
 
 	// helpers
 	Audio::FQuartzClockProxy* FindProxyByName(const FName& ClockName);
