Micro optimization in AudioMixerSourceManager

diff --git a/Engine/Source/Runtime/AudioMixer/Private/AudioMixerSourceManager.cpp b/Engine/Source/Runtime/AudioMixer/Private/AudioMixerSourceManager.cpp
--- a/Engine/Source/Runtime/AudioMixer/Private/AudioMixerSourceManager.cpp
+++ b/Engine/Source/Runtime/AudioMixer/Private/AudioMixerSourceManager.cpp
@@ -2086,19 +2086,22 @@ namespace Audio
 					SourceInfo.bHasStarted = true;
 
 					// Update the PrevFrameIndex value for the source based on alpha value
-					while (SourceInfo.CurrentFrameAlpha >= 1.0f)
+					if (SourceInfo.CurrentFrameAlpha >= 1.0f)
 					{
 						// Our inter-frame alpha lerping value is causing us to read new source frames
 						bReadNextSample = true;
+						
+						const float Delta = FMath::FloorToFloat(SourceInfo.CurrentFrameAlpha);
+						const int DeltaInt = (int)Delta;
 
 						// Bump up the current frame index
-						SourceInfo.CurrentFrameIndex++;
+						SourceInfo.CurrentFrameIndex += DeltaInt;
 
 						// Bump up the frames played -- this is tracking the total frames in source file played
 						// CurrentFrameIndex can wrap for looping sounds so won't be accurate in that case
-						SourceInfo.NumFramesPlayed++;
+						SourceInfo.NumFramesPlayed += DeltaInt;
 
-						SourceInfo.CurrentFrameAlpha -= 1.0f;
+						SourceInfo.CurrentFrameAlpha -= Delta;
 					}
 
 					// If our alpha parameter caused us to jump to a new source frame, we need
