[Quartz] Fix for queries hitting default-constructed data (game-thread copy of clock state) if the clock hasn't started yet.

diff --git a/Engine/Source/Runtime/AudioMixer/Private/Quartz/AudioMixerClock.cpp b/Engine/Source/Runtime/AudioMixer/Private/Quartz/AudioMixerClock.cpp
--- a/Engine/Source/Runtime/AudioMixer/Private/Quartz/AudioMixerClock.cpp
+++ b/Engine/Source/Runtime/AudioMixer/Private/Quartz/AudioMixerClock.cpp
@@ -120,6 +120,8 @@ namespace Audio
 		{
 			Metronome.SetSampleRate(HeadlessClockSampleRateCvar);
 		}
+
+		UpdateCachedState();
 	}
 
 	FQuartzClock::~FQuartzClock()
@@ -162,11 +164,14 @@ namespace Audio
 				Command.NumFramesUntilExec = NumFramesLeft + Ratio * (Command.NumFramesUntilExec - NumFramesLeft);
 			}
 		}
+
+		UpdateCachedState();
 	}
 
 	void FQuartzClock::ChangeTimeSignature(const FQuartzTimeSignature& InNewTimeSignature)
 	{
 		Metronome.SetTimeSignature(InNewTimeSignature);
+		UpdateCachedState();
 	}
 
 	void FQuartzClock::Resume()
@@ -370,6 +375,8 @@ namespace Audio
 
 		// update Tick Rate
 		Metronome.SetSampleRate(InNewSampleRate);
+
+		UpdateCachedState();
 	}
 
 	bool FQuartzClock::IgnoresFlush() const
