[Quartz] Adding (verbose) logging for event scheduling.

diff --git a/Engine/Source/Runtime/AudioMixer/Private/Quartz/AudioMixerClock.cpp b/Engine/Source/Runtime/AudioMixer/Private/Quartz/AudioMixerClock.cpp
--- a/Engine/Source/Runtime/AudioMixer/Private/Quartz/AudioMixerClock.cpp
+++ b/Engine/Source/Runtime/AudioMixer/Private/Quartz/AudioMixerClock.cpp
@@ -262,6 +262,7 @@ namespace Audio
 	void FQuartzClock::LowResolutionTick(float InDeltaTimeSeconds)
 	{
 		TRACE_CPUPROFILER_EVENT_SCOPE(QuartzClock::Tick_LowRes);
+		UE_LOG(LogAudioQuartz, Verbose, TEXT("Quartz Clock Tick (low-res): %s"), *Name.ToString());
 		PreTickCommands->PumpCommandQueue(this);
 		Tick(static_cast<int32>(InDeltaTimeSeconds * Metronome.GetTickRate().GetSampleRate()));
 	}
@@ -269,8 +270,10 @@ namespace Audio
 	void FQuartzClock::Tick(int32 InNumFramesUntilNextTick)
 	{
 		TRACE_CPUPROFILER_EVENT_SCOPE(QuartzClock::Tick);
-
 		TRACE_CPUPROFILER_EVENT_SCOPE(QuartzClock::GameThreadCommands);
+
+		UE_LOG(LogAudioQuartz, Verbose, TEXT("Quartz Clock Tick: %s"), *Name.ToString());
+		
 		PreTickCommands->PumpCommandQueue(this);
 
 		TRACE_CPUPROFILER_EVENT_SCOPE(QuartzClock::TickInternal);
@@ -438,6 +441,12 @@ namespace Audio
 		// if this is un-quantized, execute immediately (even if the clock is paused)
 		if (InQuantizationBondary.Quantization == EQuartzCommandQuantization::None)
 		{
+			UE_LOG(LogAudioQuartz, Display, TEXT("Quartz Command:(%s) | Deadline (frames):[%i] | Boundary: [%s]")
+				, *InNewEvent->GetCommandName().ToString()
+				, FramesUntilExec
+				, *InQuantizationBondary.ToString()
+				);
+			
 			InNewEvent->AboutToStart();
 			InNewEvent->OnFinalCallback(0);
 			return;
@@ -445,7 +454,18 @@ namespace Audio
 
 		// get number of frames until event (assuming we are at frame 0)
 		FramesUntilExec = FMath::RoundToInt(Metronome.GetFramesUntilBoundary(InQuantizationBondary)); // query metronome (round result to int)
-		FramesUntilExec = FMath::Max(0, InNewEvent->OverrideFramesUntilExec(FramesUntilExec)); // allow command to override the deadline (clamp result)
+		const int32 OverriddenFramesUntilExec = FMath::Max(0, InNewEvent->OverrideFramesUntilExec(FramesUntilExec)); // allow command to override the deadline (clamp result)
+		const bool bOverridden = (FramesUntilExec != OverriddenFramesUntilExec);
+
+		UE_LOG(LogAudioQuartz, Display, TEXT("Quartz Command:(%s) | Deadline (frames):[%i%s] | Boundary: [%s]")
+			, *InNewEvent->GetCommandName().ToString()
+			, OverriddenFramesUntilExec
+			, bOverridden? *FString::Printf(TEXT("(overridden from %i)"), FramesUntilExec) : TEXT("")
+			, *InQuantizationBondary.ToString()
+			);
+
+		// after the log, use tho Overridden value
+		FramesUntilExec = OverriddenFramesUntilExec;
 
 		// finalize the requested subscriber offsets and notify the command of their deadline
 		InNewEvent->OnScheduled(Metronome.GetTickRate());
diff --git a/Engine/Source/Runtime/Engine/Classes/Sound/QuartzQuantizationUtilities.cpp b/Engine/Source/Runtime/Engine/Classes/Sound/QuartzQuantizationUtilities.cpp
--- a/Engine/Source/Runtime/Engine/Classes/Sound/QuartzQuantizationUtilities.cpp
+++ b/Engine/Source/Runtime/Engine/Classes/Sound/QuartzQuantizationUtilities.cpp
@@ -155,6 +155,85 @@ void FQuartLatencyTracker::DigestQueue()
 	}
 }
 
+static FString EnumToString(EQuartzCommandQuantization inEnum)
+{
+	switch (inEnum)
+	{
+	case EQuartzCommandQuantization::Bar:
+		return(TEXT("Bar"));
+	case EQuartzCommandQuantization::Beat:
+		return(TEXT("Beat"));
+		
+	case EQuartzCommandQuantization::ThirtySecondNote:
+		return(TEXT("ThirtySecondNote"));
+	case EQuartzCommandQuantization::SixteenthNote:
+		return(TEXT("SixteenthNote"));
+	case EQuartzCommandQuantization::EighthNote:
+		return(TEXT("EighthNote"));
+	case EQuartzCommandQuantization::QuarterNote:
+		return(TEXT("QuarterNote"));
+	case EQuartzCommandQuantization::HalfNote:
+		return(TEXT("HalfNote"));
+	case EQuartzCommandQuantization::WholeNote:
+		return(TEXT("WholeNote"));
+
+	case EQuartzCommandQuantization::DottedSixteenthNote:
+		return(TEXT("DottedSixteenthNote"));
+	case EQuartzCommandQuantization::DottedEighthNote:
+		return(TEXT("DottedEighthNote"));
+	case EQuartzCommandQuantization::DottedQuarterNote:
+		return(TEXT("DottedQuarterNote"));
+	case EQuartzCommandQuantization::DottedHalfNote:
+		return(TEXT("DottedHalfNote"));
+	case EQuartzCommandQuantization::DottedWholeNote:
+		return(TEXT("WholeNote"));
+
+	case EQuartzCommandQuantization::SixteenthNoteTriplet:
+		return(TEXT("SixteenthNoteTriplet"));
+	case EQuartzCommandQuantization::EighthNoteTriplet:
+		return(TEXT("EighthNoteTriplet"));
+	case EQuartzCommandQuantization::QuarterNoteTriplet:
+		return(TEXT("QuarterNoteTriplet"));
+	case EQuartzCommandQuantization::HalfNoteTriplet:
+		return(TEXT("HalfNoteTriplet"));
+
+	case EQuartzCommandQuantization::None:
+		return(TEXT("None"));
+		
+	default:
+		return {};
+	}
+	
+}
+
+static FString EnumToString(EQuarztQuantizationReference inEnum)
+{
+	switch (inEnum)
+	{
+	case EQuarztQuantizationReference::BarRelative:
+			return(TEXT("BarRelative"));
+	case EQuarztQuantizationReference::TransportRelative:
+		return(TEXT("TransportRelative"));
+	case EQuarztQuantizationReference::CurrentTimeRelative:
+		return(TEXT("CurrentTimeRelative"));
+		
+		default:
+			return {};
+	}
+	
+}
+
+FString FQuartzQuantizationBoundary::ToString() const 
+{
+	FString String;
+
+	String
+	.Append(*FString::Printf(TEXT("Quant:(%s) - "), *EnumToString(Quantization)))
+	.Append(*FString::Printf(TEXT("Mult:(%f) - "), Multiplier))
+	.Append(*FString::Printf(TEXT("Ref:(%s)"), *EnumToString(CountingReferencePoint)));
+	
+	return String;
+}
 
 
 namespace Audio
diff --git a/Engine/Source/Runtime/Engine/Classes/Sound/QuartzQuantizationUtilities.h b/Engine/Source/Runtime/Engine/Classes/Sound/QuartzQuantizationUtilities.h
--- a/Engine/Source/Runtime/Engine/Classes/Sound/QuartzQuantizationUtilities.h
+++ b/Engine/Source/Runtime/Engine/Classes/Sound/QuartzQuantizationUtilities.h
@@ -526,6 +526,8 @@ struct ENGINE_API FQuartzQuantizationBoundary
 		, CountingReferencePoint(InReferencePoint)
 		, bFireOnClockStart(bInFireOnClockStart)
 	{}
+
+	FString ToString() const;
 }; // struct FQuartzQuantizationBoundary
 
 
