[Quartz] Reduce new log to Verbose instead of Display to avoid spam.

diff --git a/Engine/Source/Runtime/AudioMixer/Private/Quartz/AudioMixerClock.cpp b/Engine/Source/Runtime/AudioMixer/Private/Quartz/AudioMixerClock.cpp
--- a/Engine/Source/Runtime/AudioMixer/Private/Quartz/AudioMixerClock.cpp
+++ b/Engine/Source/Runtime/AudioMixer/Private/Quartz/AudioMixerClock.cpp
@@ -441,7 +441,7 @@ namespace Audio
 		// if this is un-quantized, execute immediately (even if the clock is paused)
 		if (InQuantizationBondary.Quantization == EQuartzCommandQuantization::None)
 		{
-			UE_LOG(LogAudioQuartz, Display, TEXT("Quartz Command:(%s) | Deadline (frames):[%i] | Boundary: [%s]")
+			UE_LOG(LogAudioQuartz, Verbose, TEXT("Quartz Command:(%s) | Deadline (frames):[%i] | Boundary: [%s]")
 				, *InNewEvent->GetCommandName().ToString()
 				, FramesUntilExec
 				, *InQuantizationBondary.ToString()
@@ -457,7 +457,7 @@ namespace Audio
 		const int32 OverriddenFramesUntilExec = FMath::Max(0, InNewEvent->OverrideFramesUntilExec(FramesUntilExec)); // allow command to override the deadline (clamp result)
 		const bool bOverridden = (FramesUntilExec != OverriddenFramesUntilExec);
 
-		UE_LOG(LogAudioQuartz, Display, TEXT("Quartz Command:(%s) | Deadline (frames):[%i%s] | Boundary: [%s]")
+		UE_LOG(LogAudioQuartz, Verbose, TEXT("Quartz Command:(%s) | Deadline (frames):[%i%s] | Boundary: [%s]")
 			, *InNewEvent->GetCommandName().ToString()
 			, OverriddenFramesUntilExec
 			, bOverridden? *FString::Printf(TEXT("(overridden from %i)"), FramesUntilExec) : TEXT("")
