Checkpoint: Pipe source data to dashboard UI listview (Part 1)
- Add trace provider, module, & analyzer and start piping data through it as proof-of-concept
#rnx

diff --git a/Engine/Plugins/Animation/GameplayInsights/Source/GameplayInsights/Public/IGameplayInsightsModule.h b/Engine/Plugins/Animation/GameplayInsights/Source/GameplayInsights/Public/IGameplayInsightsModule.h
--- a/Engine/Plugins/Animation/GameplayInsights/Source/GameplayInsights/Public/IGameplayInsightsModule.h
+++ b/Engine/Plugins/Animation/GameplayInsights/Source/GameplayInsights/Public/IGameplayInsightsModule.h
@@ -14,7 +14,7 @@ namespace TraceServices
 	class IAnalysisSession;
 }
 
-class IGameplayInsightsModule
+class GAMEPLAYINSIGHTS_API IGameplayInsightsModule
 	: public IModuleInterface
 {
 public:
diff --git a/Engine/Source/Developer/TraceAnalysis/Public/Trace/Analyzer.h b/Engine/Source/Developer/TraceAnalysis/Public/Trace/Analyzer.h
--- a/Engine/Source/Developer/TraceAnalysis/Public/Trace/Analyzer.h
+++ b/Engine/Source/Developer/TraceAnalysis/Public/Trace/Analyzer.h
@@ -253,7 +253,7 @@ public:
 		double AsSecondsAbsolute(int64 DurationCycles64) const;
 	};
 
-	struct FOnEventContext
+	struct TRACEANALYSIS_API FOnEventContext
 	{
 		const FThreadInfo&	ThreadInfo;
 		const FEventTime&	EventTime;
diff --git a/Engine/Source/Runtime/AudioMixer/AudioMixer.Build.cs b/Engine/Source/Runtime/AudioMixer/AudioMixer.Build.cs
--- a/Engine/Source/Runtime/AudioMixer/AudioMixer.Build.cs
+++ b/Engine/Source/Runtime/AudioMixer/AudioMixer.Build.cs
@@ -33,7 +33,8 @@ namespace UnrealBuildTool.Rules
 					"AudioExtensions",
 					"AudioLinkCore",
 					"AudioLinkEngine",
-					"HeadMountedDisplay"
+					"HeadMountedDisplay",
+					"TraceLog"
 				}
 			);
 
diff --git a/Engine/Source/Runtime/AudioMixer/Private/AudioMixerSource.cpp b/Engine/Source/Runtime/AudioMixer/Private/AudioMixerSource.cpp
--- a/Engine/Source/Runtime/AudioMixer/Private/AudioMixerSource.cpp
+++ b/Engine/Source/Runtime/AudioMixer/Private/AudioMixerSource.cpp
@@ -1,6 +1,7 @@
 // Copyright Epic Games, Inc. All Rights Reserved.
-
 #include "AudioMixerSource.h"
+
+#include "AudioDefines.h"
 #include "AudioMixerSourceBuffer.h"
 #include "ActiveSound.h"
 #include "AudioMixerSourceBuffer.h"
@@ -15,10 +16,38 @@
 #include "Sound/SoundModulationDestination.h"
 #include "Misc/ScopeRWLock.h"
 #include "Templates/Function.h"
+#include "Trace/Trace.h"
+
 
-// Link to "Audio" profiling category
 CSV_DECLARE_CATEGORY_MODULE_EXTERN(AUDIOMIXERCORE_API, Audio);
 
+#if ENABLE_AUDIO_TRACE && AUDIO_MIXER_CPUPROFILERTRACE_ENABLED
+UE_TRACE_EVENT_BEGIN(Audio, MixerSourceStart)
+	UE_TRACE_EVENT_FIELD(uint64, Timestamp)
+	UE_TRACE_EVENT_FIELD(int32, SourceId)
+	UE_TRACE_EVENT_FIELD(UE::Trace::WideString, Name)
+UE_TRACE_EVENT_END()
+
+UE_TRACE_EVENT_BEGIN(Audio, MixerSourceStop)
+	UE_TRACE_EVENT_FIELD(uint64, Timestamp)
+	UE_TRACE_EVENT_FIELD(int32, SourceId)
+	UE_TRACE_EVENT_FIELD(UE::Trace::WideString, Name)
+UE_TRACE_EVENT_END()
+
+UE_TRACE_EVENT_BEGIN(Audio, MixerSourceUpdate)
+	UE_TRACE_EVENT_FIELD(uint64, Timestamp)
+	UE_TRACE_EVENT_FIELD(int32, SourceId)
+	UE_TRACE_EVENT_FIELD(UE::Trace::WideString, Name)
+	UE_TRACE_EVENT_FIELD(float, Pitch)
+	UE_TRACE_EVENT_FIELD(float, LPFreq)
+	UE_TRACE_EVENT_FIELD(float, HPFreq)
+	UE_TRACE_EVENT_FIELD(float, Envelope)
+	UE_TRACE_EVENT_FIELD(float, Distance)
+	UE_TRACE_EVENT_FIELD(float, DistanceAttenuation)
+UE_TRACE_EVENT_END()
+#endif // ENABLE_AUDIO_TRACE && AUDIO_MIXER_CPUPROFILERTRACE_ENABLED
+
+
 static int32 UseListenerOverrideForSpreadCVar = 0;
 FAutoConsoleVariableRef CVarUseListenerOverrideForSpread(
 	TEXT("au.UseListenerOverrideForSpread"),
@@ -70,7 +99,8 @@ namespace Audio
 
 			return SoundClass;
 		}
-	}
+	} // namespace MixerSourcePrivate
+
 	namespace ModulationUtils
 	{
 		void MixInRoutedValue(const FModulationParameter& InParam, float& InOutValueA, float InValueB)
@@ -1011,9 +1041,7 @@ namespace Audio
 
 		UpdateChannelMaps();
 
-#if ENABLE_AUDIO_DEBUG
-		Audio::FAudioDebugger::DrawDebugInfo(*this);
-#endif // ENABLE_AUDIO_DEBUG
+		UpdateTraceInfo();
 	}
 
 	bool FMixerSource::PrepareForInitialization(FWaveInstance* InWaveInstance)
@@ -1224,6 +1252,23 @@ namespace Audio
 		if (MixerSourceVoice && InitializationState == EMixerSourceInitializationState::Initialized)
 		{
 			MixerSourceVoice->Play();
+
+#if ENABLE_AUDIO_TRACE && AUDIO_MIXER_CPUPROFILERTRACE_ENABLED
+			const bool bChannelEnabled = UE_TRACE_CHANNELEXPR_IS_ENABLED(AudioMixerChannel);
+			if (bChannelEnabled && WaveInstance)
+			{
+				int32 TraceSourceId = INDEX_NONE;
+				if (MixerSourceVoice)
+				{
+					TraceSourceId = MixerSourceVoice->GetSourceId();
+				}
+
+				UE_TRACE_LOG(Audio, MixerSourceStart, AudioMixerChannel)
+					<< MixerSourceStart.Timestamp(FPlatformTime::Cycles64())
+					<< MixerSourceStart.SourceId(TraceSourceId)
+					<< MixerSourceStart.Name(*WaveInstance->GetName());
+			}
+#endif // ENABLE_AUDIO_TRACE
 		}
 
 		bIsStopping = false;
@@ -1308,6 +1353,23 @@ namespace Audio
 		{
 			if (MixerSourceVoice && Playing)
 			{
+#if ENABLE_AUDIO_TRACE && AUDIO_MIXER_CPUPROFILERTRACE_ENABLED
+				const bool bChannelEnabled = UE_TRACE_CHANNELEXPR_IS_ENABLED(AudioMixerChannel);
+				if (bChannelEnabled)
+				{
+					int32 TraceSourceId = INDEX_NONE;
+					if (MixerSourceVoice)
+					{
+						TraceSourceId = MixerSourceVoice->GetSourceId();
+					}
+
+					UE_TRACE_LOG(Audio, MixerSourceStop, AudioMixerChannel)
+						<< MixerSourceStop.Timestamp(FPlatformTime::Cycles64())
+						<< MixerSourceStop.SourceId(TraceSourceId)
+						<< MixerSourceStop.Name(*WaveInstance->GetName());
+				}
+#endif // ENABLE_AUDIO_TRACE
+
 				MixerSourceVoice->Stop();
 			}
 
@@ -1846,6 +1908,45 @@ namespace Audio
 		bPrevAllowedSpatializationSetting = IsSpatializationCVarEnabled();
 	}
 
+	void FMixerSource::UpdateTraceInfo()
+	{
+#if ENABLE_AUDIO_DEBUG
+		Audio::FAudioDebugger::DrawDebugInfo(*this);
+#endif // ENABLE_AUDIO_DEBUG
+
+#if ENABLE_AUDIO_TRACE && AUDIO_MIXER_CPUPROFILERTRACE_ENABLED
+		const bool bChannelEnabled = UE_TRACE_CHANNELEXPR_IS_ENABLED(AudioMixerChannel);
+		if (bChannelEnabled && WaveInstance)
+		{
+			int32 TraceSourceId = INDEX_NONE;
+			float TraceVolume = 1.0f;
+			float TracePitch  = 1.0f;
+			float TraceLPFreq = MIN_FILTER_FREQUENCY;
+			float TraceHPFreq = MAX_FILTER_FREQUENCY;
+			float TraceEnv = 1.0f;
+			float TraceDistAttenuation = 0.0f;
+
+			if (MixerSourceVoice)
+			{
+				TraceSourceId = MixerSourceVoice->GetSourceId();
+				TraceEnv = MixerSourceVoice->GetEnvelopeValue();
+				TraceDistAttenuation = MixerSourceVoice->GetDistanceAttenuation();
+			}
+
+			UE_TRACE_LOG(Audio, MixerSourceUpdate, AudioMixerChannel)
+				<< MixerSourceUpdate.Timestamp(FPlatformTime::Cycles64())
+				<< MixerSourceUpdate.SourceId(TraceSourceId)
+				<< MixerSourceUpdate.Name(*WaveInstance->GetName())
+				<< MixerSourceUpdate.Pitch(TracePitch)
+				<< MixerSourceUpdate.LPFreq(TraceLPFreq)
+				<< MixerSourceUpdate.HPFreq(TraceHPFreq)
+				<< MixerSourceUpdate.Envelope(TraceEnv)
+				<< MixerSourceUpdate.DistanceAttenuation(TraceDistAttenuation);
+
+		}
+#endif // ENABLE_AUDIO_TRACE
+	}
+
 	bool FMixerSource::ComputeMonoChannelMap(Audio::FAlignedFloatBuffer& OutChannelMap)
 	{
 		if (IsUsingObjectBasedSpatialization())
diff --git a/Engine/Source/Runtime/AudioMixer/Private/AudioMixerSource.h b/Engine/Source/Runtime/AudioMixer/Private/AudioMixerSource.h
--- a/Engine/Source/Runtime/AudioMixer/Private/AudioMixerSource.h
+++ b/Engine/Source/Runtime/AudioMixer/Private/AudioMixerSource.h
@@ -85,6 +85,9 @@ namespace Audio
 		/** Updates the channel map of the sound if its a 3d sound.*/
 		void UpdateChannelMaps();
 
+		/** Forwards trace info to trace system if enabled */
+		void UpdateTraceInfo();
+
 		/** Computes the mono-channel map. */
 		bool ComputeMonoChannelMap(Audio::FAlignedFloatBuffer& OutChannelMap);
 
diff --git a/Engine/Source/Runtime/AudioMixer/Private/AudioMixerSourceManager.cpp b/Engine/Source/Runtime/AudioMixer/Private/AudioMixerSourceManager.cpp
--- a/Engine/Source/Runtime/AudioMixer/Private/AudioMixerSourceManager.cpp
+++ b/Engine/Source/Runtime/AudioMixer/Private/AudioMixerSourceManager.cpp
@@ -2,10 +2,12 @@
 
 #include "AudioMixerSourceManager.h"
 
+#include "AudioDefines.h"
 #include "AudioMixerSourceBuffer.h"
 #include "AudioMixerDevice.h"
 #include "AudioMixerSourceVoice.h"
 #include "AudioMixerSubmix.h"
+#include "AudioMixerTrace.h"
 #include "AudioThread.h"
 #include "DSP/FloatArrayMath.h"
 #include "IAudioExtensionPlugin.h"
@@ -16,6 +18,8 @@
 #include "Async/Async.h"
 #include "ProfilingDebugging/CountersTrace.h"
 #include "Stats/Stats.h"
+#include "Trace/Trace.h"
+
 
 // Link to "Audio" profiling category
 CSV_DECLARE_CATEGORY_MODULE_EXTERN(AUDIOMIXERCORE_API, Audio);
@@ -164,6 +168,14 @@ DEFINE_STAT(STAT_AudioMixerSourceEffectBuffers);
 DEFINE_STAT(STAT_AudioMixerSourceManagerUpdate);
 DEFINE_STAT(STAT_AudioMixerSourceOutputBuffers);
 
+#if ENABLE_AUDIO_TRACE && AUDIO_MIXER_CPUPROFILERTRACE_ENABLED
+UE_TRACE_EVENT_BEGIN(Audio, MixerSourceVolume)
+	UE_TRACE_EVENT_FIELD(uint64, Timestamp)
+	UE_TRACE_EVENT_FIELD(int32, SourceId)
+	UE_TRACE_EVENT_FIELD(float, Volume)
+UE_TRACE_EVENT_END()
+#endif // ENABLE_AUDIO_TRACE && AUDIO_MIXER_CPUPROFILERTRACE_ENABLED
+
 namespace Audio
 {
 	/*************************************************************************
@@ -2519,6 +2531,16 @@ namespace Audio
 				}
 
 				Audio::ArrayFade(PreDistanceAttenBufferView, VolumeStart, VolumeDestination);
+#if ENABLE_AUDIO_TRACE && AUDIO_MIXER_CPUPROFILERTRACE_ENABLED
+				if (!FMath::IsNearlyEqual(VolumeStart, VolumeDestination))
+				{
+					const bool bChannelEnabled = UE_TRACE_CHANNELEXPR_IS_ENABLED(AudioMixerChannel);
+					UE_TRACE_LOG(Audio, MixerSourceVolume, AudioMixerChannel)
+						<< MixerSourceVolume.Timestamp(FPlatformTime::Cycles64())
+						<< MixerSourceVolume.SourceId(SourceId)
+						<< MixerSourceVolume.Volume(VolumeDestination);
+				}
+#endif // ENABLE_AUDIO_TRACE && AUDIO_MIXER_CPUPROFILERTRACE_ENABLED
 			}
 			SourceInfo.VolumeSourceStart = SourceInfo.VolumeSourceDestination;
 
diff --git a/Engine/Source/Runtime/AudioMixer/Private/AudioMixerSourceManager.h b/Engine/Source/Runtime/AudioMixer/Private/AudioMixerSourceManager.h
--- a/Engine/Source/Runtime/AudioMixer/Private/AudioMixerSourceManager.h
+++ b/Engine/Source/Runtime/AudioMixer/Private/AudioMixerSourceManager.h
@@ -227,7 +227,7 @@ namespace Audio
 		void SetLPFFrequency(const int32 SourceId, const float Frequency);
 		void SetHPFFrequency(const int32 SourceId, const float Frequency);
 
-		// Sets base (i.e. carrier) frequency of modulateable parameters
+		// Sets base (i.e. carrier) frequency of modulatable parameters
 		void SetModPitch(const int32 SourceId, const float InModPitch);
 		void SetModVolume(const int32 SourceId, const float InModVolume);
 		void SetModLPFFrequency(const int32 SourceId, const float InModFrequency);
@@ -235,7 +235,6 @@ namespace Audio
 
 		void SetSourceBufferListener(const int32 SourceId, FSharedISourceBufferListenerPtr& InSourceBufferListener, bool InShouldSourceBufferListenerZeroBuffer);
 
-
 		void SetListenerTransforms(const TArray<FTransform>& ListenerTransforms);
 		const TArray<FTransform>* GetListenerTransforms() const;
 
diff --git a/Engine/Source/Runtime/AudioMixer/Private/AudioMixerSourceVoice.cpp b/Engine/Source/Runtime/AudioMixer/Private/AudioMixerSourceVoice.cpp
--- a/Engine/Source/Runtime/AudioMixer/Private/AudioMixerSourceVoice.cpp
+++ b/Engine/Source/Runtime/AudioMixer/Private/AudioMixerSourceVoice.cpp
@@ -248,6 +248,21 @@ namespace Audio
 		SourceManager->StopFade(SourceId, NumFrames);
 	}
 
+	int32 FMixerSourceVoice::GetSourceId() const
+	{
+		return SourceId;
+	}
+
+	float FMixerSourceVoice::GetDistanceAttenuation() const
+	{
+		return DistanceAttenuation;
+	}
+
+	float FMixerSourceVoice::GetDistance() const
+	{
+		return Distance;
+	}
+
 	void FMixerSourceVoice::Pause()
 	{
 		AUDIO_MIXER_CHECK_GAME_THREAD(MixerDevice);
diff --git a/Engine/Source/Runtime/AudioMixer/Private/AudioMixerSourceVoice.h b/Engine/Source/Runtime/AudioMixer/Private/AudioMixerSourceVoice.h
--- a/Engine/Source/Runtime/AudioMixer/Private/AudioMixerSourceVoice.h
+++ b/Engine/Source/Runtime/AudioMixer/Private/AudioMixerSourceVoice.h
@@ -80,8 +80,14 @@ namespace Audio
 		// Does a faded stop (to avoid discontinuity)
 		void StopFade(int32 NumFrames);
 
-		// Get the source's Id
-		int32 GetSourceId() const { return SourceId; }
+		// Returns the source's Id
+		int32 GetSourceId() const;
+
+		// Returns the source's distance attenuation
+		float GetDistanceAttenuation() const;
+
+		// Returns the source's distance from the closest listener
+		float GetDistance() const;
 
 		// Queries if the voice is playing
 		bool IsPlaying() const;
diff --git a/Engine/Source/Runtime/AudioMixerCore/AudioMixerCore.Build.cs b/Engine/Source/Runtime/AudioMixerCore/AudioMixerCore.Build.cs
--- a/Engine/Source/Runtime/AudioMixerCore/AudioMixerCore.Build.cs
+++ b/Engine/Source/Runtime/AudioMixerCore/AudioMixerCore.Build.cs
@@ -12,7 +12,8 @@ namespace UnrealBuildTool.Rules
 				new string[]
                 {
                     "Core",
-					"SignalProcessing"
+					"SignalProcessing",
+					"TraceLog"
                 }
 			);
 		}
diff --git a/Engine/Source/Runtime/AudioMixerCore/Private/AudioMixer.cpp b/Engine/Source/Runtime/AudioMixerCore/Private/AudioMixer.cpp
--- a/Engine/Source/Runtime/AudioMixerCore/Private/AudioMixer.cpp
+++ b/Engine/Source/Runtime/AudioMixerCore/Private/AudioMixer.cpp
@@ -1,6 +1,7 @@
 // Copyright Epic Games, Inc. All Rights Reserved.
-
 #include "AudioMixer.h"
+
+#include "AudioDefines.h"
 #include "DSP/BufferVectorOperations.h"
 #include "DSP/FloatArrayMath.h"
 #include "HAL/RunnableThread.h"
@@ -12,6 +13,8 @@
 #include "Misc/ScopeTryLock.h"
 #include "ProfilingDebugging/CsvProfiler.h"
 #include "ProfilingDebugging/ScopedTimers.h"
+#include "Trace/Trace.h"
+
 
 // Defines the "Audio" category in the CSV profiler.
 // This should only be defined here. Modules who wish to use this category should contain the line
@@ -19,6 +22,11 @@
 //
 CSV_DEFINE_CATEGORY_MODULE(AUDIOMIXERCORE_API, Audio, false);
 
+#if ENABLE_AUDIO_TRACE
+UE_TRACE_CHANNEL_DEFINE(AudioMixerChannel);
+#endif // ENABLE_AUDIO_TRACE
+
+
 // Command to enable logging to display accurate audio render times
 static int32 LogRenderTimesCVar = 0;
 FAutoConsoleVariableRef CVarLogRenderTimes(
diff --git a/Engine/Source/Runtime/AudioMixerCore/Public/AudioDefines.h b/Engine/Source/Runtime/AudioMixerCore/Public/AudioDefines.h
--- a/Engine/Source/Runtime/AudioMixerCore/Public/AudioDefines.h
+++ b/Engine/Source/Runtime/AudioMixerCore/Public/AudioDefines.h
@@ -2,6 +2,7 @@
 #pragma once
 
 #include "HAL/Platform.h"
+#include "Trace/Trace.h"
 
 
 /**
@@ -9,11 +10,13 @@
  */
 #define MAX_AUDIOCHANNELS				64
 
+
 /**
  * Length of sound in seconds to be considered as looping forever
  */
 #define INDEFINITELY_LOOPING_DURATION	10000.0f
 
+
 /**
  * Some defaults to help cross platform consistency
  */
@@ -35,6 +38,7 @@
 
 #define DEFAULT_SUBTITLE_PRIORITY		10000.0f
 
+
 /**
  * Some filters don't work properly with extreme values, so these are the limits 
  */
@@ -47,11 +51,23 @@
 #define MIN_FILTER_BANDWIDTH			0.1f
 #define MAX_FILTER_BANDWIDTH			2.0f
 
+
 /**
  * Debugger is Available on non-shipping builds
  */
 #define ENABLE_AUDIO_DEBUG !UE_BUILD_SHIPPING
 
+
+/**
+ * Trace macros specific to audio.
+ */
+#define ENABLE_AUDIO_TRACE !UE_BUILD_SHIPPING
+
+#if ENABLE_AUDIO_TRACE
+AUDIOMIXERCORE_API UE_TRACE_CHANNEL_EXTERN(AudioMixerChannel);
+#endif // ENABLE_AUDIO_TRACE
+
+
 /** Common Audio namespace Type Definitions/Identifiers */
 namespace Audio
 {
