Checkpoint: Audio Insights Sources Dashboard Widgets (Part 1)
- Break out FDashboard to be able to extend with 3rd party dashboards
- Viewport investigations commented out
- Add menu bar
- Ability to open/browse to selection
- Fix some stuff I broke in last CL wrt updating values
- Actually get peak amplitude & rename column (before was getting last value reported in queue of updates, not the highest)
- Misc Clean-up

diff --git a/Engine/Source/Runtime/AudioMixer/Private/AudioMixerSource.cpp b/Engine/Source/Runtime/AudioMixer/Private/AudioMixerSource.cpp
--- a/Engine/Source/Runtime/AudioMixer/Private/AudioMixerSource.cpp
+++ b/Engine/Source/Runtime/AudioMixer/Private/AudioMixerSource.cpp
@@ -31,19 +31,6 @@ UE_TRACE_EVENT_END()
 UE_TRACE_EVENT_BEGIN(Audio, MixerSourceStop)
 	UE_TRACE_EVENT_FIELD(uint64, Timestamp)
 	UE_TRACE_EVENT_FIELD(int32, SourceId)
-	UE_TRACE_EVENT_FIELD(UE::Trace::WideString, Name)
-UE_TRACE_EVENT_END()
-
-UE_TRACE_EVENT_BEGIN(Audio, MixerSourceUpdate)
-	UE_TRACE_EVENT_FIELD(uint64, Timestamp)
-	UE_TRACE_EVENT_FIELD(int32, SourceId)
-	UE_TRACE_EVENT_FIELD(UE::Trace::WideString, Name)
-	UE_TRACE_EVENT_FIELD(float, Pitch)
-	UE_TRACE_EVENT_FIELD(float, LPFreq)
-	UE_TRACE_EVENT_FIELD(float, HPFreq)
-	UE_TRACE_EVENT_FIELD(float, Envelope)
-	UE_TRACE_EVENT_FIELD(float, Distance)
-	UE_TRACE_EVENT_FIELD(float, DistanceAttenuation)
 UE_TRACE_EVENT_END()
 #endif // ENABLE_AUDIO_TRACE && AUDIO_MIXER_CPUPROFILERTRACE_ENABLED
 
@@ -1041,7 +1028,9 @@ namespace Audio
 
 		UpdateChannelMaps();
 
-		UpdateTraceInfo();
+#if ENABLE_AUDIO_DEBUG
+		Audio::FAudioDebugger::DrawDebugInfo(*this);
+#endif // ENABLE_AUDIO_DEBUG
 	}
 
 	bool FMixerSource::PrepareForInitialization(FWaveInstance* InWaveInstance)
@@ -1266,7 +1255,7 @@ namespace Audio
 				UE_TRACE_LOG(Audio, MixerSourceStart, AudioMixerChannel)
 					<< MixerSourceStart.Timestamp(FPlatformTime::Cycles64())
 					<< MixerSourceStart.SourceId(TraceSourceId)
-					<< MixerSourceStart.Name(*WaveInstance->GetName());
+					<< MixerSourceStart.Name(*WaveInstance->WaveData->GetPathName());
 			}
 #endif // ENABLE_AUDIO_TRACE
 		}
@@ -1365,8 +1354,7 @@ namespace Audio
 
 					UE_TRACE_LOG(Audio, MixerSourceStop, AudioMixerChannel)
 						<< MixerSourceStop.Timestamp(FPlatformTime::Cycles64())
-						<< MixerSourceStop.SourceId(TraceSourceId)
-						<< MixerSourceStop.Name(*WaveInstance->GetName());
+						<< MixerSourceStop.SourceId(TraceSourceId);
 				}
 #endif // ENABLE_AUDIO_TRACE
 
@@ -1908,45 +1896,6 @@ namespace Audio
 		bPrevAllowedSpatializationSetting = IsSpatializationCVarEnabled();
 	}
 
-	void FMixerSource::UpdateTraceInfo()
-	{
-#if ENABLE_AUDIO_DEBUG
-		Audio::FAudioDebugger::DrawDebugInfo(*this);
-#endif // ENABLE_AUDIO_DEBUG
-
-#if ENABLE_AUDIO_TRACE && AUDIO_MIXER_CPUPROFILERTRACE_ENABLED
-		const bool bChannelEnabled = UE_TRACE_CHANNELEXPR_IS_ENABLED(AudioMixerChannel);
-		if (bChannelEnabled && WaveInstance)
-		{
-			int32 TraceSourceId = INDEX_NONE;
-			float TraceVolume = 1.0f;
-			float TracePitch  = 1.0f;
-			float TraceLPFreq = MIN_FILTER_FREQUENCY;
-			float TraceHPFreq = MAX_FILTER_FREQUENCY;
-			float TraceEnv = 1.0f;
-			float TraceDistAttenuation = 0.0f;
-
-			if (MixerSourceVoice)
-			{
-				TraceSourceId = MixerSourceVoice->GetSourceId();
-				TraceEnv = MixerSourceVoice->GetEnvelopeValue();
-				TraceDistAttenuation = MixerSourceVoice->GetDistanceAttenuation();
-			}
-
-			UE_TRACE_LOG(Audio, MixerSourceUpdate, AudioMixerChannel)
-				<< MixerSourceUpdate.Timestamp(FPlatformTime::Cycles64())
-				<< MixerSourceUpdate.SourceId(TraceSourceId)
-				<< MixerSourceUpdate.Name(*WaveInstance->GetName())
-				<< MixerSourceUpdate.Pitch(TracePitch)
-				<< MixerSourceUpdate.LPFreq(TraceLPFreq)
-				<< MixerSourceUpdate.HPFreq(TraceHPFreq)
-				<< MixerSourceUpdate.Envelope(TraceEnv)
-				<< MixerSourceUpdate.DistanceAttenuation(TraceDistAttenuation);
-
-		}
-#endif // ENABLE_AUDIO_TRACE
-	}
-
 	bool FMixerSource::ComputeMonoChannelMap(Audio::FAlignedFloatBuffer& OutChannelMap)
 	{
 		if (IsUsingObjectBasedSpatialization())
diff --git a/Engine/Source/Runtime/AudioMixer/Private/AudioMixerSource.h b/Engine/Source/Runtime/AudioMixer/Private/AudioMixerSource.h
--- a/Engine/Source/Runtime/AudioMixer/Private/AudioMixerSource.h
+++ b/Engine/Source/Runtime/AudioMixer/Private/AudioMixerSource.h
@@ -85,9 +85,6 @@ namespace Audio
 		/** Updates the channel map of the sound if its a 3d sound.*/
 		void UpdateChannelMaps();
 
-		/** Forwards trace info to trace system if enabled */
-		void UpdateTraceInfo();
-
 		/** Computes the mono-channel map. */
 		bool ComputeMonoChannelMap(Audio::FAlignedFloatBuffer& OutChannelMap);
 
diff --git a/Engine/Source/Runtime/AudioMixer/Private/AudioMixerSourceManager.cpp b/Engine/Source/Runtime/AudioMixer/Private/AudioMixerSourceManager.cpp
--- a/Engine/Source/Runtime/AudioMixer/Private/AudioMixerSourceManager.cpp
+++ b/Engine/Source/Runtime/AudioMixer/Private/AudioMixerSourceManager.cpp
@@ -174,6 +174,31 @@ UE_TRACE_EVENT_BEGIN(Audio, MixerSourceVolume)
 	UE_TRACE_EVENT_FIELD(int32, SourceId)
 	UE_TRACE_EVENT_FIELD(float, Volume)
 UE_TRACE_EVENT_END()
+
+UE_TRACE_EVENT_BEGIN(Audio, MixerSourceDistanceAttenuation)
+	UE_TRACE_EVENT_FIELD(uint64, Timestamp)
+	UE_TRACE_EVENT_FIELD(int32, SourceId)
+	UE_TRACE_EVENT_FIELD(float, DistanceAttenuation)
+UE_TRACE_EVENT_END()
+
+UE_TRACE_EVENT_BEGIN(Audio, MixerSourcePitch)
+	UE_TRACE_EVENT_FIELD(uint64, Timestamp)
+	UE_TRACE_EVENT_FIELD(int32, SourceId)
+	UE_TRACE_EVENT_FIELD(float, Pitch)
+UE_TRACE_EVENT_END()
+
+UE_TRACE_EVENT_BEGIN(Audio, MixerSourceFilters)
+	UE_TRACE_EVENT_FIELD(uint64, Timestamp)
+	UE_TRACE_EVENT_FIELD(int32, SourceId)
+	UE_TRACE_EVENT_FIELD(float, LPFFrequency)
+	UE_TRACE_EVENT_FIELD(float, HPFFrequency)
+UE_TRACE_EVENT_END()
+
+UE_TRACE_EVENT_BEGIN(Audio, MixerSourceEnvelope)
+	UE_TRACE_EVENT_FIELD(uint64, Timestamp)
+	UE_TRACE_EVENT_FIELD(int32, SourceId)
+	UE_TRACE_EVENT_FIELD(float, Envelope)
+UE_TRACE_EVENT_END()
 #endif // ENABLE_AUDIO_TRACE && AUDIO_MIXER_CPUPROFILERTRACE_ENABLED
 
 namespace Audio
@@ -311,8 +336,6 @@ namespace Audio
 			SourceInfo.bEnableBaseSubmix = false;
 			SourceInfo.bEnableSubmixSends = false;
 			SourceInfo.bIsVorbis = false;
-			SourceInfo.bIsBypassingLPF = false;
-			SourceInfo.bIsBypassingHPF = false;
 			SourceInfo.bHasPreDistanceAttenuationSend = false;
 			SourceInfo.bModFiltersUpdated = false;
 
@@ -665,8 +688,6 @@ namespace Audio
 		SourceInfo.bEnableBusSends = false;
 		SourceInfo.bEnableBaseSubmix = false;
 		SourceInfo.bEnableSubmixSends = false;
-		SourceInfo.bIsBypassingLPF = false;
-		SourceInfo.bIsBypassingHPF = false;
 		SourceInfo.bHasPreDistanceAttenuationSend = false;
 		SourceInfo.bModFiltersUpdated = false;
 
@@ -2103,6 +2124,17 @@ namespace Audio
 				const float FinalPitch = FMath::Clamp(TargetPitch * ModPitch, MinModulationPitchRangeFreqCVar, MaxModulationPitchRangeFreqCVar);
 				SourceInfo.PitchSourceParam.SetValue(FinalPitch, NumOutputFrames);
 
+#if ENABLE_AUDIO_TRACE && AUDIO_MIXER_CPUPROFILERTRACE_ENABLED
+				const bool bChannelEnabled = UE_TRACE_CHANNELEXPR_IS_ENABLED(AudioMixerChannel);
+				if (bChannelEnabled)
+				{
+					UE_TRACE_LOG(Audio, MixerSourcePitch, AudioMixerChannel)
+						<< MixerSourcePitch.Timestamp(FPlatformTime::Cycles64())
+						<< MixerSourcePitch.SourceId(SourceId)
+						<< MixerSourcePitch.Pitch(TargetPitch);
+				}
+#endif // ENABLE_AUDIO_TRACE && AUDIO_MIXER_CPUPROFILERTRACE_ENABLED
+
 				for (int32 Frame = StartFrame; Frame < NumOutputFrames; ++Frame)
 				{
 					// If we've read our last buffer, we're done
@@ -2461,87 +2493,98 @@ namespace Audio
 			TArrayView<float> PreDistanceAttenBufferView(PreDistanceAttenBufferPtr, NumSamples);
 
 			// Update volume fade information if we're stopping
-			if (SourceInfo.bIsStopping)
 			{
-				int32 NumFadeFrames = FMath::Min(SourceInfo.VolumeFadeNumFrames - SourceInfo.VolumeFadeFramePosition, NumOutputFrames);
-
-				SourceInfo.VolumeFadeFramePosition += NumFadeFrames;
-				SourceInfo.VolumeSourceDestination = SourceInfo.VolumeFadeSlope * (float) SourceInfo.VolumeFadeFramePosition + SourceInfo.VolumeFadeStart;
-
-				if (FMath::IsNearlyZero(SourceInfo.VolumeSourceDestination, KINDA_SMALL_NUMBER))
+				float VolumeStart = 1.0f;
+				float VolumeDestination = 1.0f;
+				if (SourceInfo.bIsStopping)
 				{
-					SourceInfo.VolumeSourceDestination = 0.0f;
-				}
+					int32 NumFadeFrames = FMath::Min(SourceInfo.VolumeFadeNumFrames - SourceInfo.VolumeFadeFramePosition, NumOutputFrames);
 
-				const int32 NumFadeSamples = NumFadeFrames * SourceInfo.NumInputChannels;
+					SourceInfo.VolumeFadeFramePosition += NumFadeFrames;
+					SourceInfo.VolumeSourceDestination = SourceInfo.VolumeFadeSlope * (float)SourceInfo.VolumeFadeFramePosition + SourceInfo.VolumeFadeStart;
 
-				float VolumeStart = SourceInfo.VolumeSourceStart;
-				float VolumeDestination = SourceInfo.VolumeSourceDestination;
-				if (MixerDevice->IsModulationPluginEnabled() && MixerDevice->ModulationInterface.IsValid())
-				{
-					const bool bHasProcessed = SourceInfo.VolumeModulation.GetHasProcessed();
-					const float ModVolumeStart = SourceInfo.VolumeModulation.GetValue();
-					SourceInfo.VolumeModulation.ProcessControl(SourceInfo.VolumeModulationBase);
-					const float ModVolumeEnd = SourceInfo.VolumeModulation.GetValue();
-					if (bHasProcessed)
+					if (FMath::IsNearlyZero(SourceInfo.VolumeSourceDestination, KINDA_SMALL_NUMBER))
 					{
-						VolumeStart *= ModVolumeStart;
+						SourceInfo.VolumeSourceDestination = 0.0f;
 					}
-					else
+
+					const int32 NumFadeSamples = NumFadeFrames * SourceInfo.NumInputChannels;
+
+					VolumeStart = SourceInfo.VolumeSourceStart;
+					VolumeDestination = SourceInfo.VolumeSourceDestination;
+					if (MixerDevice->IsModulationPluginEnabled() && MixerDevice->ModulationInterface.IsValid())
 					{
-						VolumeStart *= ModVolumeEnd;
+						const bool bHasProcessed = SourceInfo.VolumeModulation.GetHasProcessed();
+						const float ModVolumeStart = SourceInfo.VolumeModulation.GetValue();
+						SourceInfo.VolumeModulation.ProcessControl(SourceInfo.VolumeModulationBase);
+						const float ModVolumeEnd = SourceInfo.VolumeModulation.GetValue();
+						if (bHasProcessed)
+						{
+							VolumeStart *= ModVolumeStart;
+						}
+						else
+						{
+							VolumeStart *= ModVolumeEnd;
+						}
+						VolumeDestination *= ModVolumeEnd;
 					}
-					VolumeDestination *= ModVolumeEnd;
-				}
-
-				TArrayView<float> PreDistanceAttenBufferFadeSamplesView(PreDistanceAttenBufferPtr, NumFadeSamples);
-				Audio::ArrayFade(PreDistanceAttenBufferFadeSamplesView, VolumeStart, VolumeDestination);
 
-				// Zero the rest of the buffer
-				if (NumFadeFrames < NumOutputFrames)
-				{
-					int32 SamplesLeft = NumSamples - NumFadeSamples;
+					TArrayView<float> PreDistanceAttenBufferFadeSamplesView(PreDistanceAttenBufferPtr, NumFadeSamples);
+					Audio::ArrayFade(PreDistanceAttenBufferFadeSamplesView, VolumeStart, VolumeDestination);
 
-					// Protect memzero call with some sanity checking on the inputs.
-					if(SamplesLeft > 0 && NumFadeSamples >= 0 && NumFadeSamples < NumSamples)
+					// Zero the rest of the buffer
+					if (NumFadeFrames < NumOutputFrames)
 					{
-						FMemory::Memzero(&PreDistanceAttenBufferPtr[NumFadeSamples], sizeof(float) * SamplesLeft);
+						int32 SamplesLeft = NumSamples - NumFadeSamples;
+
+						// Protect memzero call with some sanity checking on the inputs.
+						if (SamplesLeft > 0 && NumFadeSamples >= 0 && NumFadeSamples < NumSamples)
+						{
+							FMemory::Memzero(&PreDistanceAttenBufferPtr[NumFadeSamples], sizeof(float) * SamplesLeft);
+						}
 					}
 				}
-			}
-			else
-			{
-				float VolumeStart = SourceInfo.VolumeSourceStart;
-				float VolumeDestination = SourceInfo.VolumeSourceDestination;
-				if (MixerDevice->IsModulationPluginEnabled() && MixerDevice->ModulationInterface.IsValid())
+				else
 				{
-					const bool bHasProcessed = SourceInfo.VolumeModulation.GetHasProcessed();
-					const float ModVolumeStart = SourceInfo.VolumeModulation.GetValue();
-					SourceInfo.VolumeModulation.ProcessControl(SourceInfo.VolumeModulationBase);
-					const float ModVolumeEnd = SourceInfo.VolumeModulation.GetValue();
-					if (bHasProcessed)
-					{
-						VolumeStart *= ModVolumeStart;
-					}
-					else
+					VolumeStart = SourceInfo.VolumeSourceStart;
+					VolumeDestination = SourceInfo.VolumeSourceDestination;
+					if (MixerDevice->IsModulationPluginEnabled() && MixerDevice->ModulationInterface.IsValid())
 					{
-						VolumeStart *= ModVolumeEnd;
+						const bool bHasProcessed = SourceInfo.VolumeModulation.GetHasProcessed();
+						const float ModVolumeStart = SourceInfo.VolumeModulation.GetValue();
+						SourceInfo.VolumeModulation.ProcessControl(SourceInfo.VolumeModulationBase);
+						const float ModVolumeEnd = SourceInfo.VolumeModulation.GetValue();
+						if (bHasProcessed)
+						{
+							VolumeStart *= ModVolumeStart;
+						}
+						else
+						{
+							VolumeStart *= ModVolumeEnd;
+						}
+						VolumeDestination *= ModVolumeEnd;
 					}
-					VolumeDestination *= ModVolumeEnd;
+
+					Audio::ArrayFade(PreDistanceAttenBufferView, VolumeStart, VolumeDestination);
 				}
 
-				Audio::ArrayFade(PreDistanceAttenBufferView, VolumeStart, VolumeDestination);
 #if ENABLE_AUDIO_TRACE && AUDIO_MIXER_CPUPROFILERTRACE_ENABLED
-				if (!FMath::IsNearlyEqual(VolumeStart, VolumeDestination))
+				const bool bChannelEnabled = UE_TRACE_CHANNELEXPR_IS_ENABLED(AudioMixerChannel);
+				if (bChannelEnabled)
 				{
-					const bool bChannelEnabled = UE_TRACE_CHANNELEXPR_IS_ENABLED(AudioMixerChannel);
 					UE_TRACE_LOG(Audio, MixerSourceVolume, AudioMixerChannel)
 						<< MixerSourceVolume.Timestamp(FPlatformTime::Cycles64())
 						<< MixerSourceVolume.SourceId(SourceId)
 						<< MixerSourceVolume.Volume(VolumeDestination);
+
+					UE_TRACE_LOG(Audio, MixerSourceDistanceAttenuation, AudioMixerChannel)
+						<< MixerSourceDistanceAttenuation.Timestamp(FPlatformTime::Cycles64())
+						<< MixerSourceDistanceAttenuation.SourceId(SourceId)
+						<< MixerSourceDistanceAttenuation.DistanceAttenuation(SourceInfo.DistanceAttenuationSourceDestination);
 				}
 #endif // ENABLE_AUDIO_TRACE && AUDIO_MIXER_CPUPROFILERTRACE_ENABLED
 			}
+
 			SourceInfo.VolumeSourceStart = SourceInfo.VolumeSourceDestination;
 
 			// Now process the effect chain if it exists
@@ -2610,6 +2653,7 @@ namespace Audio
 
 			const bool bModActive = MixerDevice->IsModulationPluginEnabled() && MixerDevice->ModulationInterface.IsValid();
 			bool bUpdateModFilters = bModActive && (SourceInfo.bModFiltersUpdated || SourceInfo.LowpassModulation.IsActive() || SourceInfo.HighpassModulation.IsActive());
+
 			if (SourceInfo.IsRenderingToSubmixes() || bUpdateModFilters)
 			{
 				// Only scale with distance attenuation and send to source audio to plugins if we're not in output-to-bus only mode
@@ -2634,13 +2678,13 @@ namespace Audio
 					SourceInfo.HighPassFilter.StartFrequencyInterpolation(SourceInfo.HighpassModulation.GetValue(), NumOutputFrames);
 				}
 
-				const bool BypassLPF = DisableFilteringCvar || (SourceInfo.LowPassFilter.GetCutoffFrequency() >= (MAX_FILTER_FREQUENCY - KINDA_SMALL_NUMBER));
-				const bool BypassHPF = DisableFilteringCvar || DisableHPFilteringCvar || (SourceInfo.HighPassFilter.GetCutoffFrequency() <= (MIN_FILTER_FREQUENCY + KINDA_SMALL_NUMBER));
+				const bool bBypassLPF = DisableFilteringCvar || (SourceInfo.LowPassFilter.GetCutoffFrequency() >= (MAX_FILTER_FREQUENCY - KINDA_SMALL_NUMBER));
+				const bool bBypassHPF = DisableFilteringCvar || DisableHPFilteringCvar || (SourceInfo.HighPassFilter.GetCutoffFrequency() <= (MIN_FILTER_FREQUENCY + KINDA_SMALL_NUMBER));
 
 				float* SourceBuffer = SourceInfo.SourceBuffer.GetData();
 				float* HpfInputBuffer = PreDistanceAttenBufferPtr; // assume bypassing LPF (HPF uses input buffer as input)
 
-				if (!BypassLPF)
+				if (!bBypassLPF)
 				{
 					// Not bypassing LPF, so tell HPF to use LPF output buffer as input
 					HpfInputBuffer = SourceBuffer;
@@ -2649,17 +2693,45 @@ namespace Audio
 					SourceInfo.LowPassFilter.ProcessAudioBuffer(PreDistanceAttenBufferPtr, SourceBuffer, NumOutputSamplesThisSource);
 				}
 
-				if (!BypassHPF)
+				if (!bBypassHPF)
 				{
 					// process HPF audio block
 					SourceInfo.HighPassFilter.ProcessAudioBuffer(HpfInputBuffer, SourceBuffer, NumOutputSamplesThisSource);
 				}
 
+#if ENABLE_AUDIO_TRACE && AUDIO_MIXER_CPUPROFILERTRACE_ENABLED
+				const bool bChannelEnabled = UE_TRACE_CHANNELEXPR_IS_ENABLED(AudioMixerChannel);
+				if (bChannelEnabled)
+				{
+					float LPFFrequency = MAX_FILTER_FREQUENCY;
+					if (!bBypassLPF)
+					{
+						LPFFrequency = SourceInfo.LowpassModulation.GetValue();
+					}
+
+					float HPFFrequency = MIN_FILTER_FREQUENCY;
+					if (!bBypassHPF)
+					{
+						HPFFrequency = SourceInfo.HighpassModulation.GetValue();
+					}
+
+					UE_TRACE_LOG(Audio, MixerSourceFilters, AudioMixerChannel)
+						<< MixerSourceFilters.Timestamp(FPlatformTime::Cycles64())
+						<< MixerSourceFilters.SourceId(SourceId)
+						<< MixerSourceFilters.HPFFrequency(HPFFrequency)
+						<< MixerSourceFilters.LPFFrequency(LPFFrequency);
+					UE_TRACE_LOG(Audio, MixerSourceEnvelope, AudioMixerChannel)
+						<< MixerSourceEnvelope.Timestamp(FPlatformTime::Cycles64())
+						<< MixerSourceEnvelope.SourceId(SourceId)
+						<< MixerSourceEnvelope.Envelope(SourceInfo.SourceEnvelopeValue);
+				}
+#endif // ENABLE_AUDIO_TRACE && AUDIO_MIXER_CPUPROFILERTRACE_ENABLED
+
 				// We manually reset interpolation to avoid branches in filter code
 				SourceInfo.LowPassFilter.StopFrequencyInterpolation();
 				SourceInfo.HighPassFilter.StopFrequencyInterpolation();
 
-				if (BypassLPF && BypassHPF)
+				if (bBypassLPF && bBypassHPF)
 				{
 					FMemory::Memcpy(SourceBuffer, PreDistanceAttenBufferPtr, NumSamples * sizeof(float));
 				}
diff --git a/Engine/Source/Runtime/AudioMixer/Private/AudioMixerSourceManager.h b/Engine/Source/Runtime/AudioMixer/Private/AudioMixerSourceManager.h
--- a/Engine/Source/Runtime/AudioMixer/Private/AudioMixerSourceManager.h
+++ b/Engine/Source/Runtime/AudioMixer/Private/AudioMixerSourceManager.h
@@ -8,6 +8,7 @@
 #include "AudioMixerDevice.h"
 #include "AudioMixerSourceOutputBuffer.h"
 #include "AudioMixerSubmix.h"
+#include "AudioMixerTrace.h"
 #include "Containers/MpscQueue.h"
 #include "DSP/BufferVectorOperations.h"
 #include "DSP/EnvelopeFollower.h"
@@ -530,8 +531,6 @@ namespace Audio
 			uint8 bEnableSubmixSends : 1;
 			uint8 bIsVorbis:1;
 			uint8 bIsSoundfield:1;
-			uint8 bIsBypassingLPF:1;
-			uint8 bIsBypassingHPF:1;
 			uint8 bHasPreDistanceAttenuationSend:1;
 			uint8 bModFiltersUpdated : 1;
 			uint8 bShouldSourceBufferListenerZeroBuffer : 1;
