Checkpoint: Audio Insights Plugin dashboard work (Part 1)
- Clean-up Audio Trace macros
- Add "AudioChannel" trace channel
- Remove stubbed in properties for misc dashboard widget factories from dashboard factory
- Add virtualization tab
- Move Log & OutputMeter to separate implementation/header files
#rnx

diff --git a/Engine/Source/Runtime/AudioMixer/Private/AudioMixerSource.cpp b/Engine/Source/Runtime/AudioMixer/Private/AudioMixerSource.cpp
--- a/Engine/Source/Runtime/AudioMixer/Private/AudioMixerSource.cpp
+++ b/Engine/Source/Runtime/AudioMixer/Private/AudioMixerSource.cpp
@@ -21,7 +21,7 @@
 
 CSV_DECLARE_CATEGORY_MODULE_EXTERN(AUDIOMIXERCORE_API, Audio);
 
-#if ENABLE_AUDIO_TRACE && AUDIO_MIXER_CPUPROFILERTRACE_ENABLED
+#if UE_AUDIO_PROFILERTRACE_ENABLED
 UE_TRACE_EVENT_BEGIN(Audio, MixerSourceStart)
 	UE_TRACE_EVENT_FIELD(uint64, Timestamp)
 	UE_TRACE_EVENT_FIELD(int32, SourceId)
@@ -32,7 +32,7 @@ UE_TRACE_EVENT_BEGIN(Audio, MixerSourceStop)
 	UE_TRACE_EVENT_FIELD(uint64, Timestamp)
 	UE_TRACE_EVENT_FIELD(int32, SourceId)
 UE_TRACE_EVENT_END()
-#endif // ENABLE_AUDIO_TRACE && AUDIO_MIXER_CPUPROFILERTRACE_ENABLED
+#endif // UE_AUDIO_PROFILERTRACE_ENABLED
 
 
 static int32 UseListenerOverrideForSpreadCVar = 0;
@@ -1242,7 +1242,7 @@ namespace Audio
 		{
 			MixerSourceVoice->Play();
 
-#if ENABLE_AUDIO_TRACE && AUDIO_MIXER_CPUPROFILERTRACE_ENABLED
+#if UE_AUDIO_PROFILERTRACE_ENABLED
 			const bool bChannelEnabled = UE_TRACE_CHANNELEXPR_IS_ENABLED(AudioMixerChannel);
 			if (bChannelEnabled && WaveInstance)
 			{
@@ -1257,7 +1257,7 @@ namespace Audio
 					<< MixerSourceStart.SourceId(TraceSourceId)
 					<< MixerSourceStart.Name(*WaveInstance->WaveData->GetPathName());
 			}
-#endif // ENABLE_AUDIO_TRACE
+#endif // UE_AUDIO_PROFILERTRACE_ENABLED
 		}
 
 		bIsStopping = false;
@@ -1342,7 +1342,7 @@ namespace Audio
 		{
 			if (MixerSourceVoice && Playing)
 			{
-#if ENABLE_AUDIO_TRACE && AUDIO_MIXER_CPUPROFILERTRACE_ENABLED
+#if UE_AUDIO_PROFILERTRACE_ENABLED
 				const bool bChannelEnabled = UE_TRACE_CHANNELEXPR_IS_ENABLED(AudioMixerChannel);
 				if (bChannelEnabled)
 				{
@@ -1356,7 +1356,7 @@ namespace Audio
 						<< MixerSourceStop.Timestamp(FPlatformTime::Cycles64())
 						<< MixerSourceStop.SourceId(TraceSourceId);
 				}
-#endif // ENABLE_AUDIO_TRACE
+#endif // UE_AUDIO_PROFILERTRACE_ENABLED
 
 				MixerSourceVoice->Stop();
 			}
diff --git a/Engine/Source/Runtime/AudioMixer/Private/AudioMixerSourceManager.cpp b/Engine/Source/Runtime/AudioMixer/Private/AudioMixerSourceManager.cpp
--- a/Engine/Source/Runtime/AudioMixer/Private/AudioMixerSourceManager.cpp
+++ b/Engine/Source/Runtime/AudioMixer/Private/AudioMixerSourceManager.cpp
@@ -168,7 +168,7 @@ DEFINE_STAT(STAT_AudioMixerSourceEffectBuffers);
 DEFINE_STAT(STAT_AudioMixerSourceManagerUpdate);
 DEFINE_STAT(STAT_AudioMixerSourceOutputBuffers);
 
-#if ENABLE_AUDIO_TRACE && AUDIO_MIXER_CPUPROFILERTRACE_ENABLED
+#if UE_AUDIO_PROFILERTRACE_ENABLED
 UE_TRACE_EVENT_BEGIN(Audio, MixerSourceVolume)
 	UE_TRACE_EVENT_FIELD(uint64, Timestamp)
 	UE_TRACE_EVENT_FIELD(int32, SourceId)
@@ -199,7 +199,7 @@ UE_TRACE_EVENT_BEGIN(Audio, MixerSourceEnvelope)
 	UE_TRACE_EVENT_FIELD(int32, SourceId)
 	UE_TRACE_EVENT_FIELD(float, Envelope)
 UE_TRACE_EVENT_END()
-#endif // ENABLE_AUDIO_TRACE && AUDIO_MIXER_CPUPROFILERTRACE_ENABLED
+#endif // UE_AUDIO_PROFILERTRACE_ENABLED
 
 namespace Audio
 {
@@ -2124,7 +2124,7 @@ namespace Audio
 				const float FinalPitch = FMath::Clamp(TargetPitch * ModPitch, MinModulationPitchRangeFreqCVar, MaxModulationPitchRangeFreqCVar);
 				SourceInfo.PitchSourceParam.SetValue(FinalPitch, NumOutputFrames);
 
-#if ENABLE_AUDIO_TRACE && AUDIO_MIXER_CPUPROFILERTRACE_ENABLED
+#if UE_AUDIO_PROFILERTRACE_ENABLED
 				const bool bChannelEnabled = UE_TRACE_CHANNELEXPR_IS_ENABLED(AudioMixerChannel);
 				if (bChannelEnabled)
 				{
@@ -2133,7 +2133,7 @@ namespace Audio
 						<< MixerSourcePitch.SourceId(SourceId)
 						<< MixerSourcePitch.Pitch(TargetPitch);
 				}
-#endif // ENABLE_AUDIO_TRACE && AUDIO_MIXER_CPUPROFILERTRACE_ENABLED
+#endif // UE_AUDIO_PROFILERTRACE_ENABLED
 
 				for (int32 Frame = StartFrame; Frame < NumOutputFrames; ++Frame)
 				{
@@ -2568,7 +2568,7 @@ namespace Audio
 					Audio::ArrayFade(PreDistanceAttenBufferView, VolumeStart, VolumeDestination);
 				}
 
-#if ENABLE_AUDIO_TRACE && AUDIO_MIXER_CPUPROFILERTRACE_ENABLED
+#if UE_AUDIO_PROFILERTRACE_ENABLED
 				const bool bChannelEnabled = UE_TRACE_CHANNELEXPR_IS_ENABLED(AudioMixerChannel);
 				if (bChannelEnabled)
 				{
@@ -2582,7 +2582,7 @@ namespace Audio
 						<< MixerSourceDistanceAttenuation.SourceId(SourceId)
 						<< MixerSourceDistanceAttenuation.DistanceAttenuation(SourceInfo.DistanceAttenuationSourceDestination);
 				}
-#endif // ENABLE_AUDIO_TRACE && AUDIO_MIXER_CPUPROFILERTRACE_ENABLED
+#endif // UE_AUDIO_PROFILERTRACE_ENABLED
 			}
 
 			SourceInfo.VolumeSourceStart = SourceInfo.VolumeSourceDestination;
@@ -2699,7 +2699,7 @@ namespace Audio
 					SourceInfo.HighPassFilter.ProcessAudioBuffer(HpfInputBuffer, SourceBuffer, NumOutputSamplesThisSource);
 				}
 
-#if ENABLE_AUDIO_TRACE && AUDIO_MIXER_CPUPROFILERTRACE_ENABLED
+#if UE_AUDIO_PROFILERTRACE_ENABLED
 				const bool bChannelEnabled = UE_TRACE_CHANNELEXPR_IS_ENABLED(AudioMixerChannel);
 				if (bChannelEnabled)
 				{
@@ -2725,7 +2725,7 @@ namespace Audio
 						<< MixerSourceEnvelope.SourceId(SourceId)
 						<< MixerSourceEnvelope.Envelope(SourceInfo.SourceEnvelopeValue);
 				}
-#endif // ENABLE_AUDIO_TRACE && AUDIO_MIXER_CPUPROFILERTRACE_ENABLED
+#endif // UE_AUDIO_PROFILERTRACE_ENABLED
 
 				// We manually reset interpolation to avoid branches in filter code
 				SourceInfo.LowPassFilter.StopFrequencyInterpolation();
diff --git a/Engine/Source/Runtime/AudioMixerCore/Private/AudioMixer.cpp b/Engine/Source/Runtime/AudioMixerCore/Private/AudioMixer.cpp
--- a/Engine/Source/Runtime/AudioMixerCore/Private/AudioMixer.cpp
+++ b/Engine/Source/Runtime/AudioMixerCore/Private/AudioMixer.cpp
@@ -2,6 +2,7 @@
 #include "AudioMixer.h"
 
 #include "AudioDefines.h"
+#include "AudioMixerTrace.h"
 #include "DSP/BufferVectorOperations.h"
 #include "DSP/FloatArrayMath.h"
 #include "HAL/RunnableThread.h"
@@ -22,9 +23,10 @@
 //
 CSV_DEFINE_CATEGORY_MODULE(AUDIOMIXERCORE_API, Audio, false);
 
-#if ENABLE_AUDIO_TRACE
+#if UE_AUDIO_PROFILERTRACE_ENABLED
+UE_TRACE_CHANNEL_DEFINE(AudioChannel);
 UE_TRACE_CHANNEL_DEFINE(AudioMixerChannel);
-#endif // ENABLE_AUDIO_TRACE
+#endif // UE_AUDIO_PROFILERTRACE_ENABLED
 
 
 // Command to enable logging to display accurate audio render times
diff --git a/Engine/Source/Runtime/AudioMixerCore/Public/AudioDefines.h b/Engine/Source/Runtime/AudioMixerCore/Public/AudioDefines.h
--- a/Engine/Source/Runtime/AudioMixerCore/Public/AudioDefines.h
+++ b/Engine/Source/Runtime/AudioMixerCore/Public/AudioDefines.h
@@ -2,7 +2,6 @@
 #pragma once
 
 #include "HAL/Platform.h"
-#include "Trace/Trace.h"
 
 
 /**
@@ -58,16 +57,6 @@
 #define ENABLE_AUDIO_DEBUG !UE_BUILD_SHIPPING
 
 
-/**
- * Trace macros specific to audio.
- */
-#define ENABLE_AUDIO_TRACE !UE_BUILD_SHIPPING
-
-#if ENABLE_AUDIO_TRACE
-AUDIOMIXERCORE_API UE_TRACE_CHANNEL_EXTERN(AudioMixerChannel);
-#endif // ENABLE_AUDIO_TRACE
-
-
 /** Common Audio namespace Type Definitions/Identifiers */
 namespace Audio
 {
diff --git a/Engine/Source/Runtime/AudioMixerCore/Public/AudioMixerTrace.h b/Engine/Source/Runtime/AudioMixerCore/Public/AudioMixerTrace.h
--- a/Engine/Source/Runtime/AudioMixerCore/Public/AudioMixerTrace.h
+++ b/Engine/Source/Runtime/AudioMixerCore/Public/AudioMixerTrace.h
@@ -3,25 +3,20 @@
 #pragma once
 
 #include "ProfilingDebugging/CpuProfilerTrace.h"
+#include "Trace/Trace.h"
 
-#if !defined(AUDIO_MIXER_CPUPROFILERTRACE_ENABLED)
-#if CPUPROFILERTRACE_ENABLED
-#define AUDIO_MIXER_CPUPROFILERTRACE_ENABLED 1
-#else
-#define AUDIO_MIXER_CPUPROFILERTRACE_ENABLED 0
-#endif
-#endif
-
-#if AUDIO_MIXER_CPUPROFILERTRACE_ENABLED
-// AudioMixer CPU profiler trace enabled
-
-#define AUDIO_MIXER_TRACE_CPUPROFILER_EVENT_SCOPE(Name) \
-    TRACE_CPUPROFILER_EVENT_SCOPE(Name)
+#ifndef	UE_AUDIO_PROFILERTRACE_ENABLED
+	#define UE_AUDIO_PROFILERTRACE_ENABLED UE_TRACE_ENABLED && !UE_BUILD_SHIPPING
+#endif // UE_AUDIO_PROFILERTRACE_ENABLED
 
+#if	UE_AUDIO_PROFILERTRACE_ENABLED
+	#define AUDIO_MIXER_TRACE_CPUPROFILER_EVENT_SCOPE(Name) TRACE_CPUPROFILER_EVENT_SCOPE(Name)
+	AUDIOMIXERCORE_API UE_TRACE_CHANNEL_EXTERN(AudioChannel);
+	AUDIOMIXERCORE_API UE_TRACE_CHANNEL_EXTERN(AudioMixerChannel);
 #else
-// AudioMixer CPU profiler trace *not* enabled
-
-#define AUDIO_MIXER_TRACE_CPUPROFILER_EVENT_SCOPE(Name)
-
+	#if CPUPROFILERTRACE_ENABLED
+		#define AUDIO_MIXER_TRACE_CPUPROFILER_EVENT_SCOPE(Name) TRACE_CPUPROFILER_EVENT_SCOPE(Name)
+	#else // !CPUPROFILERTRACE_ENABLED
+		#define AUDIO_MIXER_TRACE_CPUPROFILER_EVENT_SCOPE(Name)
+	#endif
 #endif
-
