Audio bus reader node output patch buffer size fix

This fixes a bug where audio bus reader nodes might never start popping audio because AudioCallbackBufferFrameSize is set to a value other than the default 1024, such as 256.

Replace AddPatchForAudioBus with AddPatchOutputForAudioBus, which is parameterized by the desired number of frames and channels in the patch output.
Update AddPatchInputForAudioBus to match.
Both functions use a new MakePatch function to create a patch based on the parameters and scaled to fit the number of blocks specified by the new "au.PatchBufferBlocks" cvar.
Add support for quad, 5.1, and 7.1 bus reader and writer nodes.

diff --git a/Engine/Plugins/Runtime/AudioModulation/Source/AudioModulation/Private/Generators/SoundModulationEnvelopeFollower.cpp b/Engine/Plugins/Runtime/AudioModulation/Source/AudioModulation/Private/Generators/SoundModulationEnvelopeFollower.cpp
--- a/Engine/Plugins/Runtime/AudioModulation/Source/AudioModulation/Private/Generators/SoundModulationEnvelopeFollower.cpp
+++ b/Engine/Plugins/Runtime/AudioModulation/Source/AudioModulation/Private/Generators/SoundModulationEnvelopeFollower.cpp
@@ -203,7 +203,7 @@ namespace AudioModulation
 					FAudioDevice* AudioDevice = DeviceManager->GetAudioDeviceRaw(AudioDeviceId);
 					if (Audio::FMixerDevice* MixerDevice = static_cast<Audio::FMixerDevice*>(AudioDevice))
 					{
-						AudioBusPatch = MixerDevice->AddPatchForAudioBus(BusId, Gain);
+						AudioBusPatch = MixerDevice->AddPatchOutputForAudioBus(BusId, MixerDevice->GetNumOutputFrames(), EnvelopeFollower.GetNumChannels(), Gain);
 						bBusRequiresPatch = false;
 					}
 				}
diff --git a/Engine/Plugins/Runtime/Metasound/Source/MetasoundEngine/Private/MetasoundAudioBusReaderNode.cpp b/Engine/Plugins/Runtime/Metasound/Source/MetasoundEngine/Private/MetasoundAudioBusReaderNode.cpp
--- a/Engine/Plugins/Runtime/Metasound/Source/MetasoundEngine/Private/MetasoundAudioBusReaderNode.cpp
+++ b/Engine/Plugins/Runtime/Metasound/Source/MetasoundEngine/Private/MetasoundAudioBusReaderNode.cpp
@@ -113,7 +113,10 @@ namespace Metasound
 						AudioBusChannels = AudioBusProxy->NumChannels;
 						AudioDevice->StartAudioBus(AudioBusProxy->AudioBusId, AudioBusChannels, false);
 
-						AudioBusPatchOutput = AudioDevice->AddPatchForAudioBus(AudioBusProxy->AudioBusId);
+						int32 BlockSizeFrames = InSettings.GetNumFramesPerBlock();
+
+						// Create a bus patch output with enough room for the number of samples we expect and some buffering
+						AudioBusPatchOutput = AudioDevice->AddPatchOutputForAudioBus(AudioBusProxy->AudioBusId, BlockSizeFrames, AudioBusChannels);
 					}
 				}
 			}
@@ -184,7 +187,11 @@ namespace Metasound
 			if (bPerformPop)
 			{
 				// Pop off the interleaved data from the audio bus
-				AudioBusPatchOutput->PopAudio(InterleavedBuffer.GetData(), NumSamplesToPop, false);
+				int32 SamplesPopped = AudioBusPatchOutput->PopAudio(InterleavedBuffer.GetData(), NumSamplesToPop, false);
+				if (SamplesPopped < NumSamplesToPop)
+				{
+					UE_LOG(LogMetaSound, Warning, TEXT("Underrun detected in audio bus reader node."));
+				}
 
 				const uint32 MinChannels = FMath::Min(NumChannels, AudioBusChannels);
 				for (uint32 ChannelIndex = 0; ChannelIndex < MinChannels; ++ChannelIndex)
@@ -234,6 +241,9 @@ namespace Metasound
 	
 	REGISTER_AUDIO_BUS_READER_NODE(1);
 	REGISTER_AUDIO_BUS_READER_NODE(2);
+	REGISTER_AUDIO_BUS_READER_NODE(4);
+	REGISTER_AUDIO_BUS_READER_NODE(6);
+	REGISTER_AUDIO_BUS_READER_NODE(8);
 }
 
 #undef LOCTEXT_NAMESPACE
diff --git a/Engine/Plugins/Runtime/Metasound/Source/MetasoundEngine/Private/MetasoundAudioBusWriterNode.cpp b/Engine/Plugins/Runtime/Metasound/Source/MetasoundEngine/Private/MetasoundAudioBusWriterNode.cpp
--- a/Engine/Plugins/Runtime/Metasound/Source/MetasoundEngine/Private/MetasoundAudioBusWriterNode.cpp
+++ b/Engine/Plugins/Runtime/Metasound/Source/MetasoundEngine/Private/MetasoundAudioBusWriterNode.cpp
@@ -108,20 +108,15 @@ namespace Metasound
 				{
 					if (FAudioDevice* AudioDevice = ADM->GetAudioDeviceRaw(AudioDeviceId))
 					{
-						if (AudioDevice->IsAudioMixerEnabled())
-						{
-							AudioBusChannels = AudioBusProxy->NumChannels;
-							BlockSizeFrames = AudioInputs[0]->Num();
-							InterleavedBuffer.AddUninitialized(BlockSizeFrames * AudioBusChannels);
-
-							Audio::FMixerDevice& MixerDevice(static_cast<Audio::FMixerDevice&>(*AudioDevice));
-							MixerDevice.StartAudioBus(AudioBusProxy->AudioBusId, NumChannels, false);
-
-							// Create a bus patch input with enough room for the number of samples we expect and some buffering
-							int32 MaxSizeFrames = FMath::Max(BlockSizeFrames, AudioMixerOutputFrames), MinSizeFrames = FMath::Min(BlockSizeFrames, AudioMixerOutputFrames), BufferScale = 3;
-							AudioBusPatchInput = Audio::FPatchInput(MakeShared<Audio::FPatchOutput, ESPMode::ThreadSafe>(BufferScale * MinSizeFrames * FMath::DivideAndRoundUp(MaxSizeFrames, MinSizeFrames) * AudioBusChannels, 1.f));
-							MixerDevice.AddPatchInputForAudioBus(AudioBusPatchInput, AudioBusProxy->AudioBusId);
-						}
+						// Start the audio bus in case it's not already started
+						AudioBusChannels = AudioBusProxy->NumChannels;
+						AudioDevice->StartAudioBus(AudioBusProxy->AudioBusId, AudioBusChannels, false);
+
+						BlockSizeFrames = InParams.OperatorSettings.GetNumFramesPerBlock();
+						InterleavedBuffer.AddZeroed(BlockSizeFrames * AudioBusChannels);
+
+						// Create a bus patch input with enough room for the number of samples we expect and some buffering
+						AudioBusPatchInput = AudioDevice->AddPatchInputForAudioBus(AudioBusProxy->AudioBusId, BlockSizeFrames, AudioBusChannels);
 					}
 				}
 			}
@@ -162,42 +157,42 @@ namespace Metasound
 				return;
 			}
 
-			int32 NumSamplesToPush = InterleavedBuffer.Num();
-			uint32 MinChannels = NumChannels;
-			if (MinChannels < AudioBusChannels)
-			{
-				// There are fewer inputs than outputs, so zero the interleaved buffer
-				FMemory::Memzero(InterleavedBuffer.GetData(), NumSamplesToPush * sizeof(float));
-			}
-			else
-			{
-				MinChannels = AudioBusChannels;
-			}
-
 			// Retrieve input and interleaved buffer pointers
 			const float* AudioInputBufferPtrs[NumChannels];
-			for (uint32 ChannelIndex = 0; ChannelIndex < MinChannels; ++ChannelIndex)
+			for (uint32 ChannelIndex = 0; ChannelIndex < NumChannels; ++ChannelIndex)
 			{
 				AudioInputBufferPtrs[ChannelIndex] = AudioInputs[ChannelIndex]->GetData();
 			}
 			float* InterleavedBufferPtr = InterleavedBuffer.GetData();
 
-			// Interleave the inputs
-			// Writing the channels of the interleaved buffer sequentially should improve
-			// cache utilization compared to writing each input's frames sequentially.
-			// There is more likely to be a cache line for each buffer than for the
-			// entirety of the interleaved buffer.
-			for (int32 FrameIndex = 0; FrameIndex < BlockSizeFrames; ++FrameIndex)
+			if (AudioBusChannels == 1)
 			{
-				for (uint32 ChannelIndex = 0; ChannelIndex < MinChannels; ++ChannelIndex)
+				FMemory::Memcpy(InterleavedBufferPtr, AudioInputBufferPtrs[0], BlockSizeFrames * sizeof(float));
+			}
+			else
+			{
+				// Interleave the inputs
+				// Writing the channels of the interleaved buffer sequentially should improve
+				// cache utilization compared to writing each input's frames sequentially.
+				// There is more likely to be a cache line for each buffer than for the
+				// entirety of the interleaved buffer.
+				uint32 MinChannels = FMath::Min(AudioBusChannels, NumChannels);
+				for (int32 FrameIndex = 0; FrameIndex < BlockSizeFrames; ++FrameIndex)
 				{
-					InterleavedBufferPtr[ChannelIndex] = *AudioInputBufferPtrs[ChannelIndex]++;
+					for (uint32 ChannelIndex = 0; ChannelIndex < MinChannels; ++ChannelIndex)
+					{
+						InterleavedBufferPtr[ChannelIndex] = *AudioInputBufferPtrs[ChannelIndex]++;
+					}
+					InterleavedBufferPtr += AudioBusChannels;
 				}
-				InterleavedBufferPtr += AudioBusChannels;
 			}
 
 			// Pushes the interleaved data to the audio bus
-			AudioBusPatchInput.PushAudio(InterleavedBuffer.GetData(), NumSamplesToPush);
+			int32 SamplesPushed = AudioBusPatchInput.PushAudio(InterleavedBuffer.GetData(), InterleavedBuffer.Num());
+			if (SamplesPushed < InterleavedBuffer.Num())
+			{
+				UE_LOG(LogMetaSound, Warning, TEXT("Underrun detected in audio bus writer node."));
+			}
 		}
 
 	private:
@@ -229,6 +224,9 @@ namespace Metasound
 
 	REGISTER_AUDIO_BUS_WRITER_NODE(1);
 	REGISTER_AUDIO_BUS_WRITER_NODE(2);
+	REGISTER_AUDIO_BUS_WRITER_NODE(4);
+	REGISTER_AUDIO_BUS_WRITER_NODE(6);
+	REGISTER_AUDIO_BUS_WRITER_NODE(8);
 }
 
 #undef LOCTEXT_NAMESPACE
diff --git a/Engine/Plugins/Runtime/Synthesis/Source/Synthesis/Private/SourceEffectFilter.cpp b/Engine/Plugins/Runtime/Synthesis/Source/Synthesis/Private/SourceEffectFilter.cpp
--- a/Engine/Plugins/Runtime/Synthesis/Source/Synthesis/Private/SourceEffectFilter.cpp
+++ b/Engine/Plugins/Runtime/Synthesis/Source/Synthesis/Private/SourceEffectFilter.cpp
@@ -114,7 +114,7 @@ void FSourceEffectFilter::OnPresetChanged()
 					FAudioBusModulationData& NewModData = ModData.AddDefaulted_GetRef();
 
 					uint32 AudioBusId = BusModulationSettings.AudioBus->GetUniqueID();
-					NewModData.AudioBusPatch = MixerDevice->AddPatchForAudioBus(AudioBusId, 1.0f);
+					NewModData.AudioBusPatch = MixerDevice->AddPatchOutputForAudioBus(AudioBusId, MixerDevice->GetNumOutputFrames(), BusModulationSettings.AudioBus->GetNumChannels());
 
 					NewModData.MinFreqModValue = BusModulationSettings.MinFrequencyModulation;
 					NewModData.MaxFreqModValue = BusModulationSettings.MaxFrequencyModulation;
diff --git a/Engine/Plugins/Runtime/Synthesis/Source/Synthesis/Private/SourceEffectRingModulation.cpp b/Engine/Plugins/Runtime/Synthesis/Source/Synthesis/Private/SourceEffectRingModulation.cpp
--- a/Engine/Plugins/Runtime/Synthesis/Source/Synthesis/Private/SourceEffectRingModulation.cpp
+++ b/Engine/Plugins/Runtime/Synthesis/Source/Synthesis/Private/SourceEffectRingModulation.cpp
@@ -51,7 +51,7 @@ void FSourceEffectRingModulation::OnPresetChanged()
 	{
 		Audio::FMixerDevice* MixerDevice = (Audio::FMixerDevice*)AudioDeviceManager->GetAudioDeviceRaw(AudioDeviceId);
 		uint32 AudioBusId = Settings.AudioBusModulator->GetUniqueID();
-		Audio::FPatchOutputStrongPtr AudioBusPatchOutputPtr = MixerDevice->AddPatchForAudioBus(AudioBusId, 1.0f);
+		Audio::FPatchOutputStrongPtr AudioBusPatchOutputPtr = MixerDevice->AddPatchOutputForAudioBus(AudioBusId, MixerDevice->GetNumOutputFrames(), Settings.AudioBusModulator->GetNumChannels());
 		RingModulation.SetExternalPatchSource(AudioBusPatchOutputPtr);
 	}
 	else
diff --git a/Engine/Plugins/Runtime/Synthesis/Source/Synthesis/Private/SubmixEffectMultiBandCompressor.cpp b/Engine/Plugins/Runtime/Synthesis/Source/Synthesis/Private/SubmixEffectMultiBandCompressor.cpp
--- a/Engine/Plugins/Runtime/Synthesis/Source/Synthesis/Private/SubmixEffectMultiBandCompressor.cpp
+++ b/Engine/Plugins/Runtime/Synthesis/Source/Synthesis/Private/SubmixEffectMultiBandCompressor.cpp
@@ -226,7 +226,7 @@ bool FSubmixEffectMultibandCompressor::UpdateKeySourcePatch()
 				const uint32 ObjectId = KeySource.GetObjectId();
 				if (ObjectId != INDEX_NONE)
 				{
-					KeySource.Patch = MixerDevice->AddPatchForAudioBus(ObjectId, 1.0f /* PatchGain */);
+					KeySource.Patch = MixerDevice->AddPatchOutputForAudioBus(ObjectId, MixerDevice->GetNumOutputFrames(), KeySource.GetNumChannels());
 					if (KeySource.Patch.IsValid())
 					{
 						for (Audio::FDynamicsProcessor& DynamicsProcessor : DynamicsProcessors)
diff --git a/Engine/Source/Runtime/AudioAnalyzer/Private/AudioAnalyzer.cpp b/Engine/Source/Runtime/AudioAnalyzer/Private/AudioAnalyzer.cpp
--- a/Engine/Source/Runtime/AudioAnalyzer/Private/AudioAnalyzer.cpp
+++ b/Engine/Source/Runtime/AudioAnalyzer/Private/AudioAnalyzer.cpp
@@ -78,8 +78,8 @@ void UAudioAnalyzer::StartAnalyzing(UWorld* InWorld, UAudioBus* AudioBusToAnalyz
 	MixerDevice->StartAudioBus(AudioBusId, NumChannels, false);
 
 	// Get an output patch for the audio bus
-	PatchOutputStrongPtr = MixerDevice->AddPatchForAudioBus_GameThread(AudioBusId);
 	NumFramesPerBufferToAnalyze = MixerDevice->GetNumOutputFrames();
+	PatchOutputStrongPtr = MixerDevice->AddPatchOutputForAudioBus(AudioBusId, NumFramesPerBufferToAnalyze, NumChannels);
 
 	// Register this audio analyzer with the audio analyzer subsystem
 	// The subsystem will query this analyzer to see if it has enough audio to perform analysis.
diff --git a/Engine/Source/Runtime/AudioMixer/Private/AudioMixerBus.cpp b/Engine/Source/Runtime/AudioMixer/Private/AudioMixerBus.cpp
--- a/Engine/Source/Runtime/AudioMixer/Private/AudioMixerBus.cpp
+++ b/Engine/Source/Runtime/AudioMixer/Private/AudioMixerBus.cpp
@@ -163,7 +163,7 @@ namespace Audio
 			return;
 		}
 
-		int SampleCacheSize = SampleCache.Num();
+		int32 SampleCacheSize = SampleCache.Num();
 		if (0 <= SampleCacheSize)
 		{
 			float* SampleCacheBufferPtr = SampleCache.GetData();
@@ -220,7 +220,7 @@ namespace Audio
 		PatchSplitter.AddNewPatch(InPatchOutputStrongPtr);
 	}
 
-	void FMixerAudioBus::AddNewPatchInput(FPatchInput& InPatchInput)
+	void FMixerAudioBus::AddNewPatchInput(const FPatchInput& InPatchInput)
 	{
 		return PatchMixer.AddNewInput(InPatchInput);
 	}
diff --git a/Engine/Source/Runtime/AudioMixer/Private/AudioMixerDevice.cpp b/Engine/Source/Runtime/AudioMixer/Private/AudioMixerDevice.cpp
--- a/Engine/Source/Runtime/AudioMixer/Private/AudioMixerDevice.cpp
+++ b/Engine/Source/Runtime/AudioMixer/Private/AudioMixerDevice.cpp
@@ -81,6 +81,13 @@ FAutoConsoleVariableRef CVarDebugGeneratorFreq(
 	TEXT("0: Not Disabled, 1: SinTone, 2: WhiteNoise"),
 	ECVF_Default);
 
+static int32 AudioMixerPatchBufferBlocks = 3;
+FAutoConsoleVariableRef CVarAudioMixerPatchBufferBlocks(
+	TEXT("au.PatchBufferBlocks"),
+	AudioMixerPatchBufferBlocks,
+	TEXT("Determines the number of blocks that fit in a patch buffer."),
+	ECVF_Default);
+
 // Link to "Audio" profiling category
 CSV_DECLARE_CATEGORY_MODULE_EXTERN(AUDIOMIXERCORE_API, Audio);
 
@@ -2617,13 +2624,43 @@ namespace Audio
 		return SourceManager->IsAudioBusActive(InAudioBusId);
 	}
 
-
+	//UE_DEPRECATED(5.2, "AddPatchForAudioBus is deprecated.  Use AddPatchOutputForAudioBus.")
 	FPatchOutputStrongPtr FMixerDevice::AddPatchForAudioBus(uint32 InAudioBusId, float InPatchGain)
+	{
+		return AddPatchOutputForAudioBus(InAudioBusId, SourceManager->GetNumOutputFrames(), 8, InPatchGain);
+	}
+
+	//UE_DEPRECATED(5.2, "AddPatchForAudioBus_GameThread is deprecated.  Use AddPatchOutputForAudioBus.")
+	FPatchOutputStrongPtr FMixerDevice::AddPatchForAudioBus_GameThread(uint32 InAudioBusId, float InPatchGain)
+	{
+		PRAGMA_DISABLE_DEPRECATION_WARNINGS
+		return FMixerDevice::AddPatchForAudioBus(InAudioBusId, InPatchGain);
+		PRAGMA_ENABLE_DEPRECATION_WARNINGS
+	}
+
+	//UE_DEPRECATED(5.2, "This overload of AddPatchInputForAudioBus is deprecated and non-functional.  Use the overload that takes the number of frames and channels as parameters.")
+	void FMixerDevice::AddPatchInputForAudioBus(const FPatchInput& InPatchInput, uint32 InAudioBusId, float InPatchGain)
+	{
+	}
+
+	//UE_DEPRECATED(5.2, "AddPatchInputForAudioBus_GameThread is deprecated.  Use AddPatchInputForAudioBus.")
+	void FMixerDevice::AddPatchInputForAudioBus_GameThread(const FPatchInput& InPatchInput, uint32 InAudioBusId, float InPatchGain)
+	{
+		PRAGMA_DISABLE_DEPRECATION_WARNINGS
+		FMixerDevice::AddPatchInputForAudioBus(InPatchInput, InAudioBusId, InPatchGain);
+		PRAGMA_ENABLE_DEPRECATION_WARNINGS
+	}
+
+	FPatchOutputStrongPtr FMixerDevice::MakePatch(int32 InFrames, int32 InChannels, float InGain) const
+	{
+		int32 MaxSizeFrames = FMath::Max(InFrames, SourceManager->GetNumOutputFrames()), MinSizeFrames = FMath::Min(InFrames, SourceManager->GetNumOutputFrames());
+		return MakeShared<Audio::FPatchOutput, ESPMode::ThreadSafe>(AudioMixerPatchBufferBlocks * MinSizeFrames * FMath::DivideAndRoundUp(MaxSizeFrames, MinSizeFrames) * InChannels, InGain);
+	}
+
+	FPatchInput FMixerDevice::AddPatchInputForAudioBus(uint32 InAudioBusId, int32 InFrames, int32 InChannels, float InGain)
 	{
 		// This function is supporting adding audio bus patches from multiple threads (AT, ART, GT, and tasks) and is currently
 		// depending on a number of places where data lives, which accounts for the complexity here.
-		// The key idea here is to create and return a strong patch output ptr at roughly the size we expect to need then 
-		// pass the strong ptr down to the audio render thread that then is registered to the audio bus to start feeding audio to.
 		// This code needs a clean up to refactor everything into a true MPSC model, along with an MPSC refactor of the source manager
 		// and our command queues. Once we do that we can remove the code which branches based on the thread the request is coming from. 
 
@@ -2631,24 +2668,21 @@ namespace Audio
 		{
 			if (ActiveAudioBuses_GameThread.Find(InAudioBusId))
 			{
-				const int32 NumOutputFrames = SourceManager->GetNumOutputFrames();
-				FPatchOutputStrongPtr StrongOutputPtr = MakeShareable(new FPatchOutput(2 * NumOutputFrames, InPatchGain));
-				FAudioThread::RunCommandOnAudioThread([this, StrongOutputPtr, InAudioBusId]() mutable
+				FPatchInput PatchInput = MakePatch(InFrames, InChannels, InGain);
+				FAudioThread::RunCommandOnAudioThread([this, PatchInput, InAudioBusId]()
 				{
-					SourceManager->AddPatchOutputForAudioBus_AudioThread(InAudioBusId, StrongOutputPtr);
+					SourceManager->AddPatchInputForAudioBus_AudioThread(InAudioBusId, PatchInput);
 				});
-				return StrongOutputPtr;
+				return PatchInput;
 			}
-			UE_LOG(LogAudioMixer, Warning, TEXT("Unable to add a patch output for audio bus because audio bus id '%d' is not active."), InAudioBusId);
-			return nullptr;
+			UE_LOG(LogAudioMixer, Warning, TEXT("Unable to add a patch input for audio bus because audio bus id '%d' is not active."), InAudioBusId);
+			return FPatchInput();
 		}
 		else if (IsInAudioThread())
 		{
-			int32 NumOutputFrames = SourceManager->GetNumOutputFrames();
-			FPatchOutputStrongPtr StrongOutputPtr = MakeShareable(new FPatchOutput(2 * NumOutputFrames, InPatchGain));
-
-			SourceManager->AddPatchOutputForAudioBus_AudioThread(InAudioBusId, StrongOutputPtr);
-			return StrongOutputPtr;
+			FPatchInput PatchInput = MakePatch(InFrames, InChannels, InGain);
+			SourceManager->AddPatchInputForAudioBus_AudioThread(InAudioBusId, PatchInput);
+			return PatchInput;
 		}
 		else if (IsAudioRenderingThread())
 		{
@@ -2657,27 +2691,23 @@ namespace Audio
 			const int32 NumChannels = SourceManager->GetAudioBusNumChannels(InAudioBusId);
 			if (NumChannels > 0)
 			{
-				const int32 NumOutputFrames = SourceManager->GetNumOutputFrames();
-				FPatchOutputStrongPtr StrongOutputPtr = MakeShareable(new FPatchOutput(NumOutputFrames * NumChannels, InPatchGain));
-				SourceManager->AddPatchOutputForAudioBus(InAudioBusId, StrongOutputPtr);
-				return StrongOutputPtr;
+				FPatchInput PatchInput = MakePatch(InFrames, InChannels, InGain);
+				SourceManager->AddPatchInputForAudioBus(InAudioBusId, PatchInput);
+				return PatchInput;
 			}
 
-			return nullptr;
+			return FPatchInput();
 		}
 		else
 		{
-			// Need to make a strong output patch even if this is not going to ever write to it since the bus may not be running
-			const int32 NumOutputFrames = SourceManager->GetNumOutputFrames();
-			FPatchOutputStrongPtr StrongOutputPtr = MakeShareable(new FPatchOutput(3 * NumOutputFrames, InPatchGain));
-
-			GameThreadMPSCCommand([this, StrongOutputPtr, InAudioBusId]
+			FPatchInput PatchInput = MakePatch(InFrames, InChannels, InGain);
+			GameThreadMPSCCommand([this, PatchInput, InAudioBusId]() mutable
 			{
 				if (ActiveAudioBuses_GameThread.Find(InAudioBusId))
 				{
-					FAudioThread::RunCommandOnAudioThread([this, InAudioBusId, StrongOutputPtr]() mutable
+					FAudioThread::RunCommandOnAudioThread([this, InAudioBusId, PatchInput = MoveTemp(PatchInput)]()
 					{
-						SourceManager->AddPatchOutputForAudioBus_AudioThread(InAudioBusId, StrongOutputPtr);
+						SourceManager->AddPatchInputForAudioBus_AudioThread(InAudioBusId, PatchInput);
 					});
 				}
 				else
@@ -2685,41 +2715,11 @@ namespace Audio
 					UE_LOG(LogAudioMixer, Warning, TEXT("Unable to add a patch output for audio bus because audio bus id '%d' is not active."), InAudioBusId);
 				}
 			});
-			return StrongOutputPtr;
+			return PatchInput;
 		}
 	}
 
-	FPatchOutputStrongPtr FMixerDevice::AddPatchForAudioBus_GameThread(uint32 InAudioBusId, float InPatchGain)
-	{
-		FPatchOutputStrongPtr StrongOutputPtr;
-		
-		if (IsInGameThread())
-		{
-			FActiveBusData* BusData = ActiveAudioBuses_GameThread.Find(InAudioBusId);
-			if (BusData)
-			{
-				int32 NumOutputFrames = SourceManager->GetNumOutputFrames();
-				int32 BusNumChannels = BusData->NumChannels;
-				StrongOutputPtr = MakeShareable(new FPatchOutput(NumOutputFrames * BusData->NumChannels, InPatchGain));
-
-				FAudioThread::RunCommandOnAudioThread([this, InAudioBusId, StrongOutputPtr]() mutable
-				{
-					SourceManager->AddPatchOutputForAudioBus_AudioThread(InAudioBusId, StrongOutputPtr);
-				});
-			}
-			else
-			{
-				UE_LOG(LogAudioMixer, Warning, TEXT("Unable to add a patch output for audio bus because audio bus id '%d' is not active."), InAudioBusId);
-			}
-		}
-		else
-		{
-			UE_LOG(LogAudioMixer, Warning, TEXT("AddPatchForAudioBus can only be called from the game thread."));
-		}
-		return StrongOutputPtr;
-	}
-
-	void FMixerDevice::AddPatchInputForAudioBus(const FPatchInput& InPatchInput, uint32 InAudioBusId, float InPatchGain)
+	FPatchOutputStrongPtr FMixerDevice::AddPatchOutputForAudioBus(uint32 InAudioBusId, int32 InFrames, int32 InChannels, float InGain)
 	{
 		// This function is supporting adding audio bus patches from multiple threads (AT, ART, GT, and tasks) and is currently
 		// depending on a number of places where data lives, which accounts for the complexity here.
@@ -2730,36 +2730,46 @@ namespace Audio
 		{
 			if (ActiveAudioBuses_GameThread.Find(InAudioBusId))
 			{
-				FAudioThread::RunCommandOnAudioThread([this, PatchInput = InPatchInput, InAudioBusId]() mutable
+				FPatchOutputStrongPtr PatchOutput = MakePatch(InFrames, InChannels, InGain);
+				FAudioThread::RunCommandOnAudioThread([this, PatchOutput, InAudioBusId]()
 				{
-					SourceManager->AddPatchInputForAudioBus_AudioThread(InAudioBusId, PatchInput);
+					SourceManager->AddPatchOutputForAudioBus_AudioThread(InAudioBusId, PatchOutput);
 				});
-				return;
+				return PatchOutput;
 			}
 			UE_LOG(LogAudioMixer, Warning, TEXT("Unable to add a patch input for audio bus because audio bus id '%d' is not active."), InAudioBusId);
-			return;
+			return nullptr;
 		}
 		else if (IsInAudioThread())
 		{
-			SourceManager->AddPatchInputForAudioBus_AudioThread(InAudioBusId, InPatchInput);
-			return;
+			FPatchOutputStrongPtr PatchOutput = MakePatch(InFrames, InChannels, InGain);
+			SourceManager->AddPatchOutputForAudioBus_AudioThread(InAudioBusId, PatchOutput);
+			return PatchOutput;
 		}
 		else if (IsAudioRenderingThread())
 		{
 			check(SourceManager);
-			FPatchInput PatchInput(InPatchInput);
-			SourceManager->AddPatchInputForAudioBus(InAudioBusId, PatchInput);
-			return;
+
+			const int32 NumChannels = SourceManager->GetAudioBusNumChannels(InAudioBusId);
+			if (NumChannels > 0)
+			{
+				FPatchOutputStrongPtr PatchOutput = MakePatch(InFrames, InChannels, InGain);
+				SourceManager->AddPatchOutputForAudioBus(InAudioBusId, PatchOutput);
+				return PatchOutput;
+			}
+
+			return nullptr;
 		}
 		else
 		{
-			GameThreadMPSCCommand([this, PatchInput = InPatchInput, InAudioBusId]
+			FPatchOutputStrongPtr PatchOutput = MakePatch(InFrames, InChannels, InGain);
+			GameThreadMPSCCommand([this, PatchOutput, InAudioBusId]() mutable
 			{
 				if (ActiveAudioBuses_GameThread.Find(InAudioBusId))
 				{
-					FAudioThread::RunCommandOnAudioThread([this, InAudioBusId, PatchInput = PatchInput]() mutable
+					FAudioThread::RunCommandOnAudioThread([this, InAudioBusId, PatchOutput = MoveTemp(PatchOutput)]()
 					{
-						SourceManager->AddPatchInputForAudioBus_AudioThread(InAudioBusId, PatchInput);
+						SourceManager->AddPatchOutputForAudioBus_AudioThread(InAudioBusId, PatchOutput);
 					});
 				}
 				else
@@ -2767,30 +2777,7 @@ namespace Audio
 					UE_LOG(LogAudioMixer, Warning, TEXT("Unable to add a patch output for audio bus because audio bus id '%d' is not active."), InAudioBusId);
 				}
 			});
-			return;
-		}
-	}
-
-	void FMixerDevice::AddPatchInputForAudioBus_GameThread(const FPatchInput& InPatchInput, uint32 InAudioBusId, float InPatchGain)
-	{
-		if (IsInGameThread())
-		{
-			FActiveBusData* BusData = ActiveAudioBuses_GameThread.Find(InAudioBusId);
-			if (BusData)
-			{
-				FAudioThread::RunCommandOnAudioThread([this, InAudioBusId, PatchInput = InPatchInput]() mutable
-				{
-					SourceManager->AddPatchInputForAudioBus_AudioThread(InAudioBusId, PatchInput);
-				});
-			}
-			else
-			{
-				UE_LOG(LogAudioMixer, Warning, TEXT("Unable to add a patch input for audio bus because audio bus id '%d' is not active."), InAudioBusId);
-			}
-		}
-		else
-		{
-			UE_LOG(LogAudioMixer, Warning, TEXT("AddPatchForAudioBus can only be called from the game thread."));
+			return PatchOutput;
 		}
 	}
 
diff --git a/Engine/Source/Runtime/AudioMixer/Private/AudioMixerSourceManager.cpp b/Engine/Source/Runtime/AudioMixer/Private/AudioMixerSourceManager.cpp
--- a/Engine/Source/Runtime/AudioMixer/Private/AudioMixerSourceManager.cpp
+++ b/Engine/Source/Runtime/AudioMixer/Private/AudioMixerSourceManager.cpp
@@ -1244,7 +1244,7 @@ namespace Audio
 		}, AUDIO_MIXER_THREAD_COMMAND_STRING("AddPatchOutputForAudioBus_AudioThread()"));
 	}
 
-	void FMixerSourceManager::AddPatchInputForAudioBus(uint32 InAudioBusId, FPatchInput& InPatchInput)
+	void FMixerSourceManager::AddPatchInputForAudioBus(uint32 InAudioBusId, const FPatchInput& InPatchInput)
 	{
 		AUDIO_MIXER_CHECK_AUDIO_PLAT_THREAD(MixerDevice);
 		if (MixerDevice->IsAudioRenderingThread())
@@ -1258,7 +1258,7 @@ namespace Audio
 		else
 		{
 			// Queue up the command via MPSC command queue
-			AudioMixerThreadMPSCCommand([this, InAudioBusId, InPatchInput]() mutable
+			AudioMixerThreadMPSCCommand([this, InAudioBusId, InPatchInput]()
 			{
 				AddPatchInputForAudioBus(InAudioBusId, InPatchInput);
 			});
@@ -1267,9 +1267,9 @@ namespace Audio
 
 	void FMixerSourceManager::AddPatchInputForAudioBus_AudioThread(uint32 InAudioBusId, const FPatchInput& InPatchInput)
 	{
-		AudioMixerThreadCommand([this, InAudioBusId, PatchInput = InPatchInput]() mutable
+		AudioMixerThreadCommand([this, InAudioBusId, InPatchInput]()
 		{
-			AddPatchInputForAudioBus(InAudioBusId, PatchInput);
+			AddPatchInputForAudioBus(InAudioBusId, InPatchInput);
 		}, AUDIO_MIXER_THREAD_COMMAND_STRING("AddPatchInputForAudioBus_AudioThread()"));
 	}
 
diff --git a/Engine/Source/Runtime/AudioMixer/Private/AudioMixerSourceManager.h b/Engine/Source/Runtime/AudioMixer/Private/AudioMixerSourceManager.h
--- a/Engine/Source/Runtime/AudioMixer/Private/AudioMixerSourceManager.h
+++ b/Engine/Source/Runtime/AudioMixer/Private/AudioMixerSourceManager.h
@@ -209,7 +209,7 @@ namespace Audio
 		void AddPatchOutputForAudioBus_AudioThread(uint32 InAudioBusId, const FPatchOutputStrongPtr& InPatchOutputStrongPtr);
 
 		// Adds a patch input for an audio bus
-		void AddPatchInputForAudioBus(uint32 InAudioBusId, FPatchInput& InPatchInput);
+		void AddPatchInputForAudioBus(uint32 InAudioBusId, const FPatchInput& InPatchInput);
 
 		// Adds a patch input for an audio bus from the Audio Thread
 		void AddPatchInputForAudioBus_AudioThread(uint32 InAudioBusId, const FPatchInput& InPatchInput);
diff --git a/Engine/Source/Runtime/AudioMixer/Private/Effects/AudioMixerSubmixEffectDynamicsProcessor.cpp b/Engine/Source/Runtime/AudioMixer/Private/Effects/AudioMixerSubmixEffectDynamicsProcessor.cpp
--- a/Engine/Source/Runtime/AudioMixer/Private/Effects/AudioMixerSubmixEffectDynamicsProcessor.cpp
+++ b/Engine/Source/Runtime/AudioMixer/Private/Effects/AudioMixerSubmixEffectDynamicsProcessor.cpp
@@ -198,7 +198,7 @@ bool FSubmixEffectDynamicsProcessor::UpdateKeySourcePatch()
 				const uint32 ObjectId = KeySource.GetObjectId();
 				if (ObjectId != INDEX_NONE)
 				{
-					KeySource.Patch = MixerDevice->AddPatchForAudioBus(ObjectId, 1.0f /* PatchGain */);
+					KeySource.Patch = MixerDevice->AddPatchOutputForAudioBus(ObjectId, MixerDevice->GetNumOutputFrames(), KeySource.GetNumChannels());
 					if (KeySource.Patch.IsValid())
 					{
 						DynamicsProcessor.SetKeyNumChannels(KeySource.GetNumChannels());
diff --git a/Engine/Source/Runtime/AudioMixer/Public/AudioMixerBus.h b/Engine/Source/Runtime/AudioMixer/Public/AudioMixerBus.h
--- a/Engine/Source/Runtime/AudioMixer/Public/AudioMixerBus.h
+++ b/Engine/Source/Runtime/AudioMixer/Public/AudioMixerBus.h
@@ -26,7 +26,7 @@ namespace Audio
 		void AddNewPatchOutput(const FPatchOutputStrongPtr& InPatchOutputStrongPtr);
 
 		// Allow anybody to write audio into this audio bus from any thread.
-		void AddNewPatchInput(FPatchInput& InPatchInput);
+		void AddNewPatchInput(const FPatchInput& InPatchInput);
 
 		// Allow anybody to write audio into this audio bus from any thread.
 		void RemovePatchInput(const FPatchInput& InPatchInput);
diff --git a/Engine/Source/Runtime/AudioMixer/Public/AudioMixerDevice.h b/Engine/Source/Runtime/AudioMixer/Public/AudioMixerDevice.h
--- a/Engine/Source/Runtime/AudioMixer/Public/AudioMixerDevice.h
+++ b/Engine/Source/Runtime/AudioMixer/Public/AudioMixerDevice.h
@@ -299,12 +299,21 @@ namespace Audio
 		virtual void StopAudioBus(uint32 InAudioBusId) override;
 		virtual bool IsAudioBusActive(uint32 InAudioBusId) const override;
 
+		UE_DEPRECATED(5.2, "AddPatchForAudioBus is deprecated.  Use AddPatchOutputForAudioBus.")
 		virtual FPatchOutputStrongPtr AddPatchForAudioBus(uint32 InAudioBusId, float InPatchGain = 1.0f) override;
+
+		UE_DEPRECATED(5.2, "AddPatchForAudioBus_GameThread is deprecated.  Use AddPatchOutputForAudioBus.")
 		virtual FPatchOutputStrongPtr AddPatchForAudioBus_GameThread(uint32 InAudioBusId, float InPatchGain = 1.0f) override;
 
+		UE_DEPRECATED(5.2, "This overload of AddPatchInputForAudioBus is deprecated and non-functional.  Use the overload that takes the number of frames and channels as parameters.")
 		virtual void AddPatchInputForAudioBus(const FPatchInput& InPatchInput, uint32 InAudioBusId, float InPatchGain = 1.0f) override;
+
+		UE_DEPRECATED(5.2, "AddPatchInputForAudioBus_GameThread is deprecated.  Use AddPatchInputForAudioBus.")
 		virtual void AddPatchInputForAudioBus_GameThread(const FPatchInput& InPatchInput, uint32 InAudioBusId, float InPatchGain = 1.0f) override;
 
+		virtual FPatchInput AddPatchInputForAudioBus(uint32 InAudioBusId, int32 InFrames, int32 InChannels, float InGain = 1.f) override;
+		virtual FPatchOutputStrongPtr AddPatchOutputForAudioBus(uint32 InAudioBusId, int32 InFrames, int32 InChannels, float InGain = 1.f) override;
+
 		// Clock Manager for quantized event handling on Audio Render Thread
 		FQuartzClockManager QuantizedEventClockManager;
 
@@ -364,6 +373,8 @@ namespace Audio
 		void PumpCommandQueue();
 		void PumpGameThreadCommandQueue();
 		
+		FPatchOutputStrongPtr MakePatch(int32 InFrames, int32 InChannels, float InGain) const;
+
 		TArray<USoundSubmix*> MasterSubmixes;
 		TArray<FMixerSubmixPtr> MasterSubmixInstances;
 
diff --git a/Engine/Source/Runtime/Engine/Public/AudioDevice.h b/Engine/Source/Runtime/Engine/Public/AudioDevice.h
--- a/Engine/Source/Runtime/Engine/Public/AudioDevice.h
+++ b/Engine/Source/Runtime/Engine/Public/AudioDevice.h
@@ -960,24 +960,38 @@ public:
 		return false;
 	}
 
+	UE_DEPRECATED(5.2, "AddPatchForAudioBus is deprecated.  Use AddPatchOutputForAudioBus.")
 	virtual Audio::FPatchOutputStrongPtr AddPatchForAudioBus(uint32 InAudioBusId, float InPatchGain = 1.0f)
 	{
 		return nullptr;
 	}
 
+	UE_DEPRECATED(5.2, "AddPatchForAudioBus_GameThread is deprecated.  Use AddPatchOutputForAudioBus.")
 	virtual Audio::FPatchOutputStrongPtr AddPatchForAudioBus_GameThread(uint32 InAudioBusId, float InPatchGain = 1.0f)
 	{
 		return nullptr;
 	}
-	
+
+	UE_DEPRECATED(5.2, "This overload of AddPatchInputForAudioBus is deprecated.  Use the overload that takes the number of frames and channels as parameters.")
 	virtual void AddPatchInputForAudioBus(const Audio::FPatchInput& InPatchInput, uint32 InAudioBusId, float InPatchGain = 1.0f)
 	{
 	}
 
+	UE_DEPRECATED(5.2, "AddPatchInputForAudioBus_GameThread is deprecated.  Use AddPatchInputForAudioBus.")
 	virtual void AddPatchInputForAudioBus_GameThread(const Audio::FPatchInput& InPatchInput, uint32 InAudioBusId, float InPatchGain = 1.0f)
 	{
 	}
 
+	virtual Audio::FPatchInput AddPatchInputForAudioBus(uint32 InAudioBusId, int32 InFrames, int32 InChannels, float InGain = 1.f)
+	{
+		return Audio::FPatchInput();
+	}
+
+	virtual Audio::FPatchOutputStrongPtr AddPatchOutputForAudioBus(uint32 InAudioBusId, int32 InFrames, int32 InChannels, float InGain = 1.f)
+	{
+		return Audio::FPatchOutputStrongPtr();
+	}
+
 	virtual void InitSoundEffectPresets() {}
 
 	/**
