[Soundfields] Fix for invalid soundfield settings after changing encoding type.

diff --git a/Engine/Source/Runtime/AudioMixer/Private/AudioMixerSubmix.cpp b/Engine/Source/Runtime/AudioMixer/Private/AudioMixerSubmix.cpp
--- a/Engine/Source/Runtime/AudioMixer/Private/AudioMixerSubmix.cpp
+++ b/Engine/Source/Runtime/AudioMixer/Private/AudioMixerSubmix.cpp
@@ -236,6 +236,14 @@ namespace Audio
 					InitInternal();
 				});
 			}
+			else if (const USoundfieldSubmix* SoundfieldSubmix = Cast<const USoundfieldSubmix>(OwningSubmixObject))
+			{
+				ISoundfieldFactory* SoundfieldFactory = SoundfieldSubmix->GetSoundfieldFactoryForSubmix();
+				const USoundfieldEncodingSettingsBase* EncodingSettings = SoundfieldSubmix->GetSoundfieldEncodingSettings();
+
+				TArray<USoundfieldEffectBase*> Effects = SoundfieldSubmix->GetSoundfieldProcessors();
+				SetupSoundfieldStreams(EncodingSettings, Effects, SoundfieldFactory);
+			}
 		}
 	}
 
