Audio bus subsystem refactor checkpoint
- Move audio bus functionality from Mixer Device to new Audio Bus Subsystem
- Add new audio bus key struct so that FMixerAudioBuses do not have to be created with UObject ids

diff --git a/Engine/Plugins/Runtime/AudioModulation/Source/AudioModulation/Private/Generators/SoundModulationEnvelopeFollower.cpp b/Engine/Plugins/Runtime/AudioModulation/Source/AudioModulation/Private/Generators/SoundModulationEnvelopeFollower.cpp
--- a/Engine/Plugins/Runtime/AudioModulation/Source/AudioModulation/Private/Generators/SoundModulationEnvelopeFollower.cpp
+++ b/Engine/Plugins/Runtime/AudioModulation/Source/AudioModulation/Private/Generators/SoundModulationEnvelopeFollower.cpp
@@ -3,6 +3,7 @@
 #include "Generators/SoundModulationEnvelopeFollower.h"
 
 #include "Algo/MaxElement.h"
+#include "AudioBusSubsystem.h"
 #include "AudioDefines.h"
 #include "AudioDeviceManager.h"
 #include "AudioMixerDevice.h"
@@ -203,7 +204,9 @@ namespace AudioModulation
 					FAudioDevice* AudioDevice = DeviceManager->GetAudioDeviceRaw(AudioDeviceId);
 					if (Audio::FMixerDevice* MixerDevice = static_cast<Audio::FMixerDevice*>(AudioDevice))
 					{
-						AudioBusPatch = MixerDevice->AddPatchOutputForAudioBus(BusId, MixerDevice->GetNumOutputFrames(), EnvelopeFollower.GetNumChannels(), Gain);
+						UAudioBusSubsystem* AudioBusSubsystem = MixerDevice->GetSubsystem<UAudioBusSubsystem>();
+						check(AudioBusSubsystem);
+						AudioBusPatch = AudioBusSubsystem->AddPatchOutputForAudioBus(Audio::FAudioBusKey(BusId), MixerDevice->GetNumOutputFrames(), EnvelopeFollower.GetNumChannels(), Gain);
 						bBusRequiresPatch = false;
 					}
 				}
diff --git a/Engine/Plugins/Runtime/Metasound/Source/MetasoundEngine/MetasoundEngine.Build.cs b/Engine/Plugins/Runtime/Metasound/Source/MetasoundEngine/MetasoundEngine.Build.cs
--- a/Engine/Plugins/Runtime/Metasound/Source/MetasoundEngine/MetasoundEngine.Build.cs
+++ b/Engine/Plugins/Runtime/Metasound/Source/MetasoundEngine/MetasoundEngine.Build.cs
@@ -24,6 +24,7 @@ namespace UnrealBuildTool.Rules
 					"Core",
 					"AudioExtensions",
 					"AudioCodecEngine",
+					"AudioMixer",
 					"MetasoundGraphCore",
 					"MetasoundGenerator",
 					"SignalProcessing"
diff --git a/Engine/Plugins/Runtime/Metasound/Source/MetasoundEngine/Private/MetasoundAudioBusReaderNode.cpp b/Engine/Plugins/Runtime/Metasound/Source/MetasoundEngine/Private/MetasoundAudioBusReaderNode.cpp
--- a/Engine/Plugins/Runtime/Metasound/Source/MetasoundEngine/Private/MetasoundAudioBusReaderNode.cpp
+++ b/Engine/Plugins/Runtime/Metasound/Source/MetasoundEngine/Private/MetasoundAudioBusReaderNode.cpp
@@ -1,6 +1,7 @@
 // Copyright Epic Games, Inc. All Rights Reserved.
 
 #include "AudioDevice.h"
+#include "AudioBusSubsystem.h"
 #include "DSP/ConvertDeinterleave.h"
 #include "Internationalization/Text.h"
 #include "MediaPacket.h"
@@ -111,12 +112,14 @@ namespace Metasound
 					{
 						// Start the audio bus in case it's not already started
 						AudioBusChannels = AudioBusProxy->NumChannels;
-						AudioDevice->StartAudioBus(AudioBusProxy->AudioBusId, AudioBusChannels, false);
-
-						int32 BlockSizeFrames = InSettings.GetNumFramesPerBlock();
+						Audio::FAudioBusKey AudioBusKey = Audio::FAudioBusKey(AudioBusProxy->AudioBusId);
+						UAudioBusSubsystem* AudioBusSubsystem = AudioDevice->GetSubsystem<UAudioBusSubsystem>();
+						check(AudioBusSubsystem);
+						AudioBusSubsystem->StartAudioBus(AudioBusKey, AudioBusChannels, false);
 
 						// Create a bus patch output with enough room for the number of samples we expect and some buffering
-						AudioBusPatchOutput = AudioDevice->AddPatchOutputForAudioBus(AudioBusProxy->AudioBusId, BlockSizeFrames, AudioBusChannels);
+						int32 BlockSizeFrames = InSettings.GetNumFramesPerBlock();
+						AudioBusPatchOutput = AudioBusSubsystem->AddPatchOutputForAudioBus(AudioBusKey, BlockSizeFrames, AudioBusChannels);
 					}
 				}
 			}
diff --git a/Engine/Plugins/Runtime/Metasound/Source/MetasoundEngine/Private/MetasoundAudioBusWriterNode.cpp b/Engine/Plugins/Runtime/Metasound/Source/MetasoundEngine/Private/MetasoundAudioBusWriterNode.cpp
--- a/Engine/Plugins/Runtime/Metasound/Source/MetasoundEngine/Private/MetasoundAudioBusWriterNode.cpp
+++ b/Engine/Plugins/Runtime/Metasound/Source/MetasoundEngine/Private/MetasoundAudioBusWriterNode.cpp
@@ -1,6 +1,7 @@
 // Copyright Epic Games, Inc. All Rights Reserved.
 
-#include "AudioMixerDevice.h"
+#include "AudioBusSubsystem.h"
+#include "AudioDevice.h"
 #include "Internationalization/Text.h"
 #include "MediaPacket.h"
 #include "MetasoundAudioBuffer.h"
@@ -108,15 +109,18 @@ namespace Metasound
 				{
 					if (FAudioDevice* AudioDevice = ADM->GetAudioDeviceRaw(AudioDeviceId))
 					{
-						// Start the audio bus in case it's not already started
+						UAudioBusSubsystem* AudioBusSubsystem = AudioDevice->GetSubsystem<UAudioBusSubsystem>();
+						check(AudioBusSubsystem);
+						Audio::FAudioBusKey AudioBusKey = Audio::FAudioBusKey(AudioBusProxy->AudioBusId);	
+						// Start the audio bus in case it's not already started					
 						AudioBusChannels = AudioBusProxy->NumChannels;
-						AudioDevice->StartAudioBus(AudioBusProxy->AudioBusId, AudioBusChannels, false);
+						AudioBusSubsystem->StartAudioBus(AudioBusKey, AudioBusChannels, false);
 
 						BlockSizeFrames = InParams.OperatorSettings.GetNumFramesPerBlock();
 						InterleavedBuffer.AddZeroed(BlockSizeFrames * AudioBusChannels);
 
 						// Create a bus patch input with enough room for the number of samples we expect and some buffering
-						AudioBusPatchInput = AudioDevice->AddPatchInputForAudioBus(AudioBusProxy->AudioBusId, BlockSizeFrames, AudioBusChannels);
+						AudioBusPatchInput = AudioBusSubsystem->AddPatchInputForAudioBus(AudioBusKey, BlockSizeFrames, AudioBusChannels);
 					}
 				}
 			}
diff --git a/Engine/Plugins/Runtime/Synthesis/Source/Synthesis/Private/SourceEffectFilter.cpp b/Engine/Plugins/Runtime/Synthesis/Source/Synthesis/Private/SourceEffectFilter.cpp
--- a/Engine/Plugins/Runtime/Synthesis/Source/Synthesis/Private/SourceEffectFilter.cpp
+++ b/Engine/Plugins/Runtime/Synthesis/Source/Synthesis/Private/SourceEffectFilter.cpp
@@ -1,9 +1,10 @@
 // Copyright Epic Games, Inc. All Rights Reserved.
 
 #include "SourceEffects/SourceEffectFilter.h"
-#include "DSP/FloatArrayMath.h"
-#include "DSP/Dsp.h"
+#include "AudioBusSubsystem.h"
 #include "AudioMixerDevice.h"
+#include "DSP/Dsp.h"
+#include "DSP/FloatArrayMath.h"
 
 #include UE_INLINE_GENERATED_CPP_BY_NAME(SourceEffectFilter)
 
@@ -114,7 +115,9 @@ void FSourceEffectFilter::OnPresetChanged()
 					FAudioBusModulationData& NewModData = ModData.AddDefaulted_GetRef();
 
 					uint32 AudioBusId = BusModulationSettings.AudioBus->GetUniqueID();
-					NewModData.AudioBusPatch = MixerDevice->AddPatchOutputForAudioBus(AudioBusId, MixerDevice->GetNumOutputFrames(), BusModulationSettings.AudioBus->GetNumChannels());
+					UAudioBusSubsystem* AudioBusSubsystem = MixerDevice->GetSubsystem<UAudioBusSubsystem>();
+					check(AudioBusSubsystem);
+					NewModData.AudioBusPatch = AudioBusSubsystem->AddPatchOutputForAudioBus(Audio::FAudioBusKey(AudioBusId), MixerDevice->GetNumOutputFrames(), BusModulationSettings.AudioBus->GetNumChannels());
 
 					NewModData.MinFreqModValue = BusModulationSettings.MinFrequencyModulation;
 					NewModData.MaxFreqModValue = BusModulationSettings.MaxFrequencyModulation;
diff --git a/Engine/Plugins/Runtime/Synthesis/Source/Synthesis/Private/SourceEffectRingModulation.cpp b/Engine/Plugins/Runtime/Synthesis/Source/Synthesis/Private/SourceEffectRingModulation.cpp
--- a/Engine/Plugins/Runtime/Synthesis/Source/Synthesis/Private/SourceEffectRingModulation.cpp
+++ b/Engine/Plugins/Runtime/Synthesis/Source/Synthesis/Private/SourceEffectRingModulation.cpp
@@ -1,7 +1,7 @@
 // Copyright Epic Games, Inc. All Rights Reserved.
 
 #include "SourceEffects/SourceEffectRingModulation.h"
-#include "AudioDevice.h"
+#include "AudioBusSubsystem.h"
 #include "AudioMixerDevice.h"
 #include "AudioMixerBus.h"
 #include "DSP/MultithreadedPatching.h"
@@ -51,7 +51,9 @@ void FSourceEffectRingModulation::OnPresetChanged()
 	{
 		Audio::FMixerDevice* MixerDevice = (Audio::FMixerDevice*)AudioDeviceManager->GetAudioDeviceRaw(AudioDeviceId);
 		uint32 AudioBusId = Settings.AudioBusModulator->GetUniqueID();
-		Audio::FPatchOutputStrongPtr AudioBusPatchOutputPtr = MixerDevice->AddPatchOutputForAudioBus(AudioBusId, MixerDevice->GetNumOutputFrames(), Settings.AudioBusModulator->GetNumChannels());
+		UAudioBusSubsystem* AudioBusSubsystem = MixerDevice->GetSubsystem<UAudioBusSubsystem>();
+		check(AudioBusSubsystem);
+		Audio::FPatchOutputStrongPtr AudioBusPatchOutputPtr = AudioBusSubsystem->AddPatchOutputForAudioBus(Audio::FAudioBusKey(AudioBusId), MixerDevice->GetNumOutputFrames(), Settings.AudioBusModulator->GetNumChannels());
 		RingModulation.SetExternalPatchSource(AudioBusPatchOutputPtr);
 	}
 	else
diff --git a/Engine/Plugins/Runtime/Synthesis/Source/Synthesis/Private/SubmixEffectMultiBandCompressor.cpp b/Engine/Plugins/Runtime/Synthesis/Source/Synthesis/Private/SubmixEffectMultiBandCompressor.cpp
--- a/Engine/Plugins/Runtime/Synthesis/Source/Synthesis/Private/SubmixEffectMultiBandCompressor.cpp
+++ b/Engine/Plugins/Runtime/Synthesis/Source/Synthesis/Private/SubmixEffectMultiBandCompressor.cpp
@@ -2,6 +2,7 @@
 
 #include "SubmixEffects/SubmixEffectMultiBandCompressor.h"
 
+#include "AudioBusSubsystem.h"
 #include "AudioDeviceManager.h"
 #include "AudioMixerDevice.h"
 #include "AudioMixerSubmix.h"
@@ -226,7 +227,9 @@ bool FSubmixEffectMultibandCompressor::UpdateKeySourcePatch()
 				const uint32 ObjectId = KeySource.GetObjectId();
 				if (ObjectId != INDEX_NONE)
 				{
-					KeySource.Patch = MixerDevice->AddPatchOutputForAudioBus(ObjectId, MixerDevice->GetNumOutputFrames(), KeySource.GetNumChannels());
+					UAudioBusSubsystem* AudioBusSubsystem = MixerDevice->GetSubsystem<UAudioBusSubsystem>();
+					check(AudioBusSubsystem);
+					KeySource.Patch = AudioBusSubsystem->AddPatchOutputForAudioBus(ObjectId, MixerDevice->GetNumOutputFrames(), KeySource.GetNumChannels());
 					if (KeySource.Patch.IsValid())
 					{
 						for (Audio::FDynamicsProcessor& DynamicsProcessor : DynamicsProcessors)
diff --git a/Engine/Source/Runtime/AudioAnalyzer/Private/AudioAnalyzer.cpp b/Engine/Source/Runtime/AudioAnalyzer/Private/AudioAnalyzer.cpp
--- a/Engine/Source/Runtime/AudioAnalyzer/Private/AudioAnalyzer.cpp
+++ b/Engine/Source/Runtime/AudioAnalyzer/Private/AudioAnalyzer.cpp
@@ -8,6 +8,7 @@
 #include "AudioMixerDevice.h"
 #include "AudioDeviceManager.h"
 #include "AudioAnalyzerSubsystem.h"
+#include "AudioBusSubsystem.h"
 
 #include UE_INLINE_GENERATED_CPP_BY_NAME(AudioAnalyzer)
 
@@ -75,11 +76,15 @@ void UAudioAnalyzer::StartAnalyzing(UWorld* InWorld, UAudioBus* AudioBusToAnalyz
 	// Start the audio bus. This won't do anythign if the bus is already started elsewhere.
 	uint32 AudioBusId = AudioBus->GetUniqueID();
 	int32 NumChannels = (int32)AudioBus->AudioBusChannels + 1;
-	MixerDevice->StartAudioBus(AudioBusId, NumChannels, false);
+
+	UAudioBusSubsystem* AudioBusSubsystem = MixerDevice->GetSubsystem<UAudioBusSubsystem>();
+	check(AudioBusSubsystem);
+	Audio::FAudioBusKey AudioBusKey = Audio::FAudioBusKey(AudioBusId);
+	AudioBusSubsystem->StartAudioBus(AudioBusKey, NumChannels, false);
 
 	// Get an output patch for the audio bus
 	NumFramesPerBufferToAnalyze = MixerDevice->GetNumOutputFrames();
-	PatchOutputStrongPtr = MixerDevice->AddPatchOutputForAudioBus(AudioBusId, NumFramesPerBufferToAnalyze, NumChannels);
+	PatchOutputStrongPtr = AudioBusSubsystem->AddPatchOutputForAudioBus(AudioBusKey, NumFramesPerBufferToAnalyze, NumChannels);
 
 	// Register this audio analyzer with the audio analyzer subsystem
 	// The subsystem will query this analyzer to see if it has enough audio to perform analysis.
diff --git a/Engine/Source/Runtime/AudioMixer/Private/AudioBusSubsystem.cpp b/Engine/Source/Runtime/AudioMixer/Private/AudioBusSubsystem.cpp
new file mode 100644
--- /dev/null
+++ b/Engine/Source/Runtime/AudioMixer/Private/AudioBusSubsystem.cpp
@@ -0,0 +1,287 @@
+// Copyright Epic Games, Inc. All Rights Reserved.
+
+#include "AudioBusSubsystem.h"
+#include "AudioMixerDevice.h"
+#include "AudioMixerSourceManager.h"
+#include "DSP/MultithreadedPatching.h"
+#include "UObject/UObjectIterator.h"
+
+std::atomic<uint32> Audio::FAudioBusKey::InstanceIdCounter = 0;
+
+UAudioBusSubsystem::UAudioBusSubsystem()
+{
+}
+
+void UAudioBusSubsystem::PostInitProperties()
+{
+	Super::PostInitProperties();
+	if (this != GetClass()->ClassDefaultObject)
+	{
+		InitDefaultAudioBuses();
+	}
+}
+
+bool UAudioBusSubsystem::ShouldCreateSubsystem(UObject* Outer) const
+{
+	return !IsRunningDedicatedServer();
+}
+
+void UAudioBusSubsystem::StartAudioBus(Audio::FAudioBusKey InAudioBusKey, int32 InNumChannels, bool bInIsAutomatic)
+{
+	if (IsInGameThread())
+	{
+		if (ActiveAudioBuses_GameThread.Contains(InAudioBusKey))
+		{
+			return;
+		}
+
+		FActiveBusData BusData;
+		BusData.BusKey = InAudioBusKey;
+		BusData.NumChannels = InNumChannels;
+		BusData.bIsAutomatic = bInIsAutomatic;
+
+		ActiveAudioBuses_GameThread.Add(InAudioBusKey, BusData);
+
+		FAudioThread::RunCommandOnAudioThread([this, InAudioBusKey, InNumChannels, bInIsAutomatic]()
+		{
+			GetMutableSourceManager()->StartAudioBus(InAudioBusKey, InNumChannels, bInIsAutomatic);
+		});
+	}
+	else
+	{
+		// If we're not the game thread, this needs to be on the game thread, so queue up a command to execute it on the game thread
+		GetMutableMixerDevice()->GameThreadMPSCCommand([this, InAudioBusKey, InNumChannels, bInIsAutomatic]
+		{
+			StartAudioBus(InAudioBusKey, InNumChannels, bInIsAutomatic);
+		});
+	}
+}
+
+void UAudioBusSubsystem::StopAudioBus(Audio::FAudioBusKey InAudioBusKey)
+{
+	if (IsInGameThread())
+	{
+		if (!ActiveAudioBuses_GameThread.Contains(InAudioBusKey))
+		{
+			return;
+		}
+
+		ActiveAudioBuses_GameThread.Remove(InAudioBusKey);
+
+		FAudioThread::RunCommandOnAudioThread([this, InAudioBusKey]()
+		{
+			GetMutableSourceManager()->StopAudioBus(InAudioBusKey);
+		});
+	}
+	else
+	{
+		// If we're not the game thread, this needs to be on the game thread, so queue up a command to execute it on the game thread
+		GetMutableMixerDevice()->GameThreadMPSCCommand([this, InAudioBusKey]
+		{
+			StopAudioBus(InAudioBusKey);
+		});
+	}
+}
+
+bool UAudioBusSubsystem::IsAudioBusActive(Audio::FAudioBusKey InAudioBusKey) const
+{
+	if (IsInGameThread())
+	{
+		return ActiveAudioBuses_GameThread.Contains(InAudioBusKey);
+	}
+
+	check(IsInAudioThread());
+	return GetSourceManager()->IsAudioBusActive(InAudioBusKey);
+}
+
+Audio::FPatchInput UAudioBusSubsystem::AddPatchInputForAudioBus(Audio::FAudioBusKey InAudioBusKey, int32 InFrames, int32 InChannels, float InGain)
+{
+	using namespace Audio;
+	// This function is supporting adding audio bus patches from multiple threads (AT, ART, GT, and tasks) and is currently
+	// depending on a number of places where data lives, which accounts for the complexity here.
+	// This code needs a clean up to refactor everything into a true MPSC model, along with an MPSC refactor of the source manager
+	// and our command queues. Once we do that we can remove the code which branches based on the thread the request is coming from. 
+
+	Audio::FMixerSourceManager* SourceManager = GetMutableSourceManager();
+	check(SourceManager);
+
+	if (IsInGameThread())
+	{
+		if (ActiveAudioBuses_GameThread.Find(InAudioBusKey))
+		{
+			FPatchInput PatchInput = GetMutableMixerDevice()->MakePatch(InFrames, InChannels, InGain);
+			FAudioThread::RunCommandOnAudioThread([this, PatchInput, InAudioBusKey]()
+			{
+				GetMutableSourceManager()->AddPatchInputForAudioBus_AudioThread(InAudioBusKey, PatchInput);
+			});
+			return PatchInput;
+		}
+		UE_LOG(LogAudioMixer, Warning, TEXT("Unable to add a patch output for audio bus because audio bus with object id '%u' and instance id '%u' is not active."), InAudioBusKey.ObjectId, InAudioBusKey.InstanceId);
+		return FPatchInput();
+	}
+	else if (IsInAudioThread())
+	{
+		FPatchInput PatchInput = GetMutableMixerDevice()->MakePatch(InFrames, InChannels, InGain);
+		GetMutableSourceManager()->AddPatchInputForAudioBus_AudioThread(InAudioBusKey, PatchInput);
+		return PatchInput;
+	}
+	else if (GetMixerDevice()->IsAudioRenderingThread())
+	{
+		check(SourceManager);
+
+		const int32 NumChannels = GetSourceManager()->GetAudioBusNumChannels(InAudioBusKey);
+		if (NumChannels > 0)
+		{
+			FPatchInput PatchInput = GetMutableMixerDevice()->MakePatch(InFrames, InChannels, InGain);
+			SourceManager->AddPatchInputForAudioBus(InAudioBusKey, PatchInput);
+			return PatchInput;
+		}
+
+		return FPatchInput();
+	}
+	else
+	{
+		FPatchInput PatchInput = GetMutableMixerDevice()->MakePatch(InFrames, InChannels, InGain);
+		GetMutableMixerDevice()->GameThreadMPSCCommand([this, PatchInput, InAudioBusKey]() mutable
+		{
+			if (ActiveAudioBuses_GameThread.Find(InAudioBusKey))
+			{
+				FAudioThread::RunCommandOnAudioThread([this, InAudioBusKey, PatchInput = MoveTemp(PatchInput)]()
+				{
+					GetMutableSourceManager()->AddPatchInputForAudioBus_AudioThread(InAudioBusKey, PatchInput);
+				});
+			}
+			else
+			{
+				UE_LOG(LogAudioMixer, Warning, TEXT("Unable to add a patch output for audio bus because audio bus with object id '%u' and instance id '%u' is not active."), InAudioBusKey.ObjectId, InAudioBusKey.InstanceId);
+			}
+		});
+		return PatchInput;
+	}
+}
+
+Audio::FPatchOutputStrongPtr UAudioBusSubsystem::AddPatchOutputForAudioBus(Audio::FAudioBusKey InAudioBusKey, int32 InFrames, int32 InChannels, float InGain)
+{
+	// This function is supporting adding audio bus patches from multiple threads (AT, ART, GT, and tasks) and is currently
+	// depending on a number of places where data lives, which accounts for the complexity here.
+	// This code needs a clean up to refactor everything into a true MPSC model, along with an MPSC refactor of the source manager
+	// and our command queues. Once we do that we can remove the code which branches based on the thread the request is coming from. 
+
+	Audio::FMixerSourceManager* SourceManager = GetMutableSourceManager();
+	check(SourceManager);
+
+	if (IsInGameThread())
+	{
+		if (ActiveAudioBuses_GameThread.Find(InAudioBusKey))
+		{
+			Audio::FPatchOutputStrongPtr PatchOutput = GetMutableMixerDevice()->MakePatch(InFrames, InChannels, InGain);
+			FAudioThread::RunCommandOnAudioThread([this, PatchOutput, InAudioBusKey]()
+			{
+				GetMutableSourceManager()->AddPatchOutputForAudioBus_AudioThread(InAudioBusKey, PatchOutput);
+			});
+			return PatchOutput;
+		}
+		UE_LOG(LogAudioMixer, Warning, TEXT("Unable to add a patch output for audio bus because audio bus with object id '%u' and instance id '%u' is not active."), InAudioBusKey.ObjectId, InAudioBusKey.InstanceId);
+		return nullptr;
+	}
+	else if (IsInAudioThread())
+	{
+		Audio::FPatchOutputStrongPtr PatchOutput = GetMutableMixerDevice()->MakePatch(InFrames, InChannels, InGain);
+		SourceManager->AddPatchOutputForAudioBus_AudioThread(InAudioBusKey, PatchOutput);
+		return PatchOutput;
+	}
+	else if (GetMixerDevice()->IsAudioRenderingThread())
+	{
+		check(SourceManager);
+
+		const int32 NumChannels = GetMutableSourceManager()->GetAudioBusNumChannels(InAudioBusKey);
+		if (NumChannels > 0)
+		{
+			Audio::FPatchOutputStrongPtr PatchOutput = GetMutableMixerDevice()->MakePatch(InFrames, InChannels, InGain);
+			SourceManager->AddPatchOutputForAudioBus(InAudioBusKey, PatchOutput);
+			return PatchOutput;
+		}
+
+		return nullptr;
+	}
+	else
+	{
+		Audio::FPatchOutputStrongPtr PatchOutput = GetMutableMixerDevice()->MakePatch(InFrames, InChannels, InGain);
+		GetMutableMixerDevice()->GameThreadMPSCCommand([this, PatchOutput, InAudioBusKey]() mutable
+		{
+			if (ActiveAudioBuses_GameThread.Find(InAudioBusKey))
+			{
+				FAudioThread::RunCommandOnAudioThread([this, InAudioBusKey, PatchOutput = MoveTemp(PatchOutput)]()
+				{
+					GetMutableSourceManager()->AddPatchOutputForAudioBus_AudioThread(InAudioBusKey, PatchOutput);
+				});
+			}
+			else
+			{
+				UE_LOG(LogAudioMixer, Warning, TEXT("Unable to add a patch output for audio bus because audio bus with object id '%u' and instance id '%u' is not active."), InAudioBusKey.ObjectId, InAudioBusKey.InstanceId);
+			}
+		});
+		return PatchOutput;
+	}
+}
+
+void UAudioBusSubsystem::InitDefaultAudioBuses()
+{
+	if (!ensure(IsInGameThread()))
+	{
+		return;
+	}
+
+	if (const UAudioSettings* AudioSettings = GetDefault<UAudioSettings>())
+	{
+		TArray<TStrongObjectPtr<UAudioBus>> StaleBuses = DefaultAudioBuses;
+		DefaultAudioBuses.Reset();
+
+		for (const FDefaultAudioBusSettings& BusSettings : AudioSettings->DefaultAudioBuses)
+		{
+			if (UObject* BusObject = BusSettings.AudioBus.TryLoad())
+			{
+				if (UAudioBus* AudioBus = Cast<UAudioBus>(BusObject))
+				{
+					const int32 NumChannels = static_cast<int32>(AudioBus->AudioBusChannels) + 1;
+					StartAudioBus(Audio::FAudioBusKey(AudioBus->GetUniqueID()), NumChannels, false /* bInIsAutomatic */);
+
+					TStrongObjectPtr<UAudioBus>AddedBus(AudioBus);
+					DefaultAudioBuses.AddUnique(AddedBus);
+					StaleBuses.Remove(AddedBus);
+				}
+			}
+		}
+
+		for (TStrongObjectPtr<UAudioBus>& Bus : StaleBuses)
+		{
+			if (Bus.IsValid())
+			{
+				StopAudioBus(Audio::FAudioBusKey(Bus->GetUniqueID()));
+			}
+		}
+	}
+	else
+	{
+		UE_LOG(LogAudioMixer, Error, TEXT("Failed to initialize Default Audio Buses. Audio Settings not found."));
+	}
+}
+
+void UAudioBusSubsystem::ShutdownDefaultAudioBuses()
+{
+	if (!ensure(IsInGameThread()))
+	{
+		return;
+	}
+
+	for (TObjectIterator<UAudioBus> It; It; ++It)
+	{
+		UAudioBus* AudioBus = *It;
+		if (AudioBus)
+		{
+			StopAudioBus(Audio::FAudioBusKey(AudioBus->GetUniqueID()));
+		}
+	}
+
+	DefaultAudioBuses.Reset();
+}
diff --git a/Engine/Source/Runtime/AudioMixer/Private/AudioMixerBlueprintLibrary.cpp b/Engine/Source/Runtime/AudioMixer/Private/AudioMixerBlueprintLibrary.cpp
--- a/Engine/Source/Runtime/AudioMixer/Private/AudioMixerBlueprintLibrary.cpp
+++ b/Engine/Source/Runtime/AudioMixer/Private/AudioMixerBlueprintLibrary.cpp
@@ -1,18 +1,20 @@
 // Copyright Epic Games, Inc. All Rights Reserved.
 
 #include "AudioMixerBlueprintLibrary.h"
-#include "Engine/World.h"
+
+#include "Algo/Transform.h"
+#include "Async/Async.h"
+#include "AudioBusSubsystem.h"
+#include "AudioCompressionSettingsUtils.h"
 #include "AudioDevice.h"
+#include "AudioDeviceManager.h"
 #include "AudioMixerDevice.h"
+#include "ContentStreaming.h"
 #include "CoreMinimal.h"
 #include "DSP/ConstantQ.h"
 #include "DSP/SpectrumAnalyzer.h"
-#include "ContentStreaming.h"
-#include "AudioCompressionSettingsUtils.h"
-#include "Async/Async.h"
+#include "Engine/World.h"
 #include "Sound/SoundEffectPreset.h"
-#include "Algo/Transform.h"
-#include "AudioDeviceManager.h"
 
 #include UE_INLINE_GENERATED_CPP_BY_NAME(AudioMixerBlueprintLibrary)
 
@@ -653,7 +655,10 @@ void UAudioMixerBlueprintLibrary::StartAudioBus(const UObject* WorldContextObjec
 	{
 		uint32 AudioBusId = AudioBus->GetUniqueID();
 		int32 NumChannels = (int32)AudioBus->AudioBusChannels + 1;
-		MixerDevice->StartAudioBus(AudioBusId, NumChannels, false);
+
+		UAudioBusSubsystem* AudioBusSubsystem = MixerDevice->GetSubsystem<UAudioBusSubsystem>();
+		check(AudioBusSubsystem);
+		AudioBusSubsystem->StartAudioBus(Audio::FAudioBusKey(AudioBusId), NumChannels, false);
 	}
 	else
 	{
@@ -671,7 +676,9 @@ void UAudioMixerBlueprintLibrary::StopAudioBus(const UObject* WorldContextObject
 	if (Audio::FMixerDevice* MixerDevice = FAudioDeviceManager::GetAudioMixerDeviceFromWorldContext(WorldContextObject))
 	{
 		uint32 AudioBusId = AudioBus->GetUniqueID();
-		MixerDevice->StopAudioBus(AudioBusId);
+		UAudioBusSubsystem* AudioBusSubsystem = MixerDevice->GetSubsystem<UAudioBusSubsystem>();
+		check(AudioBusSubsystem);
+		AudioBusSubsystem->StopAudioBus(Audio::FAudioBusKey(AudioBusId));
 	}
 	else
 	{
@@ -688,8 +695,10 @@ bool UAudioMixerBlueprintLibrary::IsAudioBusActive(const UObject* WorldContextOb
 	}
 	if (Audio::FMixerDevice* MixerDevice = FAudioDeviceManager::GetAudioMixerDeviceFromWorldContext(WorldContextObject))
 	{
-		uint32 AudioBusId = AudioBus->GetUniqueID();
-		return MixerDevice->IsAudioBusActive(AudioBusId);
+		uint32 AudioBusId = AudioBus->GetUniqueID(); 
+		UAudioBusSubsystem* AudioBusSubsystem = MixerDevice->GetSubsystem<UAudioBusSubsystem>();
+		check(AudioBusSubsystem);
+		return AudioBusSubsystem->IsAudioBusActive(Audio::FAudioBusKey(AudioBusId));
 	}
 	else
 	{
diff --git a/Engine/Source/Runtime/AudioMixer/Private/AudioMixerDevice.cpp b/Engine/Source/Runtime/AudioMixer/Private/AudioMixerDevice.cpp
--- a/Engine/Source/Runtime/AudioMixer/Private/AudioMixerDevice.cpp
+++ b/Engine/Source/Runtime/AudioMixer/Private/AudioMixerDevice.cpp
@@ -3,6 +3,7 @@
 #include "AudioMixerDevice.h"
 
 #include "AudioAnalytics.h"
+#include "AudioBusSubsystem.h"
 #include "AudioMixerSource.h"
 #include "AudioMixerSourceManager.h"
 #include "AudioMixerSourceDecode.h"
@@ -1849,63 +1850,10 @@ namespace Audio
 
 	void FMixerDevice::InitDefaultAudioBuses()
 	{
-		if (!ensure(IsInGameThread()))
-		{
-			return;
-		}
-
-		if (const UAudioSettings* AudioSettings = GetDefault<UAudioSettings>())
-		{
-			TArray<TStrongObjectPtr<UAudioBus>> StaleBuses = DefaultAudioBuses;
-			DefaultAudioBuses.Reset();
-
-			for (const FDefaultAudioBusSettings& BusSettings : AudioSettings->DefaultAudioBuses)
-			{
-				if (UObject* BusObject = BusSettings.AudioBus.TryLoad())
-				{
-					if (UAudioBus* AudioBus = Cast<UAudioBus>(BusObject))
-					{
-						const int32 NumChannels = static_cast<int32>(AudioBus->AudioBusChannels) + 1;
-						StartAudioBus(AudioBus->GetUniqueID(), NumChannels, false /* bInIsAutomatic */);
-
-						TStrongObjectPtr<UAudioBus>AddedBus(AudioBus);
-						DefaultAudioBuses.AddUnique(AddedBus);
-						StaleBuses.Remove(AddedBus);
-					}
-				}
-			}
-
-			for (TStrongObjectPtr<UAudioBus>& Bus : StaleBuses)
-			{
-				if (Bus.IsValid())
-				{
-					StopAudioBus(Bus->GetUniqueID());
-				}
-			}
-		}
-		else
-		{
-			UE_LOG(LogAudioMixer, Error, TEXT("Failed to initialized Default Audio Buses. Audio Settings not found."));
-		}
 	}
 
 	void FMixerDevice::ShutdownDefaultAudioBuses()
 	{
-		if (!ensure(IsInGameThread()))
-		{
-			return;
-		}
-
-		for (TObjectIterator<UAudioBus> It; It; ++It)
-		{
-			UAudioBus* AudioBus = *It;
-			if (AudioBus)
-			{
-				StopAudioBus(AudioBus->GetUniqueID());
-			}
-		}
-
-		DefaultAudioBuses.Reset();
 	}
 
 	FMixerSubmixWeakPtr FMixerDevice::GetSubmixInstance(const USoundSubmixBase* SoundSubmix)
@@ -2558,79 +2506,44 @@ namespace Audio
 
 	void FMixerDevice::StartAudioBus(uint32 InAudioBusId, int32 InNumChannels, bool bInIsAutomatic)
 	{
-		if (IsInGameThread())
+		UAudioBusSubsystem* AudioBusSubsystem = GetSubsystem<UAudioBusSubsystem>();
+		if (AudioBusSubsystem)
 		{
-			if (ActiveAudioBuses_GameThread.Contains(InAudioBusId))
-			{
-				return;
-			}
-
-			FActiveBusData BusData;
-			BusData.BusId = InAudioBusId;
-			BusData.NumChannels = InNumChannels;
-			BusData.bIsAutomatic = bInIsAutomatic;
-
-			ActiveAudioBuses_GameThread.Add(InAudioBusId, BusData);
-
-			FAudioThread::RunCommandOnAudioThread([this, InAudioBusId, InNumChannels, bInIsAutomatic]()
-			{
-				SourceManager->StartAudioBus(InAudioBusId, InNumChannels, bInIsAutomatic);
-			});
-		}
-		else
-		{
-			// If we're not the game thread, this needs to be on the game thread, so queue up a command to execute it on the game thread
-			GameThreadMPSCCommand([this, InAudioBusId, InNumChannels, bInIsAutomatic]
-			{
-				StartAudioBus(InAudioBusId, InNumChannels, bInIsAutomatic);	
-			});
+			AudioBusSubsystem->StartAudioBus(FAudioBusKey(InAudioBusId), InNumChannels, bInIsAutomatic);
 		}
 	}
 
 	void FMixerDevice::StopAudioBus(uint32 InAudioBusId)
 	{
-		if (IsInGameThread())
+		UAudioBusSubsystem* AudioBusSubsystem = GetSubsystem<UAudioBusSubsystem>();
+		if (AudioBusSubsystem)
 		{
-			if (!ActiveAudioBuses_GameThread.Contains(InAudioBusId))
-			{
-				return;
-			}
-
-			ActiveAudioBuses_GameThread.Remove(InAudioBusId);
-
-			FAudioThread::RunCommandOnAudioThread([this, InAudioBusId]()
-			{
-				SourceManager->StopAudioBus(InAudioBusId);
-			});
-		}
-		else
-		{
-			// If we're not the game thread, this needs to be on the game thread, so queue up a command to execute it on the game thread
-			GameThreadMPSCCommand([this, InAudioBusId]
-			{
-				StopAudioBus(InAudioBusId);	
-			});
+			AudioBusSubsystem->StopAudioBus(FAudioBusKey(InAudioBusId));
 		}
 	}
 
 	bool FMixerDevice::IsAudioBusActive(uint32 InAudioBusId) const
 	{
-		if (IsInGameThread())
+		UAudioBusSubsystem* AudioBusSubsystem = GetSubsystem<UAudioBusSubsystem>();
+		if (AudioBusSubsystem)
 		{
-			return ActiveAudioBuses_GameThread.Contains(InAudioBusId);
+			return AudioBusSubsystem->IsAudioBusActive(FAudioBusKey(InAudioBusId));
 		}
-
-		check(IsInAudioThread());
-		return SourceManager->IsAudioBusActive(InAudioBusId);
+		return false;
 	}
 
-	//UE_DEPRECATED(5.2, "AddPatchForAudioBus is deprecated.  Use AddPatchOutputForAudioBus.")
+	//UE_DEPRECATED(5.2, "AddPatchForAudioBus is deprecated.  Use UAudioBusSubsystem::AddPatchOutputForAudioBus.")
 	FPatchOutputStrongPtr FMixerDevice::AddPatchForAudioBus(uint32 InAudioBusId, float InPatchGain)
 	{
-		return AddPatchOutputForAudioBus(InAudioBusId, SourceManager->GetNumOutputFrames(), 8, InPatchGain);
+		UAudioBusSubsystem* AudioBusSubsystem = GetSubsystem<UAudioBusSubsystem>();
+		if (AudioBusSubsystem)
+		{
+			return AudioBusSubsystem->AddPatchOutputForAudioBus(FAudioBusKey(InAudioBusId), SourceManager->GetNumOutputFrames(), 8, InPatchGain);
+		}
+		return nullptr;
 	}
 
-	//UE_DEPRECATED(5.2, "AddPatchForAudioBus_GameThread is deprecated.  Use AddPatchOutputForAudioBus.")
+	//UE_DEPRECATED(5.2, "AddPatchForAudioBus_GameThread is deprecated.  Use UAudioBusSubsystem::AddPatchOutputForAudioBus.")
 	FPatchOutputStrongPtr FMixerDevice::AddPatchForAudioBus_GameThread(uint32 InAudioBusId, float InPatchGain)
 	{
 		PRAGMA_DISABLE_DEPRECATION_WARNINGS
@@ -2643,7 +2556,7 @@ namespace Audio
 	{
 	}
 
-	//UE_DEPRECATED(5.2, "AddPatchInputForAudioBus_GameThread is deprecated.  Use AddPatchInputForAudioBus.")
+	//UE_DEPRECATED(5.2, "AddPatchInputForAudioBus_GameThread is deprecated.  Use UAudioBusSubsystem::AddPatchInputForAudioBus.")
 	void FMixerDevice::AddPatchInputForAudioBus_GameThread(const FPatchInput& InPatchInput, uint32 InAudioBusId, float InPatchGain)
 	{
 		PRAGMA_DISABLE_DEPRECATION_WARNINGS
@@ -2657,130 +2570,6 @@ namespace Audio
 		return MakeShared<Audio::FPatchOutput, ESPMode::ThreadSafe>(AudioMixerPatchBufferBlocks * MinSizeFrames * FMath::DivideAndRoundUp(MaxSizeFrames, MinSizeFrames) * InChannels, InGain);
 	}
 
-	FPatchInput FMixerDevice::AddPatchInputForAudioBus(uint32 InAudioBusId, int32 InFrames, int32 InChannels, float InGain)
-	{
-		// This function is supporting adding audio bus patches from multiple threads (AT, ART, GT, and tasks) and is currently
-		// depending on a number of places where data lives, which accounts for the complexity here.
-		// This code needs a clean up to refactor everything into a true MPSC model, along with an MPSC refactor of the source manager
-		// and our command queues. Once we do that we can remove the code which branches based on the thread the request is coming from. 
-
-		if (IsInGameThread())
-		{
-			if (ActiveAudioBuses_GameThread.Find(InAudioBusId))
-			{
-				FPatchInput PatchInput = MakePatch(InFrames, InChannels, InGain);
-				FAudioThread::RunCommandOnAudioThread([this, PatchInput, InAudioBusId]()
-				{
-					SourceManager->AddPatchInputForAudioBus_AudioThread(InAudioBusId, PatchInput);
-				});
-				return PatchInput;
-			}
-			UE_LOG(LogAudioMixer, Warning, TEXT("Unable to add a patch input for audio bus because audio bus id '%d' is not active."), InAudioBusId);
-			return FPatchInput();
-		}
-		else if (IsInAudioThread())
-		{
-			FPatchInput PatchInput = MakePatch(InFrames, InChannels, InGain);
-			SourceManager->AddPatchInputForAudioBus_AudioThread(InAudioBusId, PatchInput);
-			return PatchInput;
-		}
-		else if (IsAudioRenderingThread())
-		{
-			check(SourceManager);
-
-			const int32 NumChannels = SourceManager->GetAudioBusNumChannels(InAudioBusId);
-			if (NumChannels > 0)
-			{
-				FPatchInput PatchInput = MakePatch(InFrames, InChannels, InGain);
-				SourceManager->AddPatchInputForAudioBus(InAudioBusId, PatchInput);
-				return PatchInput;
-			}
-
-			return FPatchInput();
-		}
-		else
-		{
-			FPatchInput PatchInput = MakePatch(InFrames, InChannels, InGain);
-			GameThreadMPSCCommand([this, PatchInput, InAudioBusId]() mutable
-			{
-				if (ActiveAudioBuses_GameThread.Find(InAudioBusId))
-				{
-					FAudioThread::RunCommandOnAudioThread([this, InAudioBusId, PatchInput = MoveTemp(PatchInput)]()
-					{
-						SourceManager->AddPatchInputForAudioBus_AudioThread(InAudioBusId, PatchInput);
-					});
-				}
-				else
-				{
-					UE_LOG(LogAudioMixer, Warning, TEXT("Unable to add a patch output for audio bus because audio bus id '%d' is not active."), InAudioBusId);
-				}
-			});
-			return PatchInput;
-		}
-	}
-
-	FPatchOutputStrongPtr FMixerDevice::AddPatchOutputForAudioBus(uint32 InAudioBusId, int32 InFrames, int32 InChannels, float InGain)
-	{
-		// This function is supporting adding audio bus patches from multiple threads (AT, ART, GT, and tasks) and is currently
-		// depending on a number of places where data lives, which accounts for the complexity here.
-		// This code needs a clean up to refactor everything into a true MPSC model, along with an MPSC refactor of the source manager
-		// and our command queues. Once we do that we can remove the code which branches based on the thread the request is coming from. 
-
-		if (IsInGameThread())
-		{
-			if (ActiveAudioBuses_GameThread.Find(InAudioBusId))
-			{
-				FPatchOutputStrongPtr PatchOutput = MakePatch(InFrames, InChannels, InGain);
-				FAudioThread::RunCommandOnAudioThread([this, PatchOutput, InAudioBusId]()
-				{
-					SourceManager->AddPatchOutputForAudioBus_AudioThread(InAudioBusId, PatchOutput);
-				});
-				return PatchOutput;
-			}
-			UE_LOG(LogAudioMixer, Warning, TEXT("Unable to add a patch input for audio bus because audio bus id '%d' is not active."), InAudioBusId);
-			return nullptr;
-		}
-		else if (IsInAudioThread())
-		{
-			FPatchOutputStrongPtr PatchOutput = MakePatch(InFrames, InChannels, InGain);
-			SourceManager->AddPatchOutputForAudioBus_AudioThread(InAudioBusId, PatchOutput);
-			return PatchOutput;
-		}
-		else if (IsAudioRenderingThread())
-		{
-			check(SourceManager);
-
-			const int32 NumChannels = SourceManager->GetAudioBusNumChannels(InAudioBusId);
-			if (NumChannels > 0)
-			{
-				FPatchOutputStrongPtr PatchOutput = MakePatch(InFrames, InChannels, InGain);
-				SourceManager->AddPatchOutputForAudioBus(InAudioBusId, PatchOutput);
-				return PatchOutput;
-			}
-
-			return nullptr;
-		}
-		else
-		{
-			FPatchOutputStrongPtr PatchOutput = MakePatch(InFrames, InChannels, InGain);
-			GameThreadMPSCCommand([this, PatchOutput, InAudioBusId]() mutable
-			{
-				if (ActiveAudioBuses_GameThread.Find(InAudioBusId))
-				{
-					FAudioThread::RunCommandOnAudioThread([this, InAudioBusId, PatchOutput = MoveTemp(PatchOutput)]()
-					{
-						SourceManager->AddPatchOutputForAudioBus_AudioThread(InAudioBusId, PatchOutput);
-					});
-				}
-				else
-				{
-					UE_LOG(LogAudioMixer, Warning, TEXT("Unable to add a patch output for audio bus because audio bus id '%d' is not active."), InAudioBusId);
-				}
-			});
-			return PatchOutput;
-		}
-	}
-
 	FPatchOutputStrongPtr FMixerDevice::AddPatchForSubmix(uint32 InObjectId, float InPatchGain)
 	{
 		if (!ensure(IsAudioRenderingThread()))
@@ -2812,6 +2601,11 @@ namespace Audio
 		return SourceManager.Get();
 	}
 
+	const FMixerSourceManager* FMixerDevice::GetSourceManager() const
+	{
+		return SourceManager.Get();
+	}
+
 	bool FMixerDevice::IsMainAudioDevice() const
 	{
 		bool bIsMain = (this == FAudioDeviceManager::Get()->GetMainAudioDeviceRaw());
diff --git a/Engine/Source/Runtime/AudioMixer/Private/AudioMixerSourceManager.cpp b/Engine/Source/Runtime/AudioMixer/Private/AudioMixerSourceManager.cpp
--- a/Engine/Source/Runtime/AudioMixer/Private/AudioMixerSourceManager.cpp
+++ b/Engine/Source/Runtime/AudioMixer/Private/AudioMixerSourceManager.cpp
@@ -1135,19 +1135,19 @@ namespace Audio
 		}, AUDIO_MIXER_THREAD_COMMAND_STRING("ReleaseSourceId"));
 	}
 
-	void FMixerSourceManager::StartAudioBus(uint32 InAudioBusId, int32 InNumChannels, bool bInIsAutomatic)
+	void FMixerSourceManager::StartAudioBus(FAudioBusKey InAudioBusKey, int32 InNumChannels, bool bInIsAutomatic)
 	{
-		if (AudioBusIds_AudioThread.Contains(InAudioBusId))
+		if (AudioBusKeys_AudioThread.Contains(InAudioBusKey))
 		{
 			return;
 		}
 
-		AudioBusIds_AudioThread.Add(InAudioBusId);
+		AudioBusKeys_AudioThread.Add(InAudioBusKey);
 
-		AudioMixerThreadCommand([this, InAudioBusId, InNumChannels, bInIsAutomatic]()
+		AudioMixerThreadCommand([this, InAudioBusKey, InNumChannels, bInIsAutomatic]()
 		{
 			// If this audio bus id already exists, set it to not be automatic and return it
-			TSharedPtr<FMixerAudioBus> AudioBusPtr = AudioBuses.FindRef(InAudioBusId);
+			TSharedPtr<FMixerAudioBus> AudioBusPtr = AudioBuses.FindRef(InAudioBusKey);
 			if (AudioBusPtr.IsValid())
 			{
 				// If this audio bus already existed, make sure the num channels lines up
@@ -1159,23 +1159,23 @@ namespace Audio
 				// If the bus is not registered, make a new entry.
 				TSharedPtr<FMixerAudioBus> NewBusData(new FMixerAudioBus(this, bInIsAutomatic, InNumChannels));
 
-				AudioBuses.Add(InAudioBusId, NewBusData);
+				AudioBuses.Add(InAudioBusKey, NewBusData);
 			}
 		}, AUDIO_MIXER_THREAD_COMMAND_STRING("StartAudioBus"));
 	}
 
-	void FMixerSourceManager::StopAudioBus(uint32 InAudioBusId)
+	void FMixerSourceManager::StopAudioBus(FAudioBusKey InAudioBusKey)
 	{
-		if (!AudioBusIds_AudioThread.Contains(InAudioBusId))
+		if (!AudioBusKeys_AudioThread.Contains(InAudioBusKey))
 		{
 			return;
 		}
 
-		AudioBusIds_AudioThread.Remove(InAudioBusId);
+		AudioBusKeys_AudioThread.Remove(InAudioBusKey);
 
-		AudioMixerThreadCommand([this, InAudioBusId]()
+		AudioMixerThreadCommand([this, InAudioBusKey]()
 		{
-			TSharedPtr<FMixerAudioBus>* AudioBusPtr = AudioBuses.Find(InAudioBusId);
+			TSharedPtr<FMixerAudioBus>* AudioBusPtr = AudioBuses.Find(InAudioBusKey);
 			if (AudioBusPtr)
 			{
 				if (!(*AudioBusPtr)->IsAutomatic())
@@ -1183,7 +1183,7 @@ namespace Audio
 					// Immediately stop all sources which were source buses
 					for (FSourceInfo& SourceInfo : SourceInfos)
 					{
-						if (SourceInfo.AudioBusId == InAudioBusId)
+						if (SourceInfo.AudioBusId == InAudioBusKey.ObjectId)
 						{
 							SourceInfo.bIsPlaying = false;
 							SourceInfo.bIsPaused = false;
@@ -1191,22 +1191,22 @@ namespace Audio
 							SourceInfo.bIsStopping = false;
 						}
 					}
-					AudioBuses.Remove(InAudioBusId);
+					AudioBuses.Remove(InAudioBusKey);
 				}
 			}
 		}, AUDIO_MIXER_THREAD_COMMAND_STRING("StartAudioBus"));
 	}
 
-	bool FMixerSourceManager::IsAudioBusActive(uint32 InAudioBusId) const
+	bool FMixerSourceManager::IsAudioBusActive(FAudioBusKey InAudioBusKey) const
 	{
 		AUDIO_MIXER_CHECK_GAME_THREAD(MixerDevice);
-		return AudioBusIds_AudioThread.Contains(InAudioBusId);
+		return AudioBusKeys_AudioThread.Contains(InAudioBusKey);
 	}
 
-	int32 FMixerSourceManager::GetAudioBusNumChannels(uint32 InAudioBusId) const
+	int32 FMixerSourceManager::GetAudioBusNumChannels(FAudioBusKey InAudioBusKey) const
 	{
 		AUDIO_MIXER_CHECK_AUDIO_PLAT_THREAD(MixerDevice);
-		TSharedPtr<FMixerAudioBus> AudioBusPtr = AudioBuses.FindRef(InAudioBusId);
+		TSharedPtr<FMixerAudioBus> AudioBusPtr = AudioBuses.FindRef(InAudioBusKey);
 		if (AudioBusPtr.IsValid())
 		{
 			return AudioBusPtr->GetNumChannels();
@@ -1215,12 +1215,12 @@ namespace Audio
 		return 0;
 	}
 
-	void FMixerSourceManager::AddPatchOutputForAudioBus(uint32 InAudioBusId, const FPatchOutputStrongPtr& InPatchOutputStrongPtr)
+	void FMixerSourceManager::AddPatchOutputForAudioBus(FAudioBusKey InAudioBusKey, const FPatchOutputStrongPtr& InPatchOutputStrongPtr)
 	{
 		AUDIO_MIXER_CHECK_AUDIO_PLAT_THREAD(MixerDevice);
 		if (MixerDevice->IsAudioRenderingThread())
 		{
-			TSharedPtr<FMixerAudioBus> AudioBusPtr = AudioBuses.FindRef(InAudioBusId);
+			TSharedPtr<FMixerAudioBus> AudioBusPtr = AudioBuses.FindRef(InAudioBusKey);
 			if (AudioBusPtr.IsValid())
 			{
 				AudioBusPtr->AddNewPatchOutput(InPatchOutputStrongPtr);
@@ -1229,27 +1229,27 @@ namespace Audio
 		else
 		{
 			// Queue up the command via MPSC command queue
-			AudioMixerThreadMPSCCommand([this, InAudioBusId, InPatchOutputStrongPtr]()
+			AudioMixerThreadMPSCCommand([this, InAudioBusKey, InPatchOutputStrongPtr]()
 			{
-				AddPatchOutputForAudioBus(InAudioBusId, InPatchOutputStrongPtr);
+				AddPatchOutputForAudioBus(InAudioBusKey, InPatchOutputStrongPtr);
 			});
 		}
 	}
 
-	void FMixerSourceManager::AddPatchOutputForAudioBus_AudioThread(uint32 InAudioBusId, const FPatchOutputStrongPtr& InPatchOutputStrongPtr)
+	void FMixerSourceManager::AddPatchOutputForAudioBus_AudioThread(FAudioBusKey InAudioBusKey, const FPatchOutputStrongPtr& InPatchOutputStrongPtr)
 	{
-		AudioMixerThreadCommand([this, InAudioBusId, NewPatchPtr = InPatchOutputStrongPtr]() mutable
+		AudioMixerThreadCommand([this, InAudioBusKey, NewPatchPtr = InPatchOutputStrongPtr]() mutable
 		{
-			AddPatchOutputForAudioBus(InAudioBusId, NewPatchPtr);
+			AddPatchOutputForAudioBus(InAudioBusKey, NewPatchPtr);
 		}, AUDIO_MIXER_THREAD_COMMAND_STRING("AddPatchOutputForAudioBus_AudioThread()"));
 	}
 
-	void FMixerSourceManager::AddPatchInputForAudioBus(uint32 InAudioBusId, const FPatchInput& InPatchInput)
+	void FMixerSourceManager::AddPatchInputForAudioBus(FAudioBusKey InAudioBusKey, const FPatchInput& InPatchInput)
 	{
 		AUDIO_MIXER_CHECK_AUDIO_PLAT_THREAD(MixerDevice);
 		if (MixerDevice->IsAudioRenderingThread())
 		{
-			TSharedPtr<FMixerAudioBus> AudioBusPtr = AudioBuses.FindRef(InAudioBusId);
+			TSharedPtr<FMixerAudioBus> AudioBusPtr = AudioBuses.FindRef(InAudioBusKey);
 			if (AudioBusPtr.IsValid())
 			{
 				AudioBusPtr->AddNewPatchInput(InPatchInput);
@@ -1258,18 +1258,18 @@ namespace Audio
 		else
 		{
 			// Queue up the command via MPSC command queue
-			AudioMixerThreadMPSCCommand([this, InAudioBusId, InPatchInput]()
+			AudioMixerThreadMPSCCommand([this, InAudioBusKey, InPatchInput]()
 			{
-				AddPatchInputForAudioBus(InAudioBusId, InPatchInput);
+				AddPatchInputForAudioBus(InAudioBusKey, InPatchInput);
 			});
 		}
 	}
 
-	void FMixerSourceManager::AddPatchInputForAudioBus_AudioThread(uint32 InAudioBusId, const FPatchInput& InPatchInput)
+	void FMixerSourceManager::AddPatchInputForAudioBus_AudioThread(FAudioBusKey InAudioBusKey, const FPatchInput& InPatchInput)
 	{
-		AudioMixerThreadCommand([this, InAudioBusId, InPatchInput]()
+		AudioMixerThreadCommand([this, InAudioBusKey, InPatchInput]()
 		{
-			AddPatchInputForAudioBus(InAudioBusId, InPatchInput);
+			AddPatchInputForAudioBus(InAudioBusKey, InPatchInput);
 		}, AUDIO_MIXER_THREAD_COMMAND_STRING("AddPatchInputForAudioBus_AudioThread()"));
 	}
 
diff --git a/Engine/Source/Runtime/AudioMixer/Private/AudioMixerSourceManager.h b/Engine/Source/Runtime/AudioMixer/Private/AudioMixerSourceManager.h
--- a/Engine/Source/Runtime/AudioMixer/Private/AudioMixerSourceManager.h
+++ b/Engine/Source/Runtime/AudioMixer/Private/AudioMixerSourceManager.h
@@ -1,8 +1,7 @@
 // Copyright Epic Games, Inc. All Rights Reserved.
 #pragma once
 
-#include "CoreMinimal.h"
-
+#include "AudioBusSubsystem.h"
 #include "AudioMixerBuffer.h"
 #include "AudioMixerBus.h"
 #include "AudioMixerDevice.h"
@@ -190,29 +189,29 @@ namespace Audio
 		void InitSource(const int32 SourceId, const FMixerSourceVoiceInitParams& InitParams);
 
 		// Creates and starts an audio bus manually.
-		void StartAudioBus(uint32 InAudioBusId, int32 InNumChannels, bool bInIsAutomatic);
+		void StartAudioBus(FAudioBusKey InAudioBusKey, int32 InNumChannels, bool bInIsAutomatic);
 
 		// Stops an audio bus manually
-		void StopAudioBus(uint32 InAudioBusId);
+		void StopAudioBus(FAudioBusKey InAudioBusKey);
 
 		// Queries if an audio bus is active. Must be called from the audio thread.
-		bool IsAudioBusActive(uint32 InAudioBusId) const;
+		bool IsAudioBusActive(FAudioBusKey InAudioBusKey) const;
 
 		// Returns the number of channels currently set for the audio bus associated with
 		// the provided BusId.  Returns 0 if the audio bus is inactive.
-		int32 GetAudioBusNumChannels(uint32 InAudioBusId) const;
+		int32 GetAudioBusNumChannels(FAudioBusKey InAudioBusKey) const;
 
 		// Adds a patch output for an audio bus from the Audio Render Thread
-		void AddPatchOutputForAudioBus(uint32 InAudioBusId, const FPatchOutputStrongPtr& InPatchOutputStrongPtr);
+		void AddPatchOutputForAudioBus(FAudioBusKey InAudioBusKey, const FPatchOutputStrongPtr& InPatchOutputStrongPtr);
 
 		// Adds a patch output for an audio bus from the Audio Thread
-		void AddPatchOutputForAudioBus_AudioThread(uint32 InAudioBusId, const FPatchOutputStrongPtr& InPatchOutputStrongPtr);
+		void AddPatchOutputForAudioBus_AudioThread(FAudioBusKey InAudioBusKey, const FPatchOutputStrongPtr& InPatchOutputStrongPtr);
 
 		// Adds a patch input for an audio bus
-		void AddPatchInputForAudioBus(uint32 InAudioBusId, const FPatchInput& InPatchInput);
+		void AddPatchInputForAudioBus(FAudioBusKey InAudioBusKey, const FPatchInput& InPatchInput);
 
 		// Adds a patch input for an audio bus from the Audio Thread
-		void AddPatchInputForAudioBus_AudioThread(uint32 InAudioBusId, const FPatchInput& InPatchInput);
+		void AddPatchInputForAudioBus_AudioThread(FAudioBusKey InAudioBusKey, const FPatchInput& InPatchInput);
 
 		void Play(const int32 SourceId);
 		void Stop(const int32 SourceId);
@@ -578,8 +577,8 @@ namespace Audio
 		TArray<FMixerSourceSubmixOutputBuffer> SourceSubmixOutputBuffers;
 
 		// Map of bus object Id's to audio bus data. 
-		TMap<uint32, TSharedPtr<FMixerAudioBus>> AudioBuses;
-		TArray<uint32> AudioBusIds_AudioThread;
+		TMap<FAudioBusKey, TSharedPtr<FMixerAudioBus>> AudioBuses; 
+		TArray<FAudioBusKey> AudioBusKeys_AudioThread;
 
 		// Async task workers for processing sources in parallel
 		TArray<FAsyncTask<FAudioMixerSourceWorker>*> SourceWorkers;
diff --git a/Engine/Source/Runtime/AudioMixer/Private/Effects/AudioMixerSubmixEffectDynamicsProcessor.cpp b/Engine/Source/Runtime/AudioMixer/Private/Effects/AudioMixerSubmixEffectDynamicsProcessor.cpp
--- a/Engine/Source/Runtime/AudioMixer/Private/Effects/AudioMixerSubmixEffectDynamicsProcessor.cpp
+++ b/Engine/Source/Runtime/AudioMixer/Private/Effects/AudioMixerSubmixEffectDynamicsProcessor.cpp
@@ -2,6 +2,7 @@
 
 #include "SubmixEffects/AudioMixerSubmixEffectDynamicsProcessor.h"
 
+#include "AudioBusSubsystem.h"
 #include "AudioDeviceManager.h"
 #include "AudioMixerDevice.h"
 #include "AudioMixerSubmix.h"
@@ -198,7 +199,10 @@ bool FSubmixEffectDynamicsProcessor::UpdateKeySourcePatch()
 				const uint32 ObjectId = KeySource.GetObjectId();
 				if (ObjectId != INDEX_NONE)
 				{
-					KeySource.Patch = MixerDevice->AddPatchOutputForAudioBus(ObjectId, MixerDevice->GetNumOutputFrames(), KeySource.GetNumChannels());
+					UAudioBusSubsystem* AudioBusSubsystem = MixerDevice->GetSubsystem<UAudioBusSubsystem>();
+					check(AudioBusSubsystem);
+					KeySource.Patch = AudioBusSubsystem->AddPatchOutputForAudioBus(Audio::FAudioBusKey(ObjectId), MixerDevice->GetNumOutputFrames(), KeySource.GetNumChannels());
+
 					if (KeySource.Patch.IsValid())
 					{
 						DynamicsProcessor.SetKeyNumChannels(KeySource.GetNumChannels());
diff --git a/Engine/Source/Runtime/AudioMixer/Public/AudioBusSubsystem.h b/Engine/Source/Runtime/AudioMixer/Public/AudioBusSubsystem.h
new file mode 100644
--- /dev/null
+++ b/Engine/Source/Runtime/AudioMixer/Public/AudioBusSubsystem.h
@@ -0,0 +1,103 @@
+// Copyright Epic Games, Inc. All Rights Reserved.
+
+#pragma once
+
+#include "Subsystems/AudioEngineSubsystem.h"
+#include "Templates/TypeHash.h"
+#include "UObject/StrongObjectPtr.h"
+
+#include "AudioBusSubsystem.generated.h"
+
+class UAudioBus;
+
+namespace Audio
+{
+	// Forward declarations 
+	class FMixerAudioBus;
+	class FMixerSourceManager;
+	class FPatchInput;
+	struct FPatchOutput;
+	typedef TSharedPtr<FPatchOutput, ESPMode::ThreadSafe> FPatchOutputStrongPtr;
+
+	struct AUDIOMIXER_API FAudioBusKey
+	{
+		uint32 ObjectId = INDEX_NONE; // from a corresponding UObject (UAudioBus) if applicable
+		uint32 InstanceId = INDEX_NONE;
+
+		FAudioBusKey()
+			: InstanceId(InstanceIdCounter++)
+		{
+		}
+
+		// For construction with a given UObject unique id 
+		FAudioBusKey(uint32 InObjectId)
+			: ObjectId(InObjectId)
+		{
+		}
+
+		const bool IsValid() const
+		{
+			return ObjectId != INDEX_NONE || InstanceId != INDEX_NONE;
+		}
+
+		FORCEINLINE friend uint32 GetTypeHash(const FAudioBusKey& Key)
+		{
+			return HashCombineFast(Key.ObjectId, Key.InstanceId);
+		}
+		 		
+		FORCEINLINE friend bool operator==(const FAudioBusKey& InLHS, const FAudioBusKey& InRHS) 
+		{
+			return (InLHS.ObjectId == InRHS.ObjectId) && (InLHS.InstanceId == InRHS.InstanceId);
+		}
+
+		FORCEINLINE friend bool operator!=(const FAudioBusKey& InLHS, const FAudioBusKey& InRHS) 
+		{
+			return !(InLHS == InRHS);
+		}
+		
+
+	private:
+		static std::atomic<uint32> InstanceIdCounter;
+	};
+}
+
+/**
+*  UAudioBusSubsystem
+*/
+UCLASS()
+class AUDIOMIXER_API UAudioBusSubsystem : public UAudioEngineSubsystem
+{
+	GENERATED_BODY()
+
+public:
+	UAudioBusSubsystem();
+	virtual ~UAudioBusSubsystem() = default;
+	virtual void PostInitProperties() override;
+
+	//~ Begin USubsystem interface
+	virtual bool ShouldCreateSubsystem(UObject* Outer) const override;
+	//~ End USubsystem interface
+
+	// Audio bus API from FMixerDevice
+	void StartAudioBus(Audio::FAudioBusKey InAudioBusKey, int32 InNumChannels, bool bInIsAutomatic);
+	void StopAudioBus(Audio::FAudioBusKey InAudioBusKey);
+	bool IsAudioBusActive(Audio::FAudioBusKey InAudioBusKey) const;
+	
+	Audio::FPatchInput AddPatchInputForAudioBus(Audio::FAudioBusKey InAudioBusKey, int32 InFrames, int32 InChannels, float InGain = 1.f);
+	Audio::FPatchOutputStrongPtr AddPatchOutputForAudioBus(Audio::FAudioBusKey InAudioBusKey, int32 InFrames, int32 InChannels, float InGain = 1.f);
+
+	void InitDefaultAudioBuses();
+	void ShutdownDefaultAudioBuses();
+
+private:
+	struct FActiveBusData
+	{
+		Audio::FAudioBusKey BusKey = 0;
+		int32 NumChannels = 0;
+		bool bIsAutomatic = false;
+	};
+
+	TArray<TStrongObjectPtr<UAudioBus>> DefaultAudioBuses; 
+	// The active audio bus list accessible on the game thread
+	TMap<Audio::FAudioBusKey, FActiveBusData> ActiveAudioBuses_GameThread;
+};
diff --git a/Engine/Source/Runtime/AudioMixer/Public/AudioMixerDevice.h b/Engine/Source/Runtime/AudioMixer/Public/AudioMixerDevice.h
--- a/Engine/Source/Runtime/AudioMixer/Public/AudioMixerDevice.h
+++ b/Engine/Source/Runtime/AudioMixer/Public/AudioMixerDevice.h
@@ -125,7 +125,9 @@ namespace Audio
 		virtual void UnregisterSoundSubmix(const USoundSubmixBase* SoundSubmix) override;
 
 		virtual void InitSoundEffectPresets() override;
+		UE_DEPRECATED(5.2, "The functionality for this has been moved to UAudioBusSubsystem::InitDefaultAudioBuses, which is now automatically called on subsystem creation.")
 		virtual void InitDefaultAudioBuses() override;
+		UE_DEPRECATED(5.2, "The functionality for this has been moved to UAudioBusSubsystem::ShutdownDefaultAudioBuses.")
 		virtual void ShutdownDefaultAudioBuses() override;
 		virtual int32 GetNumActiveSources() const override;
 
@@ -244,6 +246,7 @@ namespace Audio
 		int32 GetDeviceOutputChannels() const;
 
 		FMixerSourceManager* GetSourceManager();
+		const FMixerSourceManager* GetSourceManager() const;
 
 		FMixerSubmixWeakPtr GetMasterSubmix(); 
 		FMixerSubmixWeakPtr GetBaseDefaultSubmix();
@@ -294,25 +297,27 @@ namespace Audio
 
 		static bool IsEndpointSubmix(const USoundSubmixBase* InSubmix);
 
-		// Audio bus API
+		// Audio bus API - these are deprecated. Use corresponding calls in UAudioBusSubsystem instead
+		UE_DEPRECATED(5.2, "This function is deprecated. Use UAudioBusSubsystem::StartAudioBus instead.")
 		virtual void StartAudioBus(uint32 InAudioBusId, int32 InNumChannels, bool bInIsAutomatic) override;
+		UE_DEPRECATED(5.2, "This function is deprecated. Use UAudioBusSubsystem::StopAudioBus instead.")
 		virtual void StopAudioBus(uint32 InAudioBusId) override;
+		UE_DEPRECATED(5.2, "This function is deprecated. Use UAudioBusSubsystem::IsAudioBusActive instead.")
 		virtual bool IsAudioBusActive(uint32 InAudioBusId) const override;
 
-		UE_DEPRECATED(5.2, "AddPatchForAudioBus is deprecated.  Use AddPatchOutputForAudioBus.")
+		UE_DEPRECATED(5.2, "AddPatchForAudioBus is deprecated.  Use UAudioBusSubsystem::AddPatchOutputForAudioBus.")
 		virtual FPatchOutputStrongPtr AddPatchForAudioBus(uint32 InAudioBusId, float InPatchGain = 1.0f) override;
 
-		UE_DEPRECATED(5.2, "AddPatchForAudioBus_GameThread is deprecated.  Use AddPatchOutputForAudioBus.")
+		UE_DEPRECATED(5.2, "AddPatchForAudioBus_GameThread is deprecated.  Use UAudioBusSubsystem::AddPatchOutputForAudioBus.")
 		virtual FPatchOutputStrongPtr AddPatchForAudioBus_GameThread(uint32 InAudioBusId, float InPatchGain = 1.0f) override;
 
-		UE_DEPRECATED(5.2, "This overload of AddPatchInputForAudioBus is deprecated and non-functional.  Use the overload that takes the number of frames and channels as parameters.")
+		UE_DEPRECATED(5.2, "This overload of AddPatchInputForAudioBus is deprecated and non-functional.  Use the overload in UAudioBusSubsystem that takes the number of frames and channels as parameters.")
 		virtual void AddPatchInputForAudioBus(const FPatchInput& InPatchInput, uint32 InAudioBusId, float InPatchGain = 1.0f) override;
 
-		UE_DEPRECATED(5.2, "AddPatchInputForAudioBus_GameThread is deprecated.  Use AddPatchInputForAudioBus.")
+		UE_DEPRECATED(5.2, "AddPatchInputForAudioBus_GameThread is deprecated.  Use UAudioBusSubsystem::AddPatchInputForAudioBus instead.")
 		virtual void AddPatchInputForAudioBus_GameThread(const FPatchInput& InPatchInput, uint32 InAudioBusId, float InPatchGain = 1.0f) override;
 
-		virtual FPatchInput AddPatchInputForAudioBus(uint32 InAudioBusId, int32 InFrames, int32 InChannels, float InGain = 1.f) override;
-		virtual FPatchOutputStrongPtr AddPatchOutputForAudioBus(uint32 InAudioBusId, int32 InFrames, int32 InChannels, float InGain = 1.f) override;
+		FPatchOutputStrongPtr MakePatch(int32 InFrames, int32 InChannels, float InGain) const;
 
 		// Clock Manager for quantized event handling on Audio Render Thread
 		FQuartzClockManager QuantizedEventClockManager;
@@ -373,23 +378,10 @@ namespace Audio
 		void PumpCommandQueue();
 		void PumpGameThreadCommandQueue();
 		
-		FPatchOutputStrongPtr MakePatch(int32 InFrames, int32 InChannels, float InGain) const;
-
 		TArray<USoundSubmix*> MasterSubmixes;
 		TArray<FMixerSubmixPtr> MasterSubmixInstances;
 
 		TArray<TStrongObjectPtr<UAudioBus>> DefaultAudioBuses;
-
-		struct FActiveBusData
-		{
-			int32 BusId = 0;
-			int32 NumChannels = 0;
-			bool bIsAutomatic = false;
-		};
-
-		// The active audio bus list accessible on the game thread
-		TMap<int32, FActiveBusData> ActiveAudioBuses_GameThread;
-
 		/** Ptr to the platform interface, which handles streaming audio to the hardware device. */
 		IAudioMixerPlatformInterface* AudioMixerPlatform;
 		
diff --git a/Engine/Source/Runtime/Engine/Private/AudioBus.cpp b/Engine/Source/Runtime/Engine/Private/AudioBus.cpp
--- a/Engine/Source/Runtime/Engine/Private/AudioBus.cpp
+++ b/Engine/Source/Runtime/Engine/Private/AudioBus.cpp
@@ -2,10 +2,10 @@
 
 
 #include "Sound/AudioBus.h"
+#include "AudioBusSubsystem.h"
 #include "AudioDeviceManager.h"
 #include "Engine/Engine.h"
 #include "AudioDevice.h"
-#include "AudioMixerDevice.h"
 
 #include UE_INLINE_GENERATED_CPP_BY_NAME(AudioBus)
 
@@ -44,8 +44,9 @@ void UAudioBus::BeginDestroy()
 		{
 			if (AudioDevice->IsAudioMixerEnabled())
 			{
-				Audio::FMixerDevice* MixerDevice = static_cast<Audio::FMixerDevice*>(AudioDevice);
-				MixerDevice->StopAudioBus(AudioBusId);
+				UAudioBusSubsystem* AudioBusSubsystem = AudioDevice->GetSubsystem<UAudioBusSubsystem>();
+				check(AudioBusSubsystem);
+				AudioBusSubsystem->StopAudioBus(Audio::FAudioBusKey(AudioBusId));
 			}
 		}
 	}
@@ -66,11 +67,11 @@ void UAudioBus::PostEditChangeProperty(FPropertyChangedEvent& PropertyChangedEve
 		{
 			DeviceManager->IterateOverAllDevices([BusId = GetUniqueID(), NumChannels = AudioBusChannels](Audio::FDeviceId, FAudioDevice* InDevice)
 			{
-				if (Audio::FMixerDevice* MixerDevice = static_cast<Audio::FMixerDevice*>(InDevice))
-				{
-					MixerDevice->StopAudioBus(BusId);
-					MixerDevice->StartAudioBus(BusId, (int32)NumChannels + 1, false /* bInIsAutomatic */);
-				}
+				UAudioBusSubsystem* AudioBusSubsystem = InDevice->GetSubsystem<UAudioBusSubsystem>();
+				check(AudioBusSubsystem);
+				Audio::FAudioBusKey AudioBusKey = Audio::FAudioBusKey(BusId);
+				AudioBusSubsystem->StopAudioBus(AudioBusKey);
+				AudioBusSubsystem->StartAudioBus(AudioBusKey, (int32)NumChannels + 1, false /* bInIsAutomatic */);
 			});
 		}
 	}
diff --git a/Engine/Source/Runtime/Engine/Private/AudioDeviceManager.cpp b/Engine/Source/Runtime/Engine/Private/AudioDeviceManager.cpp
--- a/Engine/Source/Runtime/Engine/Private/AudioDeviceManager.cpp
+++ b/Engine/Source/Runtime/Engine/Private/AudioDeviceManager.cpp
@@ -2,6 +2,7 @@
 #include "AudioDeviceManager.h"
 
 #include "Audio.h"
+#include "AudioBusSubsystem.h"
 #include "Audio/AudioDebug.h"
 #include "AudioAnalytics.h"
 #include "AudioDefines.h"
@@ -530,6 +531,12 @@ void FAudioDeviceManager::UnregisterWorld(UWorld* InWorld, Audio::FDeviceId Devi
 	{
 		if (DeviceContainer->WorldsUsingThisDevice.Contains(InWorld))
 		{
+			if (FAudioDevice* AudioDevice = DeviceContainer->Device)
+			{
+				UAudioBusSubsystem* AudioBusSubsystem = AudioDevice->GetSubsystem<UAudioBusSubsystem>();
+				check(AudioBusSubsystem);
+				AudioBusSubsystem->ShutdownDefaultAudioBuses();
+			}
 			UE_LOG(LogAudio, Display, TEXT("Audio Device unregistered from world '%s'."), *InWorld->GetName());
 			DeviceContainer->WorldsUsingThisDevice.Remove(InWorld);
 			FAudioDeviceWorldDelegates::OnWorldUnregisteredWithAudioDevice.Broadcast(InWorld, DeviceId);
diff --git a/Engine/Source/Runtime/Engine/Private/AudioSettings.cpp b/Engine/Source/Runtime/Engine/Private/AudioSettings.cpp
--- a/Engine/Source/Runtime/Engine/Private/AudioSettings.cpp
+++ b/Engine/Source/Runtime/Engine/Private/AudioSettings.cpp
@@ -1,6 +1,7 @@
 // Copyright Epic Games, Inc. All Rights Reserved.
 #include "Sound/AudioSettings.h"
 
+#include "AudioBusSubsystem.h"
 #include "AudioDevice.h"
 #include "AudioDeviceManager.h"
 #include "AudioMixerDevice.h"
@@ -162,7 +163,9 @@ void UAudioSettings::PostEditChangeChainProperty(FPropertyChangedChainEvent& Pro
 			{
 				DeviceManager->IterateOverAllDevices([this](Audio::FDeviceId, FAudioDevice* InDevice)
 				{
-					InDevice->InitDefaultAudioBuses();
+					UAudioBusSubsystem* AudioBusSubsystem = InDevice->GetSubsystem<UAudioBusSubsystem>();
+					check(AudioBusSubsystem);
+					AudioBusSubsystem->InitDefaultAudioBuses();
 				});
 			}
 		}
diff --git a/Engine/Source/Runtime/Engine/Private/Subsystems/AudioEngineSubsystem.cpp b/Engine/Source/Runtime/Engine/Private/Subsystems/AudioEngineSubsystem.cpp
--- a/Engine/Source/Runtime/Engine/Private/Subsystems/AudioEngineSubsystem.cpp
+++ b/Engine/Source/Runtime/Engine/Private/Subsystems/AudioEngineSubsystem.cpp
@@ -1,6 +1,8 @@
 // Copyright Epic Games, Inc. All Rights Reserved.
 
 #include "Subsystems/AudioEngineSubsystem.h"
+#include "AudioDeviceHandle.h"
+#include "AudioMixerDevice.h"
 
 #include UE_INLINE_GENERATED_CPP_BY_NAME(AudioEngineSubsystem)
 
@@ -22,3 +24,35 @@ FAudioDeviceHandle UAudioEngineSubsystem::GetAudioDeviceHandle() const
 	return FAudioDeviceHandle();
 }
 
+Audio::FMixerSourceManager* UAudioEngineSubsystem::GetMutableSourceManager()
+{
+	Audio::FMixerDevice* MixerDevice = GetMutableMixerDevice();
+	if (MixerDevice)
+	{
+		return MixerDevice->GetSourceManager();
+	}
+	return nullptr;
+}
+
+const Audio::FMixerSourceManager* UAudioEngineSubsystem::GetSourceManager() const
+{
+	const Audio::FMixerDevice* MixerDevice = GetMixerDevice();
+	if (MixerDevice)
+	{
+		return MixerDevice->GetSourceManager();
+	}
+	return nullptr;
+}
+
+Audio::FMixerDevice* UAudioEngineSubsystem::GetMutableMixerDevice()
+{
+	Audio::FMixerDevice* MixerDevice = static_cast<Audio::FMixerDevice*>(GetAudioDeviceHandle().GetAudioDevice());
+	return MixerDevice;
+}
+
+const Audio::FMixerDevice* UAudioEngineSubsystem::GetMixerDevice() const
+{
+	const Audio::FMixerDevice* MixerDevice = static_cast<Audio::FMixerDevice*>(GetAudioDeviceHandle().GetAudioDevice());
+	return MixerDevice;
+}
+
diff --git a/Engine/Source/Runtime/Engine/Public/Subsystems/AudioEngineSubsystem.h b/Engine/Source/Runtime/Engine/Public/Subsystems/AudioEngineSubsystem.h
--- a/Engine/Source/Runtime/Engine/Public/Subsystems/AudioEngineSubsystem.h
+++ b/Engine/Source/Runtime/Engine/Public/Subsystems/AudioEngineSubsystem.h
@@ -2,13 +2,18 @@
 
 #pragma once
 
+#include "AudioDeviceManager.h"
 #include "Subsystem.h"
 #include "SubsystemCollection.h"
-#include "AudioDeviceManager.h"
 #include "AudioEngineSubsystem.generated.h"
 
 // Forward Declarations 
 class FAudioDevice;
+namespace Audio
+{
+	class FMixerDevice;
+	class FMixerSourceManager;
+}
 
 /**
  * UAudioSubsystemCollectionRoot
@@ -53,6 +58,18 @@ public:
 
 	/** Returns the owning audio device handle */
 	virtual FAudioDeviceHandle GetAudioDeviceHandle() const final;
+
+	/** Return a mutable version of the source manager associated with the owning device handle */
+	virtual Audio::FMixerSourceManager* GetMutableSourceManager() final;
+
+	/** Return the source manager associated with the owning device handle */
+	virtual const Audio::FMixerSourceManager* GetSourceManager() const final;
+
+	/** Return a mutable version of the mixer device from the owning device handle */
+	virtual Audio::FMixerDevice* GetMutableMixerDevice() final;
+
+	/** Return the mixer device from the owning device handle */
+	virtual const Audio::FMixerDevice* GetMixerDevice() const final;
 };
 
 /**
