[Quartz] cleaning up casting arithmatic that results in a div by zero in an edge case (bar deadline w/ bar-relative counting point).

diff --git a/Engine/Source/Runtime/AudioMixer/Private/Quartz/QuartzMetronome.cpp b/Engine/Source/Runtime/AudioMixer/Private/Quartz/QuartzMetronome.cpp
--- a/Engine/Source/Runtime/AudioMixer/Private/Quartz/QuartzMetronome.cpp
+++ b/Engine/Source/Runtime/AudioMixer/Private/Quartz/QuartzMetronome.cpp
@@ -250,10 +250,7 @@ namespace Audio
 			return static_cast<float>(PulseDurations.Num());
 		}
 
-		int32 LengthOfBar = MusicalDurationsInFrames[EQuartzCommandQuantization::Bar];
-		int32 LengthOfOne = MusicalDurationsInFrames[InSubdivision];
-
-		return MusicalDurationsInFrames[EQuartzCommandQuantization::Bar] / static_cast<float>(MusicalDurationsInFrames[InSubdivision]);
+		return static_cast<float>(MusicalDurationsInFrames[EQuartzCommandQuantization::Bar] / MusicalDurationsInFrames[InSubdivision]);
 	}
 
 	float FQuartzMetronome::CountNumSubdivisionsSinceBarStart(EQuartzCommandQuantization InSubdivision) const
