Expanding Audio Mixer Thread debug info to track methods being called and exited

diff --git a/Engine/Source/Runtime/AudioMixer/Private/AudioMixerSourceManager.cpp b/Engine/Source/Runtime/AudioMixer/Private/AudioMixerSourceManager.cpp
--- a/Engine/Source/Runtime/AudioMixer/Private/AudioMixerSourceManager.cpp
+++ b/Engine/Source/Runtime/AudioMixer/Private/AudioMixerSourceManager.cpp
@@ -21,6 +21,7 @@
 #include "Trace/Trace.h"
 
 
+
 // Link to "Audio" profiling category
 CSV_DECLARE_CATEGORY_MODULE_EXTERN(AUDIOMIXERCORE_API, Audio);
 static int32 DisableParallelSourceProcessingCvar = 1;
@@ -101,7 +102,7 @@ FAutoConsoleVariableRef CVarCommandBufferMaxSizeMb(
 	TEXT("How big to allow the command buffer to grow before ignoring more commands"),
 	ECVF_Default);
 
-static float AudioCommandExecTimeMsWarningThresholdCvar = 1.f;
+static float AudioCommandExecTimeMsWarningThresholdCvar = 500.f;
 FAutoConsoleVariableRef CVarAudioCommandExecTimeMsWarningThreshold(
 	TEXT("au.AudioThreadCommand.ExecutionTimeWarningThresholdInMs"),
 	AudioCommandExecTimeMsWarningThresholdCvar,
@@ -137,6 +138,9 @@ static FAutoConsoleCommand GModulationSetMaxPitchRange(
 	)
 );
 
+#define AUDIO_MIXER_THREAD_DEBUG_UPDATE()\
+	FScopedAudioThreadCommandInfo ScopedDebugInfo(ANSI_TO_TCHAR(__FUNCTION__), CurrentAudioMixerThreadCommandInfo);\
+
 #define ENVELOPE_TAIL_THRESHOLD (1.58489e-5f) // -96 dB
 
 #define VALIDATE_SOURCE_MIXER_STATE 1
@@ -245,6 +249,7 @@ namespace Audio
 
 	void FMixerSourceManager::Init(const FSourceManagerInitParams& InitParams)
 	{
+		AUDIO_MIXER_THREAD_DEBUG_UPDATE()
 		AUDIO_MIXER_CHECK(InitParams.NumSources > 0);
 
 		if (bInitialized || !MixerDevice)
@@ -439,6 +444,7 @@ namespace Audio
 
 	void FMixerSourceManager::Update(bool bTimedOut)
 	{
+		AUDIO_MIXER_THREAD_DEBUG_UPDATE()
 		AUDIO_MIXER_CHECK_GAME_THREAD(MixerDevice);
 
 #if VALIDATE_SOURCE_MIXER_STATE
@@ -516,6 +522,7 @@ namespace Audio
 
 	void FMixerSourceManager::ReleaseSource(const int32 SourceId)
 	{
+		AUDIO_MIXER_THREAD_DEBUG_UPDATE()
 		AUDIO_MIXER_CHECK_AUDIO_PLAT_THREAD(MixerDevice);
 
 		AUDIO_MIXER_CHECK(SourceId < NumTotalSources);
@@ -706,6 +713,7 @@ namespace Audio
 
 	void FMixerSourceManager::BuildSourceEffectChain(const int32 SourceId, FSoundEffectSourceInitData& InitData, const TArray<FSourceEffectChainEntry>& InSourceEffectChain, TArray<TSoundEffectSourcePtr>& OutSourceEffects)
 	{
+		AUDIO_MIXER_THREAD_DEBUG_UPDATE()
 		// Create new source effects. The memory will be owned by the source manager.
 		FScopeLock ScopeLock(&EffectChainMutationCriticalSection);
 		for (const FSourceEffectChainEntry& ChainEntry : InSourceEffectChain)
@@ -729,6 +737,7 @@ namespace Audio
 
 	void FMixerSourceManager::ResetSourceEffectChain(const int32 SourceId)
 	{
+		AUDIO_MIXER_THREAD_DEBUG_UPDATE()
 		FScopeLock ScopeLock(&EffectChainMutationCriticalSection);
 		{
 			FSourceInfo& SourceInfo = SourceInfos[SourceId];
@@ -757,6 +766,7 @@ namespace Audio
 
 	bool FMixerSourceManager::GetFreeSourceId(int32& OutSourceId)
 	{
+		AUDIO_MIXER_THREAD_DEBUG_UPDATE()
 		AUDIO_MIXER_CHECK_GAME_THREAD(MixerDevice);
 
 		if (GameThreadInfo.FreeSourceIndices.Num())
@@ -789,6 +799,7 @@ namespace Audio
 
 	void FMixerSourceManager::InitSource(const int32 SourceId, const FMixerSourceVoiceInitParams& InitParams)
 	{
+		AUDIO_MIXER_THREAD_DEBUG_UPDATE()
 		AUDIO_MIXER_CHECK(SourceId < NumTotalSources);
 		AUDIO_MIXER_CHECK(GameThreadInfo.bIsBusy[SourceId]);
 		AUDIO_MIXER_CHECK(!GameThreadInfo.bIsDebugMode[SourceId]);
@@ -1111,6 +1122,7 @@ namespace Audio
 
 	void FMixerSourceManager::ReleaseSourceId(const int32 SourceId)
 	{
+		AUDIO_MIXER_THREAD_DEBUG_UPDATE()
 		AUDIO_MIXER_CHECK(GameThreadInfo.bIsBusy[SourceId]);
 		AUDIO_MIXER_CHECK_GAME_THREAD(MixerDevice);
 
@@ -1137,7 +1149,8 @@ namespace Audio
 
 	void FMixerSourceManager::StartAudioBus(FAudioBusKey InAudioBusKey, int32 InNumChannels, bool bInIsAutomatic)
 	{
-		if (AudioBusKeys_AudioThread.Contains(InAudioBusKey))
+		AUDIO_MIXER_THREAD_DEBUG_UPDATE()
+		if (AudioBusIds_AudioThread.Contains(InAudioBusKey))
 		{
 			return;
 		}
@@ -1166,7 +1179,8 @@ namespace Audio
 
 	void FMixerSourceManager::StopAudioBus(FAudioBusKey InAudioBusKey)
 	{
-		if (!AudioBusKeys_AudioThread.Contains(InAudioBusKey))
+		AUDIO_MIXER_THREAD_DEBUG_UPDATE()
+		if (!AudioBusIds_AudioThread.Contains(InAudioBusKey))
 		{
 			return;
 		}
@@ -1194,17 +1208,19 @@ namespace Audio
 					AudioBuses.Remove(InAudioBusKey);
 				}
 			}
-		}, AUDIO_MIXER_THREAD_COMMAND_STRING("StartAudioBus"));
+		}, AUDIO_MIXER_THREAD_COMMAND_STRING("StopAudioBus"));
 	}
 
 	bool FMixerSourceManager::IsAudioBusActive(FAudioBusKey InAudioBusKey) const
 	{
+		AUDIO_MIXER_THREAD_DEBUG_UPDATE()
 		AUDIO_MIXER_CHECK_GAME_THREAD(MixerDevice);
 		return AudioBusKeys_AudioThread.Contains(InAudioBusKey);
 	}
 
 	int32 FMixerSourceManager::GetAudioBusNumChannels(FAudioBusKey InAudioBusKey) const
 	{
+		AUDIO_MIXER_THREAD_DEBUG_UPDATE()
 		AUDIO_MIXER_CHECK_AUDIO_PLAT_THREAD(MixerDevice);
 		TSharedPtr<FMixerAudioBus> AudioBusPtr = AudioBuses.FindRef(InAudioBusKey);
 		if (AudioBusPtr.IsValid())
@@ -1217,6 +1233,7 @@ namespace Audio
 
 	void FMixerSourceManager::AddPatchOutputForAudioBus(FAudioBusKey InAudioBusKey, const FPatchOutputStrongPtr& InPatchOutputStrongPtr)
 	{
+		AUDIO_MIXER_THREAD_DEBUG_UPDATE()
 		AUDIO_MIXER_CHECK_AUDIO_PLAT_THREAD(MixerDevice);
 		if (MixerDevice->IsAudioRenderingThread())
 		{
@@ -1238,6 +1255,7 @@ namespace Audio
 
 	void FMixerSourceManager::AddPatchOutputForAudioBus_AudioThread(FAudioBusKey InAudioBusKey, const FPatchOutputStrongPtr& InPatchOutputStrongPtr)
 	{
+		AUDIO_MIXER_THREAD_DEBUG_UPDATE()
 		AudioMixerThreadCommand([this, InAudioBusKey, NewPatchPtr = InPatchOutputStrongPtr]() mutable
 		{
 			AddPatchOutputForAudioBus(InAudioBusKey, NewPatchPtr);
@@ -1275,6 +1293,7 @@ namespace Audio
 
 	void FMixerSourceManager::Play(const int32 SourceId)
 	{
+		AUDIO_MIXER_THREAD_DEBUG_UPDATE()
 		AUDIO_MIXER_CHECK(SourceId < NumTotalSources);
 		AUDIO_MIXER_CHECK(GameThreadInfo.bIsBusy[SourceId]);
 		AUDIO_MIXER_CHECK_GAME_THREAD(MixerDevice);
@@ -1300,6 +1319,7 @@ namespace Audio
 
 	void FMixerSourceManager::CancelQuantizedSound(const int32 SourceId)
 	{
+		AUDIO_MIXER_THREAD_DEBUG_UPDATE()
 		if (!MixerDevice)
 		{
 			return;
@@ -1333,6 +1353,7 @@ namespace Audio
 
 	void FMixerSourceManager::Stop(const int32 SourceId)
 	{
+		AUDIO_MIXER_THREAD_DEBUG_UPDATE()
 		if (!MixerDevice)
 		{
 			return;
@@ -1353,6 +1374,7 @@ namespace Audio
 
 	void FMixerSourceManager::StopInternal(const int32 SourceId)
 	{
+		AUDIO_MIXER_THREAD_DEBUG_UPDATE()
 		AUDIO_MIXER_CHECK_AUDIO_PLAT_THREAD(MixerDevice);
 
 		FSourceInfo& SourceInfo = SourceInfos[SourceId];
@@ -1374,6 +1396,7 @@ namespace Audio
 
 	void FMixerSourceManager::StopFade(const int32 SourceId, const int32 NumFrames)
 	{
+		AUDIO_MIXER_THREAD_DEBUG_UPDATE()
 		AUDIO_MIXER_CHECK(SourceId < NumTotalSources);
 		AUDIO_MIXER_CHECK(GameThreadInfo.bIsBusy[SourceId]);
 		AUDIO_MIXER_CHECK(NumFrames > 0);
@@ -1417,10 +1440,10 @@ namespace Audio
 			AUDIO_MIXER_DEBUG_LOG(SourceId, TEXT("Is stopping with fade"));
 		}, AUDIO_MIXER_THREAD_COMMAND_STRING("StopFade()"));
 	}
-
-
+	
 	void FMixerSourceManager::Pause(const int32 SourceId)
 	{
+		AUDIO_MIXER_THREAD_DEBUG_UPDATE()
 		AUDIO_MIXER_CHECK(SourceId < NumTotalSources);
 		AUDIO_MIXER_CHECK(GameThreadInfo.bIsBusy[SourceId]);
 		AUDIO_MIXER_CHECK_GAME_THREAD(MixerDevice);
@@ -1438,6 +1461,7 @@ namespace Audio
 
 	void FMixerSourceManager::SetPitch(const int32 SourceId, const float Pitch)
 	{
+		AUDIO_MIXER_THREAD_DEBUG_UPDATE()
 		AUDIO_MIXER_CHECK(SourceId < NumTotalSources);
 		AUDIO_MIXER_CHECK(GameThreadInfo.bIsBusy[SourceId]);
 
@@ -1452,6 +1476,7 @@ namespace Audio
 
 	void FMixerSourceManager::SetVolume(const int32 SourceId, const float Volume)
 	{
+		AUDIO_MIXER_THREAD_DEBUG_UPDATE()
 		AUDIO_MIXER_CHECK(SourceId < NumTotalSources);
 		AUDIO_MIXER_CHECK(GameThreadInfo.bIsBusy[SourceId]);
 		AUDIO_MIXER_CHECK_GAME_THREAD(MixerDevice);
@@ -1479,6 +1504,7 @@ namespace Audio
 
 	void FMixerSourceManager::SetDistanceAttenuation(const int32 SourceId, const float DistanceAttenuation)
 	{
+		AUDIO_MIXER_THREAD_DEBUG_UPDATE()
 		AUDIO_MIXER_CHECK(SourceId < NumTotalSources);
 		AUDIO_MIXER_CHECK(GameThreadInfo.bIsBusy[SourceId]);
 		AUDIO_MIXER_CHECK_GAME_THREAD(MixerDevice);
@@ -1500,6 +1526,7 @@ namespace Audio
 
 	void FMixerSourceManager::SetSpatializationParams(const int32 SourceId, const FSpatializationParams& InParams)
 	{
+		AUDIO_MIXER_THREAD_DEBUG_UPDATE()
 		AUDIO_MIXER_CHECK(SourceId < NumTotalSources);
 		AUDIO_MIXER_CHECK(GameThreadInfo.bIsBusy[SourceId]);
 		AUDIO_MIXER_CHECK_GAME_THREAD(MixerDevice);
@@ -1514,6 +1541,7 @@ namespace Audio
 
 	void FMixerSourceManager::SetChannelMap(const int32 SourceId, const uint32 NumInputChannels, const Audio::FAlignedFloatBuffer& ChannelMap, const bool bInIs3D, const bool bInIsCenterChannelOnly)
 	{
+		AUDIO_MIXER_THREAD_DEBUG_UPDATE()
 		AUDIO_MIXER_CHECK(SourceId < NumTotalSources);
 		AUDIO_MIXER_CHECK(GameThreadInfo.bIsBusy[SourceId]);
 		AUDIO_MIXER_CHECK_GAME_THREAD(MixerDevice);
@@ -1545,6 +1573,7 @@ namespace Audio
 
 	void FMixerSourceManager::SetLPFFrequency(const int32 SourceId, const float InLPFFrequency)
 	{
+		AUDIO_MIXER_THREAD_DEBUG_UPDATE()
 		AUDIO_MIXER_CHECK(SourceId < NumTotalSources);
 		AUDIO_MIXER_CHECK(GameThreadInfo.bIsBusy[SourceId]);
 		AUDIO_MIXER_CHECK_GAME_THREAD(MixerDevice);
@@ -1565,6 +1594,7 @@ namespace Audio
 
 	void FMixerSourceManager::SetHPFFrequency(const int32 SourceId, const float InHPFFrequency)
 	{
+		AUDIO_MIXER_THREAD_DEBUG_UPDATE()
 		AUDIO_MIXER_CHECK(SourceId < NumTotalSources);
 		AUDIO_MIXER_CHECK(GameThreadInfo.bIsBusy[SourceId]);
 		AUDIO_MIXER_CHECK_GAME_THREAD(MixerDevice);
@@ -1584,6 +1614,7 @@ namespace Audio
 
 	void FMixerSourceManager::SetModLPFFrequency(const int32 SourceId, const float InLPFFrequency)
 	{
+		AUDIO_MIXER_THREAD_DEBUG_UPDATE()
 		AUDIO_MIXER_CHECK(SourceId < NumTotalSources);
 		AUDIO_MIXER_CHECK(GameThreadInfo.bIsBusy[SourceId]);
 		AUDIO_MIXER_CHECK_GAME_THREAD(MixerDevice);
@@ -1600,6 +1631,7 @@ namespace Audio
 
 	void FMixerSourceManager::SetModHPFFrequency(const int32 SourceId, const float InHPFFrequency)
 	{
+		AUDIO_MIXER_THREAD_DEBUG_UPDATE()
 		AUDIO_MIXER_CHECK(SourceId < NumTotalSources);
 		AUDIO_MIXER_CHECK(GameThreadInfo.bIsBusy[SourceId]);
 		AUDIO_MIXER_CHECK_GAME_THREAD(MixerDevice);
@@ -1616,6 +1648,7 @@ namespace Audio
 
 	void FMixerSourceManager::SetSourceBufferListener(const int32 SourceId, FSharedISourceBufferListenerPtr& InSourceBufferListener, bool InShouldSourceBufferListenerZeroBuffer)
 	{
+		AUDIO_MIXER_THREAD_DEBUG_UPDATE()
 		AUDIO_MIXER_CHECK(SourceId < NumTotalSources);
 		AUDIO_MIXER_CHECK(GameThreadInfo.bIsBusy[SourceId]);
 		AUDIO_MIXER_CHECK_GAME_THREAD(MixerDevice);
@@ -1632,6 +1665,7 @@ namespace Audio
 
 	void FMixerSourceManager::SetModVolume(const int32 SourceId, const float InModVolume)
 	{
+		AUDIO_MIXER_THREAD_DEBUG_UPDATE()
 		AUDIO_MIXER_CHECK(SourceId < NumTotalSources);
 		AUDIO_MIXER_CHECK(GameThreadInfo.bIsBusy[SourceId]);
 		AUDIO_MIXER_CHECK_GAME_THREAD(MixerDevice);
@@ -1647,6 +1681,7 @@ namespace Audio
 
 	void FMixerSourceManager::SetModPitch(const int32 SourceId, const float InModPitch)
 	{
+		AUDIO_MIXER_THREAD_DEBUG_UPDATE()
 		AUDIO_MIXER_CHECK(SourceId < NumTotalSources);
 		AUDIO_MIXER_CHECK(GameThreadInfo.bIsBusy[SourceId]);
 		AUDIO_MIXER_CHECK_GAME_THREAD(MixerDevice);
@@ -1662,6 +1697,7 @@ namespace Audio
 
 	void FMixerSourceManager::SetSubmixSendInfo(const int32 SourceId, const FMixerSourceSubmixSend& InSubmixSend)
 	{
+		AUDIO_MIXER_THREAD_DEBUG_UPDATE()
 		AUDIO_MIXER_CHECK(SourceId < NumTotalSources);
 		AUDIO_MIXER_CHECK(GameThreadInfo.bIsBusy[SourceId]);
 		AUDIO_MIXER_CHECK_GAME_THREAD(MixerDevice);
@@ -1721,6 +1757,7 @@ namespace Audio
 
 	void FMixerSourceManager::ClearSubmixSendInfo(const int32 SourceId, const FMixerSourceSubmixSend& InSubmixSend)
 	{
+		AUDIO_MIXER_THREAD_DEBUG_UPDATE()
 		AUDIO_MIXER_CHECK(SourceId < NumTotalSources);
 		AUDIO_MIXER_CHECK(GameThreadInfo.bIsBusy[SourceId]);
 		AUDIO_MIXER_CHECK_GAME_THREAD(MixerDevice);
@@ -1769,6 +1806,7 @@ namespace Audio
 
 	void FMixerSourceManager::SetBusSendInfo(const int32 SourceId, EBusSendType InAudioBusSendType, uint32 AudioBusId, float BusSendLevel)
 	{
+		AUDIO_MIXER_THREAD_DEBUG_UPDATE()
 		AUDIO_MIXER_CHECK(SourceId < NumTotalSources);
 		AUDIO_MIXER_CHECK(GameThreadInfo.bIsBusy[SourceId]);
 		AUDIO_MIXER_CHECK_GAME_THREAD(MixerDevice);
@@ -1824,6 +1862,7 @@ namespace Audio
 
 	void FMixerSourceManager::SetListenerTransforms(const TArray<FTransform>& InListenerTransforms)
 	{
+		AUDIO_MIXER_THREAD_DEBUG_UPDATE()
 		AudioMixerThreadCommand([this, InListenerTransforms]()
 		{
 			ListenerTransforms = InListenerTransforms;
@@ -1862,6 +1901,7 @@ namespace Audio
 
 	void FMixerSourceManager::ReadSourceFrame(const int32 SourceId)
 	{
+		AUDIO_MIXER_THREAD_DEBUG_UPDATE()
 		FSourceInfo& SourceInfo = SourceInfos[SourceId];
 
 		const int32 NumChannels = SourceInfo.NumInputChannels;
@@ -1979,6 +2019,7 @@ namespace Audio
 
 	void FMixerSourceManager::ComputeSourceBuffersForIdRange(const bool bGenerateBuses, const int32 SourceIdStart, const int32 SourceIdEnd)
 	{
+		AUDIO_MIXER_THREAD_DEBUG_UPDATE()
 		CSV_SCOPED_TIMING_STAT(Audio, SourceBuffers);
 		CONDITIONAL_SCOPE_CYCLE_COUNTER(STAT_AudioMixerSourceBuffers, (SourceIdStart < SourceIdEnd));
 
@@ -2219,6 +2260,7 @@ namespace Audio
 
 	void FMixerSourceManager::ComputeBuses()
 	{
+		AUDIO_MIXER_THREAD_DEBUG_UPDATE()
 		// Loop through the bus registry and mix source audio
 		for (auto& Entry : AudioBuses)
 		{
@@ -2229,6 +2271,7 @@ namespace Audio
 
 	void FMixerSourceManager::UpdateBuses()
 	{
+		AUDIO_MIXER_THREAD_DEBUG_UPDATE()
 		// Update the bus states post mixing. This flips the current/previous buffer indices.
 		for (auto& Entry : AudioBuses)
 		{
@@ -2268,6 +2311,7 @@ namespace Audio
 
 	void FMixerSourceManager::ComputePluginAudio(FSourceInfo& SourceInfo, FMixerSourceSubmixOutputBuffer& InSourceSubmixOutputBuffer, int32 SourceId, int32 NumSamples)
 	{
+		AUDIO_MIXER_THREAD_DEBUG_UPDATE()
 		if (BypassAudioPluginsCvar)
 		{
 			// If we're bypassing audio plugins, our pre- and post-effect channels are the same as the input channels
@@ -2459,6 +2503,7 @@ namespace Audio
 	{
 		CSV_SCOPED_TIMING_STAT(Audio, SourceEffectsBuffers);
 		CONDITIONAL_SCOPE_CYCLE_COUNTER(STAT_AudioMixerSourceEffectBuffers, (SourceIdStart < SourceIdEnd));
+		AUDIO_MIXER_THREAD_DEBUG_UPDATE()
 
 		const bool bIsDebugModeEnabled = DebugSoloSources.Num() > 0;
 
@@ -2763,6 +2808,7 @@ namespace Audio
 	{
 		CSV_SCOPED_TIMING_STAT(Audio, SourceOutputBuffers);
 		CONDITIONAL_SCOPE_CYCLE_COUNTER(STAT_AudioMixerSourceOutputBuffers, (SourceIdStart < SourceIdEnd));
+		AUDIO_MIXER_THREAD_DEBUG_UPDATE()
 
 		for (int32 SourceId = SourceIdStart; SourceId < SourceIdEnd; ++SourceId)
 		{
@@ -2791,6 +2837,7 @@ namespace Audio
 
 	void FMixerSourceManager::GenerateSourceAudio(const bool bGenerateBuses, const int32 SourceIdStart, const int32 SourceIdEnd)
 	{
+		AUDIO_MIXER_THREAD_DEBUG_UPDATE()
 		// Buses generate their input buffers independently
 		// Get the next block of frames from the source buffers
 		ComputeSourceBuffersForIdRange(bGenerateBuses, SourceIdStart, SourceIdEnd);
@@ -2804,6 +2851,7 @@ namespace Audio
 
 	void FMixerSourceManager::GenerateSourceAudio(const bool bGenerateBuses)
 	{
+		AUDIO_MIXER_THREAD_DEBUG_UPDATE()
 		// If there are no buses, don't need to do anything here
 		if (bGenerateBuses && !AudioBuses.Num())
 		{
@@ -2834,6 +2882,7 @@ namespace Audio
 
 	void FMixerSourceManager::MixOutputBuffers(const int32 SourceId, int32 InNumOutputChannels, const float InSendLevel, EMixerSourceSubmixSendStage InSubmixSendStage, FAlignedFloatBuffer& OutWetBuffer) const
 	{
+		AUDIO_MIXER_THREAD_DEBUG_UPDATE()
 		if (InSendLevel > 0.0f)
 		{
 			const FSourceInfo& SourceInfo = SourceInfos[SourceId];
@@ -2857,6 +2906,7 @@ namespace Audio
 
 	const ISoundfieldAudioPacket* FMixerSourceManager::GetEncodedOutput(const int32 SourceId, const FSoundfieldEncodingKey& InKey) const
 	{
+		AUDIO_MIXER_THREAD_DEBUG_UPDATE()
 		AUDIO_MIXER_CHECK_AUDIO_PLAT_THREAD(MixerDevice);
 
 		const FSourceInfo& SourceInfo = SourceInfos[SourceId];
@@ -2879,6 +2929,7 @@ namespace Audio
 
 	void FMixerSourceManager::UpdateDeviceChannelCount(const int32 InNumOutputChannels)
 	{
+		AUDIO_MIXER_THREAD_DEBUG_UPDATE()
 		AudioMixerThreadCommand([this, InNumOutputChannels]()
 		{
 			NumOutputSamples = NumOutputFrames * MixerDevice->GetNumDeviceChannels();
@@ -2921,6 +2972,7 @@ namespace Audio
 
 	void FMixerSourceManager::UpdateSourceEffectChain(const uint32 InSourceEffectChainId, const TArray<FSourceEffectChainEntry>& InSourceEffectChain, const bool bPlayEffectChainTails)
 	{
+		AUDIO_MIXER_THREAD_DEBUG_UPDATE()
 		AudioMixerThreadCommand([this, InSourceEffectChainId, InSourceEffectChain, bPlayEffectChainTails]()
 		{
 			FSoundEffectSourceInitData InitData;
@@ -3176,18 +3228,25 @@ namespace Audio
 		MpscCommandQueue.Enqueue(MoveTemp(InCommand));
 	}
 
+	void FMixerSourceManager::FCurrentAudioMixerThreadCommandInfo::SetNewTask(FName InTask)
+	{
+#if WITH_AUDIO_MIXER_THREAD_COMMAND_DEBUG
+		CallerDebugInfo = InTask;
+		QueueTimeInCycles = FPlatformTime::Cycles64();
+#endif // #if WITH_AUDIO_MIXER_THREAD_COMMAND_DEBUG
+	}
+
 	void FMixerSourceManager::FCurrentAudioMixerThreadCommandInfo::LogWarning() const
 	{
 #if WITH_AUDIO_MIXER_THREAD_COMMAND_DEBUG
-		if(CallerDebugInfo.IsValid())
+		if(!CallerDebugInfo.IsNone())
 		{
-
 			UE_LOG(LogAudioMixer, Warning, TEXT("The AudioThreadCommand currently executing was called by %s, and started execution %f milliseconds ago")
 				, *CallerDebugInfo.ToString(), FPlatformTime::ToMilliseconds64(FPlatformTime::Cycles64() - QueueTimeInCycles.GetValue()));
 		}
 		else
 		{
-			UE_LOG(LogAudioMixer, Warning, TEXT("The AudioThreadCommand has not yet executed a command"));
+			UE_LOG(LogAudioMixer, Warning, TEXT("We don't currently have a task named"));
 		}
 #endif // #if WITH_AUDIO_MIXER_THREAD_COMMAND_DEBUG
 	}
@@ -3195,15 +3254,14 @@ namespace Audio
 	void FMixerSourceManager::FCurrentAudioMixerThreadCommandInfo::LogError() const
 	{
 #if WITH_AUDIO_MIXER_THREAD_COMMAND_DEBUG
-		if(CallerDebugInfo.IsValid())
+		if(!CallerDebugInfo.IsNone())
 		{
-
 			UE_LOG(LogAudioMixer, Error, TEXT("The AudioThreadCommand currently executing was called by %s, and started execution %f milliseconds ago")
 				, *CallerDebugInfo.ToString(), FPlatformTime::ToMilliseconds64(FPlatformTime::Cycles64() - QueueTimeInCycles.GetValue()));
 		}
 		else
 		{
-			UE_LOG(LogAudioMixer, Error, TEXT("The AudioThreadCommand has not yet executed a command"));
+			UE_LOG(LogAudioMixer, Warning, TEXT("We don't currently have a task named"));
 		}
 #endif // #if WITH_AUDIO_MIXER_THREAD_COMMAND_DEBUG
 	}
@@ -3216,6 +3274,17 @@ namespace Audio
 		QueueTimeInCycles.Reset();
 	}
 
+	FMixerSourceManager::FScopedAudioThreadCommandInfo::FScopedAudioThreadCommandInfo(FName InTaskName, FCurrentAudioMixerThreadCommandInfo& InInfoRef)
+		: InfoRef(InInfoRef)
+	{
+		InfoRef.SetNewTask(InTaskName);
+	}
+
+	FMixerSourceManager::FScopedAudioThreadCommandInfo::~FScopedAudioThreadCommandInfo()
+	{
+		InfoRef.Reset();
+	}
+
 	void FMixerSourceManager::AudioMixerThreadCommand(TFunction<void()> InFunction, FName CallingFunction, bool bDeferExecution)
 	{
 		FAudioMixerThreadCommand AudioCommand(InFunction, CallingFunction, bDeferExecution);
@@ -3260,6 +3329,8 @@ namespace Audio
 	void FMixerSourceManager::PumpCommandQueue()
 	{
 		TRACE_CPUPROFILER_EVENT_SCOPE(AudioMixerThreadCommands::PumpCommandQueue)
+		AUDIO_MIXER_THREAD_DEBUG_UPDATE()
+		
 		// If we're already triggered, we need to wait for the audio thread to reset it before pumping
 		if (FPlatformProcess::SupportsMultithreading())
 		{
@@ -3294,8 +3365,7 @@ namespace Audio
 			FAudioMixerThreadCommand& AudioCommand = Commands.SourceCommandQueue[Id];
 #if WITH_AUDIO_MIXER_THREAD_COMMAND_DEBUG
 			// cache debug info about the command we are about to execute (for safe access from the producer thread)
-			CurrentAudioMixerThreadCommandInfo.CallerDebugInfo = AudioCommand.CallerDebugInfo;
-			CurrentAudioMixerThreadCommandInfo.QueueTimeInCycles = FPlatformTime::Cycles64();
+			CurrentAudioMixerThreadCommandInfo.SetNewTask(AudioCommand.CallerDebugInfo);
 #endif // #if WITH_AUDIO_MIXER_THREAD_COMMAND_DEBUG
 			if(AudioCommand.bDeferExecution)
 			{
@@ -3308,11 +3378,11 @@ namespace Audio
 			}
 
 #if WITH_AUDIO_MIXER_THREAD_COMMAND_DEBUG
-			const float ExecutionTimeInMs = static_cast<float>(FPlatformTime::ToMilliseconds64(FPlatformTime::Cycles64() - CurrentAudioMixerThreadCommandInfo.QueueTimeInCycles.GetValue()));
+			const float ExecutionTimeInMs = CurrentAudioMixerThreadCommandInfo.GetCurrentExecutionTimeMs();
 			if(LogEveryAudioThreadCommandCvar || (ExecutionTimeInMs >= AudioCommandExecTimeMsWarningThresholdCvar))
 			{
-				UE_LOG(LogAudioMixer, Warning, TEXT("Audio Thread Command from %s() took %f milleseconds to execute"),
-				*CurrentAudioMixerThreadCommandInfo.CallerDebugInfo.ToString(), ExecutionTimeInMs);
+				UE_LOG(LogAudioMixer, Warning, TEXT("Audio Thread Command from %s() took %f milleseconds to execute")
+				, *CurrentAudioMixerThreadCommandInfo.GetCurrentTaskName().ToString(), ExecutionTimeInMs);
 			}
 #endif // #if WITH_AUDIO_MIXER_THREAD_COMMAND_DEBUG
 
@@ -3337,6 +3407,7 @@ namespace Audio
 
 	void FMixerSourceManager::FlushCommandQueue(bool bPumpInCommand)
 	{
+		AUDIO_MIXER_THREAD_DEBUG_UPDATE()
 		check(CommandsProcessedEvent != nullptr);
 
 		// If we have no commands enqueued, exit
@@ -3381,6 +3452,8 @@ namespace Audio
 
 	void FMixerSourceManager::UpdatePendingReleaseData(bool bForceWait)
 	{
+		AUDIO_MIXER_THREAD_DEBUG_UPDATE()
+		
 		// Don't block, but let tasks finish naturally
 		for (int32 i = PendingSourceBuffers.Num() - 1; i >= 0; --i)
 		{
diff --git a/Engine/Source/Runtime/AudioMixer/Private/AudioMixerSourceManager.h b/Engine/Source/Runtime/AudioMixer/Private/AudioMixerSourceManager.h
--- a/Engine/Source/Runtime/AudioMixer/Private/AudioMixerSourceManager.h
+++ b/Engine/Source/Runtime/AudioMixer/Private/AudioMixerSourceManager.h
@@ -324,18 +324,49 @@ namespace Audio
 			// and that source gets initialized after the command executes
 			bool bDeferExecution;
 		};
+		
 
 		struct FCurrentAudioMixerThreadCommandInfo
 		{
+		private:
 			FName CallerDebugInfo;
 			FThreadSafeCounter64 QueueTimeInCycles;
+			
+		public:
+			void SetNewTask(FName InTask);
+			FName GetCurrentTaskName() const
+			{
+				return CallerDebugInfo;
+			}
+
+			float GetCurrentExecutionTimeMs() const
+			{
+				if(!CallerDebugInfo.IsNone())
+				{
+					return static_cast<float>(FPlatformTime::ToMilliseconds64(FPlatformTime::Cycles64() - QueueTimeInCycles.GetValue()));
+				}
+				return 0.f; // no task currently running
+			}
 
 			void LogWarning() const;
 			void LogError() const;
 			void Reset();
 		};
 
-		FCurrentAudioMixerThreadCommandInfo CurrentAudioMixerThreadCommandInfo;
+		struct FScopedAudioThreadCommandInfo
+		{
+			// ctor - name the task, reset the counter
+			FScopedAudioThreadCommandInfo() = delete;
+			FScopedAudioThreadCommandInfo(FName InTaskName, FCurrentAudioMixerThreadCommandInfo& InInfoRef);
+
+			// dtor - reset the task name and counter (not in a function)
+			~FScopedAudioThreadCommandInfo();
+
+			FCurrentAudioMixerThreadCommandInfo& InfoRef;
+			
+		};
+
+		mutable FCurrentAudioMixerThreadCommandInfo CurrentAudioMixerThreadCommandInfo; // mutable to allow const functions to register debug info
 
 		void AudioMixerThreadCommand(TFunction<void()> InFunction, FName CallingFunction = {}, bool bInDeferExecution = false);
 
