Removing audio thread debug logging (log spam and FName processing for little gain).
#pf 63bdf81f763c6c1064df28c6

diff --git a/Engine/Source/Runtime/AudioMixer/Private/AudioMixerSourceManager.cpp b/Engine/Source/Runtime/AudioMixer/Private/AudioMixerSourceManager.cpp
--- a/Engine/Source/Runtime/AudioMixer/Private/AudioMixerSourceManager.cpp
+++ b/Engine/Source/Runtime/AudioMixer/Private/AudioMixerSourceManager.cpp
@@ -138,9 +138,6 @@ static FAutoConsoleCommand GModulationSetMaxPitchRange(
 	)
 );
 
-#define AUDIO_MIXER_THREAD_DEBUG_UPDATE()\
-	FScopedAudioThreadCommandInfo ScopedDebugInfo(ANSI_TO_TCHAR(__FUNCTION__), CurrentAudioMixerThreadCommandInfo);\
-
 #define ENVELOPE_TAIL_THRESHOLD (1.58489e-5f) // -96 dB
 
 #define VALIDATE_SOURCE_MIXER_STATE 1
@@ -249,7 +246,6 @@ namespace Audio
 
 	void FMixerSourceManager::Init(const FSourceManagerInitParams& InitParams)
 	{
-		AUDIO_MIXER_THREAD_DEBUG_UPDATE()
 		AUDIO_MIXER_CHECK(InitParams.NumSources > 0);
 
 		if (bInitialized || !MixerDevice)
@@ -417,7 +413,7 @@ namespace Audio
 				} Payload;
 				for (int32 i = 0; i < 65536; ++i)
 				{
-					AudioMixerThreadCommand([Payload] {}, AUDIO_MIXER_THREAD_COMMAND_STRING("SpamNopsCmd() -- Console command"));
+					AudioMixerThreadCommand([Payload] {});
 				}
 			})
 		);
@@ -434,7 +430,7 @@ namespace Audio
 					uint8 JunkBytes[1024];
 				} Payload;
 
-				AudioMixerThreadCommand([Payload] {while(true){}}, AUDIO_MIXER_THREAD_COMMAND_STRING("SpamEndlessCmd() -- Console command"));
+				AudioMixerThreadCommand([Payload] {while(true){}});
 			})
 		);
 
@@ -444,7 +440,6 @@ namespace Audio
 
 	void FMixerSourceManager::Update(bool bTimedOut)
 	{
-		AUDIO_MIXER_THREAD_DEBUG_UPDATE()
 		AUDIO_MIXER_CHECK_GAME_THREAD(MixerDevice);
 
 #if VALIDATE_SOURCE_MIXER_STATE
@@ -522,7 +517,6 @@ namespace Audio
 
 	void FMixerSourceManager::ReleaseSource(const int32 SourceId)
 	{
-		AUDIO_MIXER_THREAD_DEBUG_UPDATE()
 		AUDIO_MIXER_CHECK_AUDIO_PLAT_THREAD(MixerDevice);
 
 		AUDIO_MIXER_CHECK(SourceId < NumTotalSources);
@@ -713,7 +707,6 @@ namespace Audio
 
 	void FMixerSourceManager::BuildSourceEffectChain(const int32 SourceId, FSoundEffectSourceInitData& InitData, const TArray<FSourceEffectChainEntry>& InSourceEffectChain, TArray<TSoundEffectSourcePtr>& OutSourceEffects)
 	{
-		AUDIO_MIXER_THREAD_DEBUG_UPDATE()
 		// Create new source effects. The memory will be owned by the source manager.
 		FScopeLock ScopeLock(&EffectChainMutationCriticalSection);
 		for (const FSourceEffectChainEntry& ChainEntry : InSourceEffectChain)
@@ -737,7 +730,6 @@ namespace Audio
 
 	void FMixerSourceManager::ResetSourceEffectChain(const int32 SourceId)
 	{
-		AUDIO_MIXER_THREAD_DEBUG_UPDATE()
 		FScopeLock ScopeLock(&EffectChainMutationCriticalSection);
 		{
 			FSourceInfo& SourceInfo = SourceInfos[SourceId];
@@ -766,7 +758,6 @@ namespace Audio
 
 	bool FMixerSourceManager::GetFreeSourceId(int32& OutSourceId)
 	{
-		AUDIO_MIXER_THREAD_DEBUG_UPDATE()
 		AUDIO_MIXER_CHECK_GAME_THREAD(MixerDevice);
 
 		if (GameThreadInfo.FreeSourceIndices.Num())
@@ -799,7 +790,6 @@ namespace Audio
 
 	void FMixerSourceManager::InitSource(const int32 SourceId, const FMixerSourceVoiceInitParams& InitParams)
 	{
-		AUDIO_MIXER_THREAD_DEBUG_UPDATE()
 		AUDIO_MIXER_CHECK(SourceId < NumTotalSources);
 		AUDIO_MIXER_CHECK(GameThreadInfo.bIsBusy[SourceId]);
 		AUDIO_MIXER_CHECK(!GameThreadInfo.bIsDebugMode[SourceId]);
@@ -1117,12 +1107,11 @@ namespace Audio
 #endif 
 
 			AUDIO_MIXER_DEBUG_LOG(SourceId, TEXT("Is initializing"));
-		}, AUDIO_MIXER_THREAD_COMMAND_STRING("InitSource"));
+		});
 	}
 
 	void FMixerSourceManager::ReleaseSourceId(const int32 SourceId)
 	{
-		AUDIO_MIXER_THREAD_DEBUG_UPDATE()
 		AUDIO_MIXER_CHECK(GameThreadInfo.bIsBusy[SourceId]);
 		AUDIO_MIXER_CHECK_GAME_THREAD(MixerDevice);
 
@@ -1144,12 +1133,11 @@ namespace Audio
 			AUDIO_MIXER_CHECK_AUDIO_PLAT_THREAD(MixerDevice);
 
 			ReleaseSource(SourceId);
-		}, AUDIO_MIXER_THREAD_COMMAND_STRING("ReleaseSourceId"));
+		});
 	}
 
 	void FMixerSourceManager::StartAudioBus(FAudioBusKey InAudioBusKey, int32 InNumChannels, bool bInIsAutomatic)
 	{
-		AUDIO_MIXER_THREAD_DEBUG_UPDATE()
 		if (AudioBusKeys_AudioThread.Contains(InAudioBusKey))
 		{
 			return;
@@ -1174,12 +1162,11 @@ namespace Audio
 
 				AudioBuses.Add(InAudioBusKey, NewBusData);
 			}
-		}, AUDIO_MIXER_THREAD_COMMAND_STRING("StartAudioBus"));
+		});
 	}
 
 	void FMixerSourceManager::StopAudioBus(FAudioBusKey InAudioBusKey)
 	{
-		AUDIO_MIXER_THREAD_DEBUG_UPDATE()
 		if (!AudioBusKeys_AudioThread.Contains(InAudioBusKey))
 		{
 			return;
@@ -1208,19 +1195,17 @@ namespace Audio
 					AudioBuses.Remove(InAudioBusKey);
 				}
 			}
-		}, AUDIO_MIXER_THREAD_COMMAND_STRING("StopAudioBus"));
+		});
 	}
 
 	bool FMixerSourceManager::IsAudioBusActive(FAudioBusKey InAudioBusKey) const
 	{
-		AUDIO_MIXER_THREAD_DEBUG_UPDATE()
 		AUDIO_MIXER_CHECK_GAME_THREAD(MixerDevice);
 		return AudioBusKeys_AudioThread.Contains(InAudioBusKey);
 	}
 
 	int32 FMixerSourceManager::GetAudioBusNumChannels(FAudioBusKey InAudioBusKey) const
 	{
-		AUDIO_MIXER_THREAD_DEBUG_UPDATE()
 		AUDIO_MIXER_CHECK_AUDIO_PLAT_THREAD(MixerDevice);
 		TSharedPtr<FMixerAudioBus> AudioBusPtr = AudioBuses.FindRef(InAudioBusKey);
 		if (AudioBusPtr.IsValid())
@@ -1233,7 +1218,6 @@ namespace Audio
 
 	void FMixerSourceManager::AddPatchOutputForAudioBus(FAudioBusKey InAudioBusKey, const FPatchOutputStrongPtr& InPatchOutputStrongPtr)
 	{
-		AUDIO_MIXER_THREAD_DEBUG_UPDATE()
 		AUDIO_MIXER_CHECK_AUDIO_PLAT_THREAD(MixerDevice);
 		if (MixerDevice->IsAudioRenderingThread())
 		{
@@ -1255,11 +1239,10 @@ namespace Audio
 
 	void FMixerSourceManager::AddPatchOutputForAudioBus_AudioThread(FAudioBusKey InAudioBusKey, const FPatchOutputStrongPtr& InPatchOutputStrongPtr)
 	{
-		AUDIO_MIXER_THREAD_DEBUG_UPDATE()
 		AudioMixerThreadCommand([this, InAudioBusKey, NewPatchPtr = InPatchOutputStrongPtr]() mutable
 		{
 			AddPatchOutputForAudioBus(InAudioBusKey, NewPatchPtr);
-		}, AUDIO_MIXER_THREAD_COMMAND_STRING("AddPatchOutputForAudioBus_AudioThread()"));
+		});
 	}
 
 	void FMixerSourceManager::AddPatchInputForAudioBus(FAudioBusKey InAudioBusKey, const FPatchInput& InPatchInput)
@@ -1288,12 +1271,11 @@ namespace Audio
 		AudioMixerThreadCommand([this, InAudioBusKey, InPatchInput]()
 		{
 			AddPatchInputForAudioBus(InAudioBusKey, InPatchInput);
-		}, AUDIO_MIXER_THREAD_COMMAND_STRING("AddPatchInputForAudioBus_AudioThread()"));
+		});
 	}
 
 	void FMixerSourceManager::Play(const int32 SourceId)
 	{
-		AUDIO_MIXER_THREAD_DEBUG_UPDATE()
 		AUDIO_MIXER_CHECK(SourceId < NumTotalSources);
 		AUDIO_MIXER_CHECK(GameThreadInfo.bIsBusy[SourceId]);
 		AUDIO_MIXER_CHECK_GAME_THREAD(MixerDevice);
@@ -1314,12 +1296,11 @@ namespace Audio
 			SourceInfo.StartTime = StartTime;
 
 			AUDIO_MIXER_DEBUG_LOG(SourceId, TEXT("Is playing"));
-		}, AUDIO_MIXER_THREAD_COMMAND_STRING("Play()"));
+		});
 	}
 
 	void FMixerSourceManager::CancelQuantizedSound(const int32 SourceId)
 	{
-		AUDIO_MIXER_THREAD_DEBUG_UPDATE()
 		if (!MixerDevice)
 		{
 			return;
@@ -1353,7 +1334,6 @@ namespace Audio
 
 	void FMixerSourceManager::Stop(const int32 SourceId)
 	{
-		AUDIO_MIXER_THREAD_DEBUG_UPDATE()
 		if (!MixerDevice)
 		{
 			return;
@@ -1369,12 +1349,11 @@ namespace Audio
 		AudioMixerThreadCommand([this, SourceId]()
 		{
 			StopInternal(SourceId);
-		}, AUDIO_MIXER_THREAD_COMMAND_STRING("Stop()"));
+		});
 	}
 
 	void FMixerSourceManager::StopInternal(const int32 SourceId)
 	{
-		AUDIO_MIXER_THREAD_DEBUG_UPDATE()
 		AUDIO_MIXER_CHECK_AUDIO_PLAT_THREAD(MixerDevice);
 
 		FSourceInfo& SourceInfo = SourceInfos[SourceId];
@@ -1396,7 +1375,6 @@ namespace Audio
 
 	void FMixerSourceManager::StopFade(const int32 SourceId, const int32 NumFrames)
 	{
-		AUDIO_MIXER_THREAD_DEBUG_UPDATE()
 		AUDIO_MIXER_CHECK(SourceId < NumTotalSources);
 		AUDIO_MIXER_CHECK(GameThreadInfo.bIsBusy[SourceId]);
 		AUDIO_MIXER_CHECK(NumFrames > 0);
@@ -1438,12 +1416,11 @@ namespace Audio
 			}
 
 			AUDIO_MIXER_DEBUG_LOG(SourceId, TEXT("Is stopping with fade"));
-		}, AUDIO_MIXER_THREAD_COMMAND_STRING("StopFade()"));
+		});
 	}
 	
 	void FMixerSourceManager::Pause(const int32 SourceId)
 	{
-		AUDIO_MIXER_THREAD_DEBUG_UPDATE()
 		AUDIO_MIXER_CHECK(SourceId < NumTotalSources);
 		AUDIO_MIXER_CHECK(GameThreadInfo.bIsBusy[SourceId]);
 		AUDIO_MIXER_CHECK_GAME_THREAD(MixerDevice);
@@ -1456,12 +1433,11 @@ namespace Audio
 
 			SourceInfo.bIsPaused = true;
 			SourceInfo.bIsActive = false;
-		}, AUDIO_MIXER_THREAD_COMMAND_STRING("Pause()"));
+		});
 	}
 
 	void FMixerSourceManager::SetPitch(const int32 SourceId, const float Pitch)
 	{
-		AUDIO_MIXER_THREAD_DEBUG_UPDATE()
 		AUDIO_MIXER_CHECK(SourceId < NumTotalSources);
 		AUDIO_MIXER_CHECK(GameThreadInfo.bIsBusy[SourceId]);
 
@@ -1471,12 +1447,11 @@ namespace Audio
 			check(NumOutputFrames > 0);
 
 			SourceInfos[SourceId].PitchSourceParam.SetValue(Pitch, NumOutputFrames);
-		}, AUDIO_MIXER_THREAD_COMMAND_STRING("SetPitch()"));
+		});
 	}
 
 	void FMixerSourceManager::SetVolume(const int32 SourceId, const float Volume)
 	{
-		AUDIO_MIXER_THREAD_DEBUG_UPDATE()
 		AUDIO_MIXER_CHECK(SourceId < NumTotalSources);
 		AUDIO_MIXER_CHECK(GameThreadInfo.bIsBusy[SourceId]);
 		AUDIO_MIXER_CHECK_GAME_THREAD(MixerDevice);
@@ -1499,12 +1474,11 @@ namespace Audio
 
 				SourceInfos[SourceId].VolumeSourceDestination = Volume;
 			}
-		}, AUDIO_MIXER_THREAD_COMMAND_STRING("SetVolume()"));
+		});
 	}
 
 	void FMixerSourceManager::SetDistanceAttenuation(const int32 SourceId, const float DistanceAttenuation)
 	{
-		AUDIO_MIXER_THREAD_DEBUG_UPDATE()
 		AUDIO_MIXER_CHECK(SourceId < NumTotalSources);
 		AUDIO_MIXER_CHECK(GameThreadInfo.bIsBusy[SourceId]);
 		AUDIO_MIXER_CHECK_GAME_THREAD(MixerDevice);
@@ -1521,12 +1495,11 @@ namespace Audio
 			}
 
 			SourceInfos[SourceId].DistanceAttenuationSourceDestination = DistanceAttenuation;
-		}, AUDIO_MIXER_THREAD_COMMAND_STRING("SetDistanceAttenuation()"));
+		});
 	}
 
 	void FMixerSourceManager::SetSpatializationParams(const int32 SourceId, const FSpatializationParams& InParams)
 	{
-		AUDIO_MIXER_THREAD_DEBUG_UPDATE()
 		AUDIO_MIXER_CHECK(SourceId < NumTotalSources);
 		AUDIO_MIXER_CHECK(GameThreadInfo.bIsBusy[SourceId]);
 		AUDIO_MIXER_CHECK_GAME_THREAD(MixerDevice);
@@ -1536,12 +1509,11 @@ namespace Audio
 			AUDIO_MIXER_CHECK_AUDIO_PLAT_THREAD(MixerDevice);
 
 			SourceInfos[SourceId].SpatParams = InParams;
-		}, AUDIO_MIXER_THREAD_COMMAND_STRING("SetSpatializationParams()"));
+		});
 	}
 
 	void FMixerSourceManager::SetChannelMap(const int32 SourceId, const uint32 NumInputChannels, const Audio::FAlignedFloatBuffer& ChannelMap, const bool bInIs3D, const bool bInIsCenterChannelOnly)
 	{
-		AUDIO_MIXER_THREAD_DEBUG_UPDATE()
 		AUDIO_MIXER_CHECK(SourceId < NumTotalSources);
 		AUDIO_MIXER_CHECK(GameThreadInfo.bIsBusy[SourceId]);
 		AUDIO_MIXER_CHECK_GAME_THREAD(MixerDevice);
@@ -1568,12 +1540,11 @@ namespace Audio
 
 			bool bNeedsSpeakerMap = SourceSubmixOutput.SetChannelMap(ChannelMap, bInIsCenterChannelOnly);
 			GameThreadInfo.bNeedsSpeakerMap[SourceId] = bNeedsSpeakerMap;
-		}, AUDIO_MIXER_THREAD_COMMAND_STRING("SetChannelMap()"));
+		});
 	}
 
 	void FMixerSourceManager::SetLPFFrequency(const int32 SourceId, const float InLPFFrequency)
 	{
-		AUDIO_MIXER_THREAD_DEBUG_UPDATE()
 		AUDIO_MIXER_CHECK(SourceId < NumTotalSources);
 		AUDIO_MIXER_CHECK(GameThreadInfo.bIsBusy[SourceId]);
 		AUDIO_MIXER_CHECK_GAME_THREAD(MixerDevice);
@@ -1589,12 +1560,11 @@ namespace Audio
 			// On callback, the client version may be overridden with the modulation LPF value depending on which is more aggressive.  
 			SourceInfo.LowPassFreq = InLPFFrequency;
 			SourceInfo.LowPassFilter.StartFrequencyInterpolation(InLPFFrequency, NumOutputFrames);
-		}, AUDIO_MIXER_THREAD_COMMAND_STRING("SetLPFFrequency()"));
+		});
 	}
 
 	void FMixerSourceManager::SetHPFFrequency(const int32 SourceId, const float InHPFFrequency)
 	{
-		AUDIO_MIXER_THREAD_DEBUG_UPDATE()
 		AUDIO_MIXER_CHECK(SourceId < NumTotalSources);
 		AUDIO_MIXER_CHECK(GameThreadInfo.bIsBusy[SourceId]);
 		AUDIO_MIXER_CHECK_GAME_THREAD(MixerDevice);
@@ -1609,12 +1579,11 @@ namespace Audio
 			// On callback, the client version may be overridden with the modulation HPF value depending on which is more aggressive.  
 			SourceInfo.HighPassFreq = InHPFFrequency;
 			SourceInfo.HighPassFilter.StartFrequencyInterpolation(InHPFFrequency, NumOutputFrames);
-		}, AUDIO_MIXER_THREAD_COMMAND_STRING("SetHPFFrequency()"));
+		});
 	}
 
 	void FMixerSourceManager::SetModLPFFrequency(const int32 SourceId, const float InLPFFrequency)
 	{
-		AUDIO_MIXER_THREAD_DEBUG_UPDATE()
 		AUDIO_MIXER_CHECK(SourceId < NumTotalSources);
 		AUDIO_MIXER_CHECK(GameThreadInfo.bIsBusy[SourceId]);
 		AUDIO_MIXER_CHECK_GAME_THREAD(MixerDevice);
@@ -1626,12 +1595,11 @@ namespace Audio
 			FSourceInfo& SourceInfo = SourceInfos[SourceId];
 			SourceInfo.LowpassModulationBase = InLPFFrequency;
 			SourceInfo.bModFiltersUpdated = true;
-		}, AUDIO_MIXER_THREAD_COMMAND_STRING("SetModLPFFrequency()"));
+		});
 	}
 
 	void FMixerSourceManager::SetModHPFFrequency(const int32 SourceId, const float InHPFFrequency)
 	{
-		AUDIO_MIXER_THREAD_DEBUG_UPDATE()
 		AUDIO_MIXER_CHECK(SourceId < NumTotalSources);
 		AUDIO_MIXER_CHECK(GameThreadInfo.bIsBusy[SourceId]);
 		AUDIO_MIXER_CHECK_GAME_THREAD(MixerDevice);
@@ -1643,7 +1611,7 @@ namespace Audio
 			FSourceInfo& SourceInfo = SourceInfos[SourceId];
 			SourceInfo.HighpassModulationBase = InHPFFrequency;
 			SourceInfo.bModFiltersUpdated = true;
-		}, AUDIO_MIXER_THREAD_COMMAND_STRING("SetModHPFFrequency()"));
+		});
 	}
 
 	void FMixerSourceManager::SetModulationRouting(const int32 SourceId, FSoundModulationDefaultSettings& ModulationSettings)
@@ -1687,7 +1655,6 @@ namespace Audio
 
 	void FMixerSourceManager::SetSourceBufferListener(const int32 SourceId, FSharedISourceBufferListenerPtr& InSourceBufferListener, bool InShouldSourceBufferListenerZeroBuffer)
 	{
-		AUDIO_MIXER_THREAD_DEBUG_UPDATE()
 		AUDIO_MIXER_CHECK(SourceId < NumTotalSources);
 		AUDIO_MIXER_CHECK(GameThreadInfo.bIsBusy[SourceId]);
 		AUDIO_MIXER_CHECK_GAME_THREAD(MixerDevice);
@@ -1704,7 +1671,6 @@ namespace Audio
 
 	void FMixerSourceManager::SetModVolume(const int32 SourceId, const float InModVolume)
 	{
-		AUDIO_MIXER_THREAD_DEBUG_UPDATE()
 		AUDIO_MIXER_CHECK(SourceId < NumTotalSources);
 		AUDIO_MIXER_CHECK(GameThreadInfo.bIsBusy[SourceId]);
 		AUDIO_MIXER_CHECK_GAME_THREAD(MixerDevice);
@@ -1715,12 +1681,11 @@ namespace Audio
 
 			FSourceInfo& SourceInfo = SourceInfos[SourceId];
 			SourceInfo.VolumeModulationBase = InModVolume;
-		}, AUDIO_MIXER_THREAD_COMMAND_STRING("SetModVolume()"));
+		});
 	}
 
 	void FMixerSourceManager::SetModPitch(const int32 SourceId, const float InModPitch)
 	{
-		AUDIO_MIXER_THREAD_DEBUG_UPDATE()
 		AUDIO_MIXER_CHECK(SourceId < NumTotalSources);
 		AUDIO_MIXER_CHECK(GameThreadInfo.bIsBusy[SourceId]);
 		AUDIO_MIXER_CHECK_GAME_THREAD(MixerDevice);
@@ -1731,12 +1696,11 @@ namespace Audio
 
 			FSourceInfo& SourceInfo = SourceInfos[SourceId];
 			SourceInfo.PitchModulationBase = InModPitch;
-		}, AUDIO_MIXER_THREAD_COMMAND_STRING("SetModPitch()"));
+		});
 	}
 
 	void FMixerSourceManager::SetSubmixSendInfo(const int32 SourceId, const FMixerSourceSubmixSend& InSubmixSend)
 	{
-		AUDIO_MIXER_THREAD_DEBUG_UPDATE()
 		AUDIO_MIXER_CHECK(SourceId < NumTotalSources);
 		AUDIO_MIXER_CHECK(GameThreadInfo.bIsBusy[SourceId]);
 		AUDIO_MIXER_CHECK_GAME_THREAD(MixerDevice);
@@ -1791,12 +1755,11 @@ namespace Audio
 					InSubmixPtr->AddOrSetSourceVoice(SourceVoice, InSubmixSend.SendLevel, InSubmixSend.SubmixSendStage);
 				}
 			}
-		}, AUDIO_MIXER_THREAD_COMMAND_STRING("SetSubmixSendInfo()"));
+		});
 	}
 
 	void FMixerSourceManager::ClearSubmixSendInfo(const int32 SourceId, const FMixerSourceSubmixSend& InSubmixSend)
 	{
-		AUDIO_MIXER_THREAD_DEBUG_UPDATE()
 		AUDIO_MIXER_CHECK(SourceId < NumTotalSources);
 		AUDIO_MIXER_CHECK(GameThreadInfo.bIsBusy[SourceId]);
 		AUDIO_MIXER_CHECK_GAME_THREAD(MixerDevice);
@@ -1840,12 +1803,11 @@ namespace Audio
 				// Now remove the source voice from the submix send list
 				InSubmixPtr->RemoveSourceVoice(MixerSources[SourceId]);
 			}
-		}, AUDIO_MIXER_THREAD_COMMAND_STRING("ClearSubmixSendInfo()"));
+		});
 	}
 
 	void FMixerSourceManager::SetBusSendInfo(const int32 SourceId, EBusSendType InAudioBusSendType, uint32 AudioBusId, float BusSendLevel)
 	{
-		AUDIO_MIXER_THREAD_DEBUG_UPDATE()
 		AUDIO_MIXER_CHECK(SourceId < NumTotalSources);
 		AUDIO_MIXER_CHECK(GameThreadInfo.bIsBusy[SourceId]);
 		AUDIO_MIXER_CHECK_GAME_THREAD(MixerDevice);
@@ -1896,16 +1858,15 @@ namespace Audio
 			{
 				SourceInfo.AudioBusSends[(int32)InAudioBusSendType].Add(AudioBusId);
 			}
-		}, AUDIO_MIXER_THREAD_COMMAND_STRING("SetBusSendInfo()"));
+		});
 	}
 
 	void FMixerSourceManager::SetListenerTransforms(const TArray<FTransform>& InListenerTransforms)
 	{
-		AUDIO_MIXER_THREAD_DEBUG_UPDATE()
 		AudioMixerThreadCommand([this, InListenerTransforms]()
 		{
 			ListenerTransforms = InListenerTransforms;
-		}, AUDIO_MIXER_THREAD_COMMAND_STRING("SetListenerTransforms()"));
+		});
 	}
 
 	const TArray<FTransform>* FMixerSourceManager::GetListenerTransforms() const
@@ -1940,7 +1901,6 @@ namespace Audio
 
 	void FMixerSourceManager::ReadSourceFrame(const int32 SourceId)
 	{
-		AUDIO_MIXER_THREAD_DEBUG_UPDATE()
 		FSourceInfo& SourceInfo = SourceInfos[SourceId];
 
 		const int32 NumChannels = SourceInfo.NumInputChannels;
@@ -2058,7 +2018,6 @@ namespace Audio
 
 	void FMixerSourceManager::ComputeSourceBuffersForIdRange(const bool bGenerateBuses, const int32 SourceIdStart, const int32 SourceIdEnd)
 	{
-		AUDIO_MIXER_THREAD_DEBUG_UPDATE()
 		CSV_SCOPED_TIMING_STAT(Audio, SourceBuffers);
 		CONDITIONAL_SCOPE_CYCLE_COUNTER(STAT_AudioMixerSourceBuffers, (SourceIdStart < SourceIdEnd));
 
@@ -2299,7 +2258,6 @@ namespace Audio
 
 	void FMixerSourceManager::ComputeBuses()
 	{
-		AUDIO_MIXER_THREAD_DEBUG_UPDATE()
 		// Loop through the bus registry and mix source audio
 		for (auto& Entry : AudioBuses)
 		{
@@ -2310,7 +2268,6 @@ namespace Audio
 
 	void FMixerSourceManager::UpdateBuses()
 	{
-		AUDIO_MIXER_THREAD_DEBUG_UPDATE()
 		// Update the bus states post mixing. This flips the current/previous buffer indices.
 		for (auto& Entry : AudioBuses)
 		{
@@ -2320,14 +2277,9 @@ namespace Audio
 	}
 
 	// ctor
-	FMixerSourceManager::FAudioMixerThreadCommand::FAudioMixerThreadCommand(TFunction<void()> InFunction, FName InCallerDebugInfo, bool bInDeferExecution)
-#if WITH_AUDIO_MIXER_THREAD_COMMAND_DEBUG
-		: CallerDebugInfo(InCallerDebugInfo)
-		, Function(MoveTemp(InFunction))
-		, bDeferExecution(bInDeferExecution)
-#else
+	FMixerSourceManager::FAudioMixerThreadCommand::FAudioMixerThreadCommand(TFunction<void()> InFunction, bool bInDeferExecution)
 		: Function(MoveTemp(InFunction))
-#endif // #if WITH_AUDIO_MIXER_THREAD_COMMAND_DEBUG
+		, bDeferExecution(bInDeferExecution)
 	{
 	}
 
@@ -2350,7 +2302,6 @@ namespace Audio
 
 	void FMixerSourceManager::ComputePluginAudio(FSourceInfo& SourceInfo, FMixerSourceSubmixOutputBuffer& InSourceSubmixOutputBuffer, int32 SourceId, int32 NumSamples)
 	{
-		AUDIO_MIXER_THREAD_DEBUG_UPDATE()
 		if (BypassAudioPluginsCvar)
 		{
 			// If we're bypassing audio plugins, our pre- and post-effect channels are the same as the input channels
@@ -2542,7 +2493,6 @@ namespace Audio
 	{
 		CSV_SCOPED_TIMING_STAT(Audio, SourceEffectsBuffers);
 		CONDITIONAL_SCOPE_CYCLE_COUNTER(STAT_AudioMixerSourceEffectBuffers, (SourceIdStart < SourceIdEnd));
-		AUDIO_MIXER_THREAD_DEBUG_UPDATE()
 
 		const bool bIsDebugModeEnabled = DebugSoloSources.Num() > 0;
 
@@ -2847,7 +2797,6 @@ namespace Audio
 	{
 		CSV_SCOPED_TIMING_STAT(Audio, SourceOutputBuffers);
 		CONDITIONAL_SCOPE_CYCLE_COUNTER(STAT_AudioMixerSourceOutputBuffers, (SourceIdStart < SourceIdEnd));
-		AUDIO_MIXER_THREAD_DEBUG_UPDATE()
 
 		for (int32 SourceId = SourceIdStart; SourceId < SourceIdEnd; ++SourceId)
 		{
@@ -2876,7 +2825,6 @@ namespace Audio
 
 	void FMixerSourceManager::GenerateSourceAudio(const bool bGenerateBuses, const int32 SourceIdStart, const int32 SourceIdEnd)
 	{
-		AUDIO_MIXER_THREAD_DEBUG_UPDATE()
 		// Buses generate their input buffers independently
 		// Get the next block of frames from the source buffers
 		ComputeSourceBuffersForIdRange(bGenerateBuses, SourceIdStart, SourceIdEnd);
@@ -2890,7 +2838,6 @@ namespace Audio
 
 	void FMixerSourceManager::GenerateSourceAudio(const bool bGenerateBuses)
 	{
-		AUDIO_MIXER_THREAD_DEBUG_UPDATE()
 		// If there are no buses, don't need to do anything here
 		if (bGenerateBuses && !AudioBuses.Num())
 		{
@@ -2921,7 +2868,6 @@ namespace Audio
 
 	void FMixerSourceManager::MixOutputBuffers(const int32 SourceId, int32 InNumOutputChannels, const float InSendLevel, EMixerSourceSubmixSendStage InSubmixSendStage, FAlignedFloatBuffer& OutWetBuffer) const
 	{
-		AUDIO_MIXER_THREAD_DEBUG_UPDATE()
 		if (InSendLevel > 0.0f)
 		{
 			const FSourceInfo& SourceInfo = SourceInfos[SourceId];
@@ -2945,7 +2891,6 @@ namespace Audio
 
 	const ISoundfieldAudioPacket* FMixerSourceManager::GetEncodedOutput(const int32 SourceId, const FSoundfieldEncodingKey& InKey) const
 	{
-		AUDIO_MIXER_THREAD_DEBUG_UPDATE()
 		AUDIO_MIXER_CHECK_AUDIO_PLAT_THREAD(MixerDevice);
 
 		const FSourceInfo& SourceInfo = SourceInfos[SourceId];
@@ -2968,7 +2913,6 @@ namespace Audio
 
 	void FMixerSourceManager::UpdateDeviceChannelCount(const int32 InNumOutputChannels)
 	{
-		AUDIO_MIXER_THREAD_DEBUG_UPDATE()
 		AudioMixerThreadCommand([this, InNumOutputChannels]()
 		{
 			NumOutputSamples = NumOutputFrames * MixerDevice->GetNumDeviceChannels();
@@ -3006,12 +2950,11 @@ namespace Audio
 
 				SourceSubmixOutputBuffer.SetChannelMap(SourceInfo.ScratchChannelMap, SourceInfo.bIsCenterChannelOnly);
 			}
-		}, AUDIO_MIXER_THREAD_COMMAND_STRING("UpdateDeviceChannelCount()"));
+		});
 	}
 
 	void FMixerSourceManager::UpdateSourceEffectChain(const uint32 InSourceEffectChainId, const TArray<FSourceEffectChainEntry>& InSourceEffectChain, const bool bPlayEffectChainTails)
 	{
-		AUDIO_MIXER_THREAD_DEBUG_UPDATE()
 		AudioMixerThreadCommand([this, InSourceEffectChainId, InSourceEffectChain, bPlayEffectChainTails]()
 		{
 			FSoundEffectSourceInitData InitData;
@@ -3072,7 +3015,7 @@ namespace Audio
 					}
 				}
 			}
-		}, AUDIO_MIXER_THREAD_COMMAND_STRING("UpdateSourceEffectChain()"), /*bDeferExecution*/true);
+		}, /*bDeferExecution*/true);
 	}
 
 	void FMixerSourceManager::PauseSoundForQuantizationCommand(const int32 SourceId)
@@ -3267,66 +3210,11 @@ namespace Audio
 		MpscCommandQueue.Enqueue(MoveTemp(InCommand));
 	}
 
-	void FMixerSourceManager::FCurrentAudioMixerThreadCommandInfo::SetNewTask(FName InTask)
-	{
-#if WITH_AUDIO_MIXER_THREAD_COMMAND_DEBUG
-		CallerDebugInfo = InTask;
-		QueueTimeInCycles = FPlatformTime::Cycles64();
-#endif // #if WITH_AUDIO_MIXER_THREAD_COMMAND_DEBUG
-	}
 
-	void FMixerSourceManager::FCurrentAudioMixerThreadCommandInfo::LogWarning() const
-	{
-#if WITH_AUDIO_MIXER_THREAD_COMMAND_DEBUG
-		if(!CallerDebugInfo.IsNone())
-		{
-			UE_LOG(LogAudioMixer, Warning, TEXT("The AudioThreadCommand currently executing was called by %s, and started execution %f milliseconds ago")
-				, *CallerDebugInfo.ToString(), FPlatformTime::ToMilliseconds64(FPlatformTime::Cycles64() - QueueTimeInCycles.GetValue()));
-		}
-		else
-		{
-			UE_LOG(LogAudioMixer, Warning, TEXT("We don't currently have a task named"));
-		}
-#endif // #if WITH_AUDIO_MIXER_THREAD_COMMAND_DEBUG
-	}
-
-	void FMixerSourceManager::FCurrentAudioMixerThreadCommandInfo::LogError() const
-	{
-#if WITH_AUDIO_MIXER_THREAD_COMMAND_DEBUG
-		if(!CallerDebugInfo.IsNone())
-		{
-			UE_LOG(LogAudioMixer, Error, TEXT("The AudioThreadCommand currently executing was called by %s, and started execution %f milliseconds ago")
-				, *CallerDebugInfo.ToString(), FPlatformTime::ToMilliseconds64(FPlatformTime::Cycles64() - QueueTimeInCycles.GetValue()));
-		}
-		else
-		{
-			UE_LOG(LogAudioMixer, Warning, TEXT("We don't currently have a task named"));
-		}
-#endif // #if WITH_AUDIO_MIXER_THREAD_COMMAND_DEBUG
-	}
 
-	void FMixerSourceManager::FCurrentAudioMixerThreadCommandInfo::Reset()
+	void FMixerSourceManager::AudioMixerThreadCommand(TFunction<void()> InFunction, bool bDeferExecution)
 	{
-#if WITH_AUDIO_MIXER_THREAD_COMMAND_DEBUG
-		CallerDebugInfo = FName();
-#endif // #if WITH_AUDIO_MIXER_THREAD_COMMAND_DEBUG
-		QueueTimeInCycles.Reset();
-	}
-
-	FMixerSourceManager::FScopedAudioThreadCommandInfo::FScopedAudioThreadCommandInfo(FName InTaskName, FCurrentAudioMixerThreadCommandInfo& InInfoRef)
-		: InfoRef(InInfoRef)
-	{
-		InfoRef.SetNewTask(InTaskName);
-	}
-
-	FMixerSourceManager::FScopedAudioThreadCommandInfo::~FScopedAudioThreadCommandInfo()
-	{
-		InfoRef.Reset();
-	}
-
-	void FMixerSourceManager::AudioMixerThreadCommand(TFunction<void()> InFunction, FName CallingFunction, bool bDeferExecution)
-	{
-		FAudioMixerThreadCommand AudioCommand(InFunction, CallingFunction, bDeferExecution);
+		FAudioMixerThreadCommand AudioCommand(InFunction, bDeferExecution);
 
 		// Here, we make sure that we don't flip our command double buffer while we are executing this function.
 		FScopeLock ScopeLock(&CommandBufferIndexCriticalSection);
@@ -3340,9 +3228,7 @@ namespace Audio
 
 		static SIZE_T WarnSize = 1024 * 1024;
 		if (CurrentBufferSizeInBytes > WarnSize )
-		{		
-			CurrentAudioMixerThreadCommandInfo.LogError();
-
+		{
 			SIZE_T Num = CommandBuffers[AudioThreadCommandIndex].SourceCommandQueue.Num();
 			UE_LOG(LogAudioMixer, Error, TEXT("Command Queue has grown to %ukb, containing %d cmds, last pump was %.2fms ago."),
 				CurrentBufferSizeInBytes >> 10, Num, FPlatformTime::ToMilliseconds64(FPlatformTime::Cycles64() - LastPumpTimeInCycles));
@@ -3360,7 +3246,6 @@ namespace Audio
 		}
 		else
 		{
-			CurrentAudioMixerThreadCommandInfo.LogError();
 			UE_LOG(LogAudioMixer, Error,  TEXT("Command buffer grown to %umb, preventing any more adds! Likely cause, the h/w has stopped consuming data."), CurrentBufferSizeInBytes >>20);
 		}
 	}
@@ -3368,7 +3253,6 @@ namespace Audio
 	void FMixerSourceManager::PumpCommandQueue()
 	{
 		TRACE_CPUPROFILER_EVENT_SCOPE(AudioMixerThreadCommands::PumpCommandQueue)
-		AUDIO_MIXER_THREAD_DEBUG_UPDATE()
 		
 		// If we're already triggered, we need to wait for the audio thread to reset it before pumping
 		if (FPlatformProcess::SupportsMultithreading())
@@ -3402,10 +3286,7 @@ namespace Audio
 		for (int32 Id = 0; Id < NumCommandsToExecute; ++Id)
 		{
 			FAudioMixerThreadCommand& AudioCommand = Commands.SourceCommandQueue[Id];
-#if WITH_AUDIO_MIXER_THREAD_COMMAND_DEBUG
-			// cache debug info about the command we are about to execute (for safe access from the producer thread)
-			CurrentAudioMixerThreadCommandInfo.SetNewTask(AudioCommand.CallerDebugInfo);
-#endif // #if WITH_AUDIO_MIXER_THREAD_COMMAND_DEBUG
+
 			if(AudioCommand.bDeferExecution)
 			{
 				AudioCommand.bDeferExecution = false;
@@ -3416,16 +3297,6 @@ namespace Audio
 				AudioCommand(); // execute
 			}
 
-#if WITH_AUDIO_MIXER_THREAD_COMMAND_DEBUG
-			const float ExecutionTimeInMs = CurrentAudioMixerThreadCommandInfo.GetCurrentExecutionTimeMs();
-			if(LogEveryAudioThreadCommandCvar || (ExecutionTimeInMs >= AudioCommandExecTimeMsWarningThresholdCvar))
-			{
-				UE_LOG(LogAudioMixer, Warning, TEXT("Audio Thread Command from %s() took %f milleseconds to execute")
-				, *CurrentAudioMixerThreadCommandInfo.GetCurrentTaskName().ToString(), ExecutionTimeInMs);
-			}
-#endif // #if WITH_AUDIO_MIXER_THREAD_COMMAND_DEBUG
-
-			CurrentAudioMixerThreadCommandInfo.Reset();
 			NumCommands.Decrement();
 		}
 
@@ -3446,7 +3317,6 @@ namespace Audio
 
 	void FMixerSourceManager::FlushCommandQueue(bool bPumpInCommand)
 	{
-		AUDIO_MIXER_THREAD_DEBUG_UPDATE()
 		check(CommandsProcessedEvent != nullptr);
 
 		// If we have no commands enqueued, exit
@@ -3491,7 +3361,6 @@ namespace Audio
 
 	void FMixerSourceManager::UpdatePendingReleaseData(bool bForceWait)
 	{
-		AUDIO_MIXER_THREAD_DEBUG_UPDATE()
 		
 		// Don't block, but let tasks finish naturally
 		for (int32 i = PendingSourceBuffers.Num() - 1; i >= 0; --i)
diff --git a/Engine/Source/Runtime/AudioMixer/Private/AudioMixerSourceManager.h b/Engine/Source/Runtime/AudioMixer/Private/AudioMixerSourceManager.h
--- a/Engine/Source/Runtime/AudioMixer/Private/AudioMixerSourceManager.h
+++ b/Engine/Source/Runtime/AudioMixer/Private/AudioMixerSourceManager.h
@@ -19,19 +19,6 @@
 #include "Sound/QuartzQuantizationUtilities.h"
 #include "Stats/Stats.h"
 
-// defines for debug data/logging in AudioMixerThreadCommand execution
-// if this hasn't been explicitly defined, assume it is enabled
-#ifndef WITH_AUDIO_MIXER_THREAD_COMMAND_DEBUG
-#define WITH_AUDIO_MIXER_THREAD_COMMAND_DEBUG 1
-#endif // #ifndef WITH_AUDIO_MIXER_THREAD_COMMAND_DEBUG
-
-#if WITH_AUDIO_MIXER_THREAD_COMMAND_DEBUG
-#define AUDIO_MIXER_THREAD_COMMAND_STRING(X) ( FName(X) )
-#else //WITH_AUDIO_MIXER_THREAD_COMMAND_DEBUG
-static const FName NAME_NONE;
-#define AUDIO_MIXER_THREAD_COMMAND_STRING(X) ( NAME_NONE )
-#endif //WITH_AUDIO_MIXER_THREAD_COMMAND_DEBUG
-
 
 // Tracks the time it takes to up the source manager (computes source buffers, source effects, sample rate conversion)
 DECLARE_CYCLE_STAT_EXTERN(TEXT("Source Manager Update"), STAT_AudioMixerSourceManagerUpdate, STATGROUP_AudioMixer, AUDIOMIXER_API);
@@ -310,15 +297,11 @@ namespace Audio
 		struct FAudioMixerThreadCommand
 		{
 			// ctor
-			FAudioMixerThreadCommand(TFunction<void()> InFunction, FName InCallerDebugInfo, bool bInDeferExecution = false);
+			FAudioMixerThreadCommand(TFunction<void()> InFunction, bool bInDeferExecution = false);
 
 			// function-call operator
 			void operator()() const;
 
-			// data
-#if WITH_AUDIO_MIXER_THREAD_COMMAND_DEBUG
-			FName CallerDebugInfo;
-#endif // #if WITH_AUDIO_MIXER_THREAD_COMMAND_DEBUG
 			TFunction<void()> Function;
 
 			// Defers the execution by a single call to PumpCommandQueue()
@@ -326,51 +309,8 @@ namespace Audio
 			// and that source gets initialized after the command executes
 			bool bDeferExecution;
 		};
-		
-
-		struct FCurrentAudioMixerThreadCommandInfo
-		{
-		private:
-			FName CallerDebugInfo;
-			FThreadSafeCounter64 QueueTimeInCycles;
-			
-		public:
-			void SetNewTask(FName InTask);
-			FName GetCurrentTaskName() const
-			{
-				return CallerDebugInfo;
-			}
-
-			float GetCurrentExecutionTimeMs() const
-			{
-				if(!CallerDebugInfo.IsNone())
-				{
-					return static_cast<float>(FPlatformTime::ToMilliseconds64(FPlatformTime::Cycles64() - QueueTimeInCycles.GetValue()));
-				}
-				return 0.f; // no task currently running
-			}
-
-			void LogWarning() const;
-			void LogError() const;
-			void Reset();
-		};
-
-		struct FScopedAudioThreadCommandInfo
-		{
-			// ctor - name the task, reset the counter
-			FScopedAudioThreadCommandInfo() = delete;
-			FScopedAudioThreadCommandInfo(FName InTaskName, FCurrentAudioMixerThreadCommandInfo& InInfoRef);
-
-			// dtor - reset the task name and counter (not in a function)
-			~FScopedAudioThreadCommandInfo();
-
-			FCurrentAudioMixerThreadCommandInfo& InfoRef;
-			
-		};
-
-		mutable FCurrentAudioMixerThreadCommandInfo CurrentAudioMixerThreadCommandInfo; // mutable to allow const functions to register debug info
 
-		void AudioMixerThreadCommand(TFunction<void()> InFunction, FName CallingFunction = {}, bool bInDeferExecution = false);
+		void AudioMixerThreadCommand(TFunction<void()> InFunction, bool bInDeferExecution = false);
 
 		
 		static const int32 NUM_BYTES_PER_SAMPLE = 2;
