CIS: fix up lingering use of deleted #define
#rnx

diff --git a/Engine/Source/Runtime/AudioMixer/Private/AudioMixerSourceManager.cpp b/Engine/Source/Runtime/AudioMixer/Private/AudioMixerSourceManager.cpp
--- a/Engine/Source/Runtime/AudioMixer/Private/AudioMixerSourceManager.cpp
+++ b/Engine/Source/Runtime/AudioMixer/Private/AudioMixerSourceManager.cpp
@@ -1650,7 +1650,7 @@ namespace Audio
 				SourceInfo.PitchModulation = MoveTemp(PitchModulation);
 				SourceInfo.LowpassModulation = MoveTemp(LowpassModulation);
 				SourceInfo.HighpassModulation = MoveTemp(HighpassModulation);
-			}, AUDIO_MIXER_THREAD_COMMAND_STRING("SetModulationRouting"));
+			});
 	}
 
 	void FMixerSourceManager::SetSourceBufferListener(const int32 SourceId, FSharedISourceBufferListenerPtr& InSourceBufferListener, bool InShouldSourceBufferListenerZeroBuffer)
