Move audio bus subsystem deinitialization later to prevent race condition on audio device destroy

diff --git a/Engine/Source/Runtime/AudioMixer/Private/Effects/AudioMixerSubmixEffectDynamicsProcessor.cpp b/Engine/Source/Runtime/AudioMixer/Private/Effects/AudioMixerSubmixEffectDynamicsProcessor.cpp
--- a/Engine/Source/Runtime/AudioMixer/Private/Effects/AudioMixerSubmixEffectDynamicsProcessor.cpp
+++ b/Engine/Source/Runtime/AudioMixer/Private/Effects/AudioMixerSubmixEffectDynamicsProcessor.cpp
@@ -200,8 +200,10 @@ bool FSubmixEffectDynamicsProcessor::UpdateKeySourcePatch()
 				if (ObjectId != INDEX_NONE)
 				{
 					UAudioBusSubsystem* AudioBusSubsystem = MixerDevice->GetSubsystem<UAudioBusSubsystem>();
-					check(AudioBusSubsystem);
-					KeySource.Patch = AudioBusSubsystem->AddPatchOutputForAudioBus(Audio::FAudioBusKey(ObjectId), MixerDevice->GetNumOutputFrames(), KeySource.GetNumChannels());
+					if (AudioBusSubsystem)
+					{
+						KeySource.Patch = AudioBusSubsystem->AddPatchOutputForAudioBus(Audio::FAudioBusKey(ObjectId), MixerDevice->GetNumOutputFrames(), KeySource.GetNumChannels());
+					}
 
 					if (KeySource.Patch.IsValid())
 					{
diff --git a/Engine/Source/Runtime/Engine/Private/AudioDevice.cpp b/Engine/Source/Runtime/Engine/Private/AudioDevice.cpp
--- a/Engine/Source/Runtime/Engine/Private/AudioDevice.cpp
+++ b/Engine/Source/Runtime/Engine/Private/AudioDevice.cpp
@@ -689,7 +689,6 @@ void FAudioDevice::OnDeviceDestroyed(Audio::FDeviceId InDeviceID)
 {
 	if (InDeviceID == DeviceID)
 	{
-		Deinitialize();
 		FAudioDeviceManagerDelegates::OnAudioDeviceDestroyed.Remove(DeviceDestroyedHandle);
 	}
 }
diff --git a/Engine/Source/Runtime/Engine/Private/AudioDeviceManager.cpp b/Engine/Source/Runtime/Engine/Private/AudioDeviceManager.cpp
--- a/Engine/Source/Runtime/Engine/Private/AudioDeviceManager.cpp
+++ b/Engine/Source/Runtime/Engine/Private/AudioDeviceManager.cpp
@@ -792,6 +792,9 @@ void FAudioDeviceManager::DecrementDevice(Audio::FDeviceId DeviceID, UWorld* InW
 			bDestroyingDevice = true;
 			FAudioDeviceManagerDelegates::OnAudioDeviceDestroyed.Broadcast(DeviceID);
 
+			// Subsystems deinitialization
+			Container.Device->Deinitialize();
+
 			// If this is the active device and being destroyed, set the main device as the active device.
 			if (DeviceID == ActiveAudioDeviceID)
 			{
