Add pass-through functions for Audio/Source Bus Pre/Post Send to SynthComponent

diff --git a/Engine/Source/Runtime/AudioMixer/Private/Components/SynthComponent.cpp b/Engine/Source/Runtime/AudioMixer/Private/Components/SynthComponent.cpp
--- a/Engine/Source/Runtime/AudioMixer/Private/Components/SynthComponent.cpp
+++ b/Engine/Source/Runtime/AudioMixer/Private/Components/SynthComponent.cpp
@@ -568,6 +568,38 @@ void USynthComponent::SetSubmixSend(USoundSubmixBase* Submix, float SendLevel)
 	}
 }
 
+void USynthComponent::SetSourceBusSendPreEffect(USoundSourceBus* SoundSourceBus, float SourceBusSendLevel)
+{
+	if (AudioComponent)
+	{
+		AudioComponent->SetSourceBusSendPreEffect(SoundSourceBus, SourceBusSendLevel);
+	}
+}
+
+void USynthComponent::SetSourceBusSendPostEffect(USoundSourceBus* SoundSourceBus, float SourceBusSendLevel)
+{
+	if (AudioComponent)
+	{
+		AudioComponent->SetSourceBusSendPostEffect(SoundSourceBus, SourceBusSendLevel);
+	}
+}
+
+void USynthComponent::SetAudioBusSendPreEffect(UAudioBus* AudioBus, float AudioBusSendLevel)
+{
+	if (AudioComponent)
+	{
+		AudioComponent->SetAudioBusSendPreEffect(AudioBus, AudioBusSendLevel);
+	}
+}
+
+void USynthComponent::SetAudioBusSendPostEffect(UAudioBus* AudioBus, float AudioBusSendLevel)
+{
+	if (AudioComponent)
+	{
+		AudioComponent->SetAudioBusSendPostEffect(AudioBus, AudioBusSendLevel);
+	}
+}
+
 void USynthComponent::SetLowPassFilterEnabled(bool InLowPassFilterEnabled)
 {
 	if (AudioComponent)
diff --git a/Engine/Source/Runtime/AudioMixer/Public/Components/SynthComponent.h b/Engine/Source/Runtime/AudioMixer/Public/Components/SynthComponent.h
--- a/Engine/Source/Runtime/AudioMixer/Public/Components/SynthComponent.h
+++ b/Engine/Source/Runtime/AudioMixer/Public/Components/SynthComponent.h
@@ -18,6 +18,9 @@
 #include "DSP/SinOsc.h"
 #endif
 
+class UAudioBus;
+class USoundSourceBus;
+
 /** Simple interface class to allow objects to route audio between them. */
 class IAudioBufferListener
 {
@@ -120,6 +123,22 @@ public:
 	UFUNCTION(BlueprintCallable, Category = "Audio|Components|Audio")
 	void SetSubmixSend(USoundSubmixBase* Submix, float SendLevel);
 
+	/** Sets how much audio the sound should send to the given SourceBus (pre effect). */
+	UFUNCTION(BlueprintCallable, Category = "Audio|Components|Audio")
+	void SetSourceBusSendPreEffect(USoundSourceBus* SoundSourceBus, float SourceBusSendLevel);
+
+	/** Sets how much audio the sound should send to the given SourceBus (post effect). */
+	UFUNCTION(BlueprintCallable, Category = "Audio|Components|Audio")
+	void SetSourceBusSendPostEffect(USoundSourceBus* SoundSourceBus, float SourceBusSendLevel);
+
+	/** Sets how much audio the sound should send to the given AudioBus (pre effect). */
+	UFUNCTION(BlueprintCallable, Category = "Audio|Components|Audio")
+	void SetAudioBusSendPreEffect(UAudioBus* AudioBus, float AudioBusSendLevel);
+
+	/** Sets how much audio the sound should send to the given AudioBus (post effect). */
+	UFUNCTION(BlueprintCallable, Category = "Audio|Components|Audio")
+	void SetAudioBusSendPostEffect(UAudioBus* AudioBus, float AudioBusSendLevel);
+
 	/** Sets whether or not the low pass filter is enabled on the audio component. */
 	UFUNCTION(BlueprintCallable, Category = "Audio|Components|Audio")
 	void SetLowPassFilterEnabled(bool InLowPassFilterEnabled);
