Start audio bus before adding patch to key source in Submix Effect Dynamics Processor, remove misleading warning

diff --git a/Engine/Plugins/Runtime/Synthesis/Source/Synthesis/Private/SubmixEffectMultiBandCompressor.cpp b/Engine/Plugins/Runtime/Synthesis/Source/Synthesis/Private/SubmixEffectMultiBandCompressor.cpp
--- a/Engine/Plugins/Runtime/Synthesis/Source/Synthesis/Private/SubmixEffectMultiBandCompressor.cpp
+++ b/Engine/Plugins/Runtime/Synthesis/Source/Synthesis/Private/SubmixEffectMultiBandCompressor.cpp
@@ -239,10 +239,6 @@ bool FSubmixEffectMultibandCompressor::UpdateKeySourcePatch()
 
 						return true;
 					}
-					else if (KeySource.ShouldReportInactive())
-					{
-						KeySource.SetReportInactive(false);
-					}
 				}
 			}
 		}
@@ -276,10 +272,6 @@ bool FSubmixEffectMultibandCompressor::UpdateKeySourcePatch()
 							return true;
 						}
 					}
-					else if (KeySource.ShouldReportInactive())
-					{
-						KeySource.SetReportInactive(false);
-					}
 				}
 			}
 		}
diff --git a/Engine/Source/Runtime/AudioMixer/Classes/SubmixEffects/AudioMixerSubmixEffectDynamicsProcessor.h b/Engine/Source/Runtime/AudioMixer/Classes/SubmixEffects/AudioMixerSubmixEffectDynamicsProcessor.h
--- a/Engine/Source/Runtime/AudioMixer/Classes/SubmixEffects/AudioMixerSubmixEffectDynamicsProcessor.h
+++ b/Engine/Source/Runtime/AudioMixer/Classes/SubmixEffects/AudioMixerSubmixEffectDynamicsProcessor.h
@@ -71,7 +71,6 @@ class FKeySource
 	ESubmixEffectDynamicsKeySource Type = ESubmixEffectDynamicsKeySource::Default;
 	int32 NumChannels = 0;
 	uint32 ObjectId = INDEX_NONE;
-	bool bReportInactive = true;
 
 	mutable FCriticalSection MutateSourceCritSection;
 
@@ -87,7 +86,6 @@ public:
 			NumChannels = 0;
 			ObjectId = INDEX_NONE;
 			Type = ESubmixEffectDynamicsKeySource::Default;
-			bReportInactive = true;
 		}
 	}
 
@@ -103,18 +101,6 @@ public:
 		return NumChannels;
 	}
 
-	void SetReportInactive(bool bInReportInactive)
-	{
-		const FScopeLock ScopeLock(&MutateSourceCritSection);
-		bReportInactive = bInReportInactive;
-	}
-
-	bool ShouldReportInactive() const
-	{
-		const FScopeLock ScopeLock(&MutateSourceCritSection);
-		return bReportInactive;
-	}
-
 	ESubmixEffectDynamicsKeySource GetType() const
 	{
 		const FScopeLock ScopeLock(&MutateSourceCritSection);
@@ -140,7 +126,6 @@ public:
 				NumChannels = InNumChannels;
 
 				bResetPatch = true;
-				bReportInactive = true;
 			}
 		}
 
diff --git a/Engine/Source/Runtime/AudioMixer/Private/Effects/AudioMixerSubmixEffectDynamicsProcessor.cpp b/Engine/Source/Runtime/AudioMixer/Private/Effects/AudioMixerSubmixEffectDynamicsProcessor.cpp
--- a/Engine/Source/Runtime/AudioMixer/Private/Effects/AudioMixerSubmixEffectDynamicsProcessor.cpp
+++ b/Engine/Source/Runtime/AudioMixer/Private/Effects/AudioMixerSubmixEffectDynamicsProcessor.cpp
@@ -202,18 +202,10 @@ bool FSubmixEffectDynamicsProcessor::UpdateKeySourcePatch()
 					UAudioBusSubsystem* AudioBusSubsystem = MixerDevice->GetSubsystem<UAudioBusSubsystem>();
 					if (AudioBusSubsystem)
 					{
-						KeySource.Patch = AudioBusSubsystem->AddPatchOutputForAudioBus(Audio::FAudioBusKey(ObjectId), MixerDevice->GetNumOutputFrames(), KeySource.GetNumChannels());
-					}
-
-					if (KeySource.Patch.IsValid())
-					{
-						DynamicsProcessor.SetKeyNumChannels(KeySource.GetNumChannels());
-						return true;
-					}
-					else if (KeySource.ShouldReportInactive())
-					{
-						KeySource.SetReportInactive(false);
-						UE_LOG(LogAudioMixer, Warning, TEXT("DynamicsProcessor failed to add patch output to inactive AudioBus (ID: '%d')"), ObjectId);
+						const int32 NumChannels = KeySource.GetNumChannels();
+						AudioBusSubsystem->StartAudioBus(Audio::FAudioBusKey(ObjectId), NumChannels, /*bInIsAutomatic=*/false);
+						KeySource.Patch = AudioBusSubsystem->AddPatchOutputForAudioBus(Audio::FAudioBusKey(ObjectId), MixerDevice->GetNumOutputFrames(), NumChannels);
+						DynamicsProcessor.SetKeyNumChannels(NumChannels);
 					}
 				}
 			}
@@ -243,11 +235,6 @@ bool FSubmixEffectDynamicsProcessor::UpdateKeySourcePatch()
 							return true;
 						}
 					}
-					else if (KeySource.ShouldReportInactive())
-					{
-						KeySource.SetReportInactive(false);
-						UE_LOG(LogAudioMixer, Warning, TEXT("DynamicsProcessor failed to add patch output to inactive Submix (ID: '%d')"), ObjectId);
-					}
 				}	
 			}
 		}
