Fix for undefined behavior in Quartz when using non-Audio Mixer devices

diff --git a/Engine/Source/Runtime/AudioMixer/Private/Quartz/QuartzSubsystem.cpp b/Engine/Source/Runtime/AudioMixer/Private/Quartz/QuartzSubsystem.cpp
--- a/Engine/Source/Runtime/AudioMixer/Private/Quartz/QuartzSubsystem.cpp
+++ b/Engine/Source/Runtime/AudioMixer/Private/Quartz/QuartzSubsystem.cpp
@@ -57,6 +57,27 @@ static TUniquePtr<FScopeLock> GetPersistentStateScopeLock(const UWorld* InWorld)
 	return {}; // Empty lock for AudioMixer-less runtimes (server)
 }
 
+bool UQuartzSubsystem::ShouldCreateSubsystem(UObject* Outer) const
+{
+	// Only create the Subsystem if we're using an Audio Mixer Device
+	if (UWorld* OuterWorld = Outer->GetWorld())
+	{
+		FAudioDevice* OuterAudioDevice = OuterWorld->GetAudioDeviceRaw();
+
+		if (OuterAudioDevice && OuterAudioDevice->IsAudioMixerEnabled())
+		{
+			UE_LOG(LogAudioQuartz, Verbose, TEXT("Audio mixer device valid, will create subsystem"));
+
+			// If we do have an audio mixer device, we can create this subsystem.
+			return true;
+		}
+	}
+
+	UE_LOG(LogAudioQuartz, Verbose, TEXT("Quartz Subsystem cannot be created, Audio Mixer is not enabled"));
+
+	// Quartz is not currently supported on non-audio mixer devices.
+	return false;
+}
 
 void UQuartzSubsystem::Initialize(FSubsystemCollectionBase& Collection)
 {
diff --git a/Engine/Source/Runtime/AudioMixer/Public/Quartz/QuartzSubsystem.h b/Engine/Source/Runtime/AudioMixer/Public/Quartz/QuartzSubsystem.h
--- a/Engine/Source/Runtime/AudioMixer/Public/Quartz/QuartzSubsystem.h
+++ b/Engine/Source/Runtime/AudioMixer/Public/Quartz/QuartzSubsystem.h
@@ -56,6 +56,7 @@ public:
 	virtual ~UQuartzSubsystem() override = default;
 
 	//~ Begin UWorldSubsystem Interface
+	virtual bool ShouldCreateSubsystem(UObject* Outer) const override;
 	virtual void Initialize(FSubsystemCollectionBase& Collection) override;
 	virtual void Deinitialize() override;
 	virtual bool DoesSupportWorldType(EWorldType::Type WorldType) const override;
