Add Modulation Destinations to USynthComponent

diff --git a/Engine/Source/Runtime/AudioMixer/Private/Components/SynthComponent.cpp b/Engine/Source/Runtime/AudioMixer/Private/Components/SynthComponent.cpp
--- a/Engine/Source/Runtime/AudioMixer/Private/Components/SynthComponent.cpp
+++ b/Engine/Source/Runtime/AudioMixer/Private/Components/SynthComponent.cpp
@@ -513,6 +513,7 @@ void USynthComponent::Start()
 		AudioComponent->SoundClassOverride = SoundClass;
 		AudioComponent->EnvelopeFollowerAttackTime = EnvelopeFollowerAttackTime;
 		AudioComponent->EnvelopeFollowerReleaseTime = EnvelopeFollowerReleaseTime;
+		AudioComponent->ModulationRouting = ModulationRouting;
 
 		// Copy sound base data to the sound
 		Synth->AttenuationSettings = AttenuationSettings;
@@ -649,6 +650,24 @@ void USynthComponent::AdjustVolume(float AdjustVolumeDuration, float AdjustVolum
 	}
 }
 
+void USynthComponent::SetModulationRouting(const TSet<USoundModulatorBase*>& Modulators, const EModulationDestination Destination, const EModulationRouting RoutingMethod)
+{
+	if (AudioComponent)
+	{
+		AudioComponent->SetModulationRouting(Modulators, Destination, RoutingMethod);
+	}
+}
+
+TSet<USoundModulatorBase*> USynthComponent::GetModulators(const EModulationDestination Destination)
+{
+	if (AudioComponent)
+	{
+		return AudioComponent->GetModulators(Destination);
+	}
+
+	return TSet<USoundModulatorBase*>();
+}
+
 void USynthComponent::SynthCommand(TFunction<void()> Command)
 {
 	if (SoundGenerator.IsValid())
diff --git a/Engine/Source/Runtime/AudioMixer/Public/Components/SynthComponent.h b/Engine/Source/Runtime/AudioMixer/Public/Components/SynthComponent.h
--- a/Engine/Source/Runtime/AudioMixer/Public/Components/SynthComponent.h
+++ b/Engine/Source/Runtime/AudioMixer/Public/Components/SynthComponent.h
@@ -181,6 +181,23 @@ public:
 	UFUNCTION(BlueprintCallable, Category = "Audio|Components|Audio")
 	void AdjustVolume(float AdjustVolumeDuration, float AdjustVolumeLevel, const EAudioFaderCurve FadeCurve = EAudioFaderCurve::Linear) const;
 
+	/**
+	* Sets the routing for one of the given Synth component's Modulation Destinations.
+	* @param Modulators The set of modulators to apply to the given destination on the component.
+	* @param Destination The destination to assign the modulators to.
+	* @param RoutingMethod The routing method to use for the given modulator.
+	*/
+	UFUNCTION(BlueprintCallable, Category = "Audio|Components|Audio", DisplayName = "Set Modulation Routing")
+	void SetModulationRouting(const TSet<USoundModulatorBase*>& Modulators, const EModulationDestination Destination, const EModulationRouting RoutingMethod = EModulationRouting::Inherit);
+	
+	/**
+	* Gets the set of currently active modulators for a given Modulation Destination.
+	* @param Destination The Destination to retrieve the Modulators from.
+	* @return The set of of Modulators applied to this component for the given Destination.
+	*/
+	UFUNCTION(BlueprintPure, Category = "Audio|Components|Audio", DisplayName = "Get Modulators")
+	UPARAM(DisplayName = "Modulators") TSet<USoundModulatorBase*> GetModulators(const EModulationDestination Destination);
+
 	/** Auto destroy this component on completion */
 	UPROPERTY()
 	uint8 bAutoDestroy : 1;
@@ -231,6 +248,9 @@ public:
 	UPROPERTY(EditAnywhere, BlueprintReadWrite, Category = Concurrency)
 	TSet<TObjectPtr<USoundConcurrency>> ConcurrencySet;
 
+	UPROPERTY(EditAnywhere, BlueprintReadWrite, Category = Modulation)
+	FSoundModulationDefaultRoutingSettings ModulationRouting;
+
 	/** Sound class this sound belongs to */
 	UPROPERTY(EditAnywhere, Category = SoundClass)
 	TObjectPtr<USoundClass> SoundClass;
