Adding ability to have different options for 2D, non-spatialized radius blend.

* Legacy omni blend
* Direct 2D (i.e. stereo)
* Immersive 2D (i.e. quad)

diff --git a/Engine/Config/BaseEngine.ini b/Engine/Config/BaseEngine.ini
--- a/Engine/Config/BaseEngine.ini
+++ b/Engine/Config/BaseEngine.ini
@@ -1134,6 +1134,8 @@ DefaultWorldPartitionHLODLayer=/Engine/Maps/Templates/HLODs/HLODLayer_Instanced.
 +PropertyRedirects=(OldName="MaterialInstanceBasePropertyOverrides.bOverride_MaxWorldPositionOffsetDistance", NewName="MaterialInstanceBasePropertyOverrides.bOverride_MaxWorldPositionOffsetDisplacement")
 +PropertyRedirects=(OldName="Material.MaxWorldPositionOffsetDistance", NewName="Material.MaxWorldPositionOffsetDisplacement")
 
++PropertyRedirects=(OldName="/Script/Engine.SoundAttenuationSettings.OmniRadius", NewName="NonSpatializedRadiusStart")
+
 ; 5.3
 +FunctionRedirects=(OldName="KismetMathLibrary.BreakVector3f", NewName="BreakVector")
 
diff --git a/Engine/Source/Editor/DetailCustomizations/Private/AttenuationSettingsCustomizations.cpp b/Engine/Source/Editor/DetailCustomizations/Private/AttenuationSettingsCustomizations.cpp
--- a/Engine/Source/Editor/DetailCustomizations/Private/AttenuationSettingsCustomizations.cpp
+++ b/Engine/Source/Editor/DetailCustomizations/Private/AttenuationSettingsCustomizations.cpp
@@ -469,7 +469,15 @@ void FSoundAttenuationSettingsCustomization::CustomizeChildren(TSharedRef<IPrope
 		.Visibility(IsAttenuationOverriddenVisibleAttribute())
 		.EditCondition(GetIsSpatializationEnabledAttribute(), nullptr);
 
-	LayoutBuilder.AddPropertyToCategory(PropertyHandles.FindChecked(GET_MEMBER_NAME_CHECKED(FSoundAttenuationSettings, OmniRadius)))
+	LayoutBuilder.AddPropertyToCategory(PropertyHandles.FindChecked(GET_MEMBER_NAME_CHECKED(FSoundAttenuationSettings, NonSpatializedRadiusStart)))
+		.Visibility(IsAttenuationOverriddenVisibleAttribute())
+		.EditCondition(GetIsSpatializationEnabledAttribute(), nullptr);
+
+	LayoutBuilder.AddPropertyToCategory(PropertyHandles.FindChecked(GET_MEMBER_NAME_CHECKED(FSoundAttenuationSettings, NonSpatializedRadiusEnd)))
+		.Visibility(IsAttenuationOverriddenVisibleAttribute())
+		.EditCondition(GetIsSpatializationEnabledAttribute(), nullptr);
+
+	LayoutBuilder.AddPropertyToCategory(PropertyHandles.FindChecked(GET_MEMBER_NAME_CHECKED(FSoundAttenuationSettings, NonSpatializedRadiusMode)))
 		.Visibility(IsAttenuationOverriddenVisibleAttribute())
 		.EditCondition(GetIsSpatializationEnabledAttribute(), nullptr);
 
@@ -678,7 +686,7 @@ void FSoundAttenuationSettingsCustomization::CustomizeChildren(TSharedRef<IPrope
 		.Visibility(IsAttenuationOverriddenVisibleAttribute())
 		.EditCondition(GetIsSubmixSendAttenuationEnabledAttribute(), nullptr);
 
-	LayoutBuilder.AddPropertyToCategory(bIsSourceDataOverrideEnabledHandle)
+	LayoutBuilder.AddPropertyToCategory(PropertyHandles.FindChecked(GET_MEMBER_NAME_CHECKED(FSoundAttenuationSettings, bEnableSourceDataOverride)))
 		.Visibility(IsAttenuationOverriddenVisibleAttribute())
 		.EditCondition(IsAttenuationOverriddenAttribute(), nullptr);
 
@@ -717,7 +725,7 @@ void FSoundAttenuationSettingsCustomization::CustomizeChildren(TSharedRef<IPrope
 	LayoutBuilder.EditCategory("AttenuationAudioLink", LOCTEXT("AttenuationAudioLink", "Attenuation (AudioLink)"));
 	LayoutBuilder.SortCategories(AttenuationSettingsUtils::SortCategories);
 
-	if (PropertyHandles.Num() != 68)
+	if (PropertyHandles.Num() != 70)
 	{
 		ensureMsgf(false, TEXT("Unexpected property handle(s) customizing FSoundAttenuationSettings. %d handles found"), PropertyHandles.Num());
 	}
diff --git a/Engine/Source/Runtime/AudioExtensions/Public/IAudioExtensionPlugin.h b/Engine/Source/Runtime/AudioExtensions/Public/IAudioExtensionPlugin.h
--- a/Engine/Source/Runtime/AudioExtensions/Public/IAudioExtensionPlugin.h
+++ b/Engine/Source/Runtime/AudioExtensions/Public/IAudioExtensionPlugin.h
@@ -97,9 +97,12 @@ struct FSpatializationParams
 	/** The distance used to compute attenuation. Maybe different from the distance between listener and emitter if it's overridden. */
 	float AttenuationDistance;
 
-	/** The normalized omni radius, or the radius that will blend a sound to non-3d */
+	/** Deprecated */
 	float NormalizedOmniRadius;
 
+	/** The amount of non-spatialized this source is. 1.0 means fully 2D, 0.0 means fully 3D. */
+	float NonSpatializedAmount;
+
 	/** The time when this spatialization params was built. */
 	double AudioClock;
 
@@ -115,6 +118,7 @@ struct FSpatializationParams
 		, Distance(0.0f)
 		, AttenuationDistance(0.0f)
 		, NormalizedOmniRadius(0.0f)
+		, NonSpatializedAmount(0.0f)
 		, AudioClock(0.0)
 	{}
 };
diff --git a/Engine/Source/Runtime/AudioMixer/Private/AudioMixerDevice.cpp b/Engine/Source/Runtime/AudioMixer/Private/AudioMixerDevice.cpp
--- a/Engine/Source/Runtime/AudioMixer/Private/AudioMixerDevice.cpp
+++ b/Engine/Source/Runtime/AudioMixer/Private/AudioMixerDevice.cpp
@@ -1929,7 +1929,7 @@ namespace Audio
 		return SourceManager->GetNumActiveSources();
 	}
 
-	void FMixerDevice::Get3DChannelMap(const int32 InSubmixNumChannels, const FWaveInstance* InWaveInstance, float EmitterAzimith, float NormalizedOmniRadius, Audio::FAlignedFloatBuffer& OutChannelMap)
+	void FMixerDevice::Get3DChannelMap(const int32 InSubmixNumChannels, const FWaveInstance* InWaveInstance, float EmitterAzimith, float InNonSpatializedAmount, const TMap<EAudioMixerChannel::Type, float>* InOmniMap, float InDefaultOmniValue, Audio::FAlignedFloatBuffer& OutChannelMap)
 	{
 		// If we're center-channel only, then no need for spatial calculations, but need to build a channel map
 		if (InWaveInstance->bCenterChannelOnly)
@@ -2025,13 +2025,7 @@ namespace Audio
 			PrevChannelPan = 1.0f - Fraction;
 		}
 
-		float NormalizedOmniRadSquared = NormalizedOmniRadius * NormalizedOmniRadius;
-		float OmniAmount = 0.0f;
-
-		if (NormalizedOmniRadSquared > 1.0f)
-		{
-			OmniAmount = 1.0f - 1.0f / NormalizedOmniRadSquared;
-		}
+		float OmniAmount = InNonSpatializedAmount;
 
 		// Build the output channel map based on the current platform device output channel array 
 
@@ -2040,36 +2034,80 @@ namespace Audio
 		{
 			NumSpatialChannels--;
 		}
-		float OmniPanFactor = 1.0f / NumSpatialChannels;
 
-		float DefaultEffectivePan = !OmniAmount ? 0.0f : FMath::Lerp(0.0f, OmniPanFactor, OmniAmount);
 		const TArray<EAudioMixerChannel::Type>& ChannelArray = GetChannelArray();
 
-		for (EAudioMixerChannel::Type Channel : ChannelArray)
+		if (OmniAmount > 0.0f)
 		{
-			float EffectivePan = DefaultEffectivePan;
-
-			// Check for manual channel mapping parameters (LFE and Front Center)
-			if (Channel == EAudioMixerChannel::LowFrequency)
-			{
-				EffectivePan = InWaveInstance->LFEBleed;
-			}
-			else if (Channel == PrevChannelInfo->Channel)
+			for (EAudioMixerChannel::Type Channel : ChannelArray)
 			{
-				EffectivePan = !OmniAmount ? PrevChannelPan : FMath::Lerp(PrevChannelPan, OmniPanFactor, OmniAmount);
+				float OmniPanFactor = InDefaultOmniValue;
+
+				if (InOmniMap)
+				{
+					const float* MappedOmniPanFactor = InOmniMap->Find(Channel);
+					if (MappedOmniPanFactor)
+					{
+						OmniPanFactor = *MappedOmniPanFactor;
+					}
+				}
+
+				float EffectivePan = 0.0f;
+
+				// Check for manual channel mapping parameters (LFE and Front Center)
+				if (Channel == EAudioMixerChannel::LowFrequency)
+				{
+					EffectivePan = InWaveInstance->LFEBleed;
+				}
+				else if (Channel == PrevChannelInfo->Channel)
+				{
+					EffectivePan = FMath::Lerp(PrevChannelPan, OmniPanFactor, OmniAmount);
+				}
+				else if (Channel == NextChannelInfo->Channel)
+				{
+					EffectivePan = FMath::Lerp(NextChannelPan, OmniPanFactor, OmniAmount);
+				}
+				else if (Channel == EAudioMixerChannel::FrontCenter)
+				{
+					EffectivePan = FMath::Lerp(0.0f, OmniPanFactor, OmniAmount);
+					EffectivePan = FMath::Max(InWaveInstance->VoiceCenterChannelVolume, EffectivePan);
+				}
+				else
+				{
+					EffectivePan = FMath::Lerp(0.0f, OmniPanFactor, OmniAmount);
+				}
+
+				AUDIO_MIXER_CHECK(EffectivePan >= 0.0f && EffectivePan <= 1.0f);
+				OutChannelMap.Add(EffectivePan);
 			}
-			else if (Channel == NextChannelInfo->Channel)
+		}
+		else
+		{
+			for (EAudioMixerChannel::Type Channel : ChannelArray)
 			{
-				EffectivePan = !OmniAmount ? NextChannelPan : FMath::Lerp(NextChannelPan, OmniPanFactor, OmniAmount);
-			}
+				float EffectivePan = 0.0f;
 
-			if (Channel == EAudioMixerChannel::FrontCenter)
-			{
-				EffectivePan = FMath::Max(InWaveInstance->VoiceCenterChannelVolume, EffectivePan);
-			}
+				// Check for manual channel mapping parameters (LFE and Front Center)
+				if (Channel == EAudioMixerChannel::LowFrequency)
+				{
+					EffectivePan = InWaveInstance->LFEBleed;
+				}
+				else if (Channel == PrevChannelInfo->Channel)
+				{
+					EffectivePan = PrevChannelPan;
+				}
+				else if (Channel == NextChannelInfo->Channel)
+				{
+					EffectivePan = NextChannelPan;
+				}
+				else if (Channel == EAudioMixerChannel::FrontCenter)
+				{
+					EffectivePan = FMath::Max(InWaveInstance->VoiceCenterChannelVolume, EffectivePan);
+				}
 
-			AUDIO_MIXER_CHECK(EffectivePan >= 0.0f && EffectivePan <= 1.0f);
-			OutChannelMap.Add(EffectivePan);
+				AUDIO_MIXER_CHECK(EffectivePan >= 0.0f && EffectivePan <= 1.0f);
+				OutChannelMap.Add(EffectivePan);
+			}
 		}
 	}
 
diff --git a/Engine/Source/Runtime/AudioMixer/Private/AudioMixerSource.cpp b/Engine/Source/Runtime/AudioMixer/Private/AudioMixerSource.cpp
--- a/Engine/Source/Runtime/AudioMixer/Private/AudioMixerSource.cpp
+++ b/Engine/Source/Runtime/AudioMixer/Private/AudioMixerSource.cpp
@@ -1944,7 +1944,10 @@ namespace Audio
 			// Don't need to compute the source channel map if the absolute azimuth hasn't changed much
 			PreviousAzimuth = WaveInstance->AbsoluteAzimuth;
 			OutChannelMap.Reset();
-			MixerDevice->Get3DChannelMap(MixerDevice->GetNumDeviceChannels(), WaveInstance, WaveInstance->AbsoluteAzimuth, SpatializationParams.NormalizedOmniRadius, OutChannelMap);
+
+			int32 NumDeviceChannels = MixerDevice->GetNumDeviceChannels();
+			float DefaultOmniAmount = 1.0f / NumDeviceChannels;
+			MixerDevice->Get3DChannelMap(NumDeviceChannels, WaveInstance, WaveInstance->AbsoluteAzimuth, SpatializationParams.NonSpatializedAmount, nullptr, DefaultOmniAmount, OutChannelMap);
 			return true;
 		}
 		else if (!OutChannelMap.Num() || (IsSpatializationCVarEnabled() != bPrevAllowedSpatializationSetting))
@@ -2000,10 +2003,97 @@ namespace Audio
 				// Reset the channel map, the stereo spatialization channel mapping calls below will append their mappings
 				OutChannelMap.Reset();
 
-				const int32 NumOutputChannels = MixerDevice->GetNumDeviceChannels();
+				int32 NumOutputChannels = MixerDevice->GetNumDeviceChannels();
+
+				if (WaveInstance->NonSpatializedRadiusMode == ENonSpatializedRadiusSpeakerMapMode::OmniDirectional)
+				{
+					float DefaultOmniAmount = 1.0f / NumOutputChannels;
+					MixerDevice->Get3DChannelMap(NumOutputChannels, WaveInstance, LeftAzimuth, SpatializationParams.NonSpatializedAmount, nullptr, DefaultOmniAmount, OutChannelMap);
+					MixerDevice->Get3DChannelMap(NumOutputChannels, WaveInstance, RightAzimuth, SpatializationParams.NonSpatializedAmount, nullptr, DefaultOmniAmount, OutChannelMap);
+				}
+				else if (WaveInstance->NonSpatializedRadiusMode == ENonSpatializedRadiusSpeakerMapMode::Direct2D)
+				{
+					// Create some omni maps for left and right channels
+					auto CreateLeftOmniMap = []() -> TMap<EAudioMixerChannel::Type, float>
+					{
+						TMap<EAudioMixerChannel::Type, float> LeftOmniMap;
+						LeftOmniMap.Add(EAudioMixerChannel::FrontLeft, 1.0f);
+						return LeftOmniMap;
+					};
+
+					auto CreateRightOmniMap = []() -> TMap<EAudioMixerChannel::Type, float>
+					{
+						TMap<EAudioMixerChannel::Type, float> RightOmniMap;
+						RightOmniMap.Add(EAudioMixerChannel::FrontRight, 1.0f);
+						return RightOmniMap;
+					};
+
+					static const TMap<EAudioMixerChannel::Type, float> LeftOmniMap = CreateLeftOmniMap();
+					static const TMap<EAudioMixerChannel::Type, float> RightOmniMap = CreateRightOmniMap();
+
+					MixerDevice->Get3DChannelMap(NumOutputChannels, WaveInstance, LeftAzimuth, SpatializationParams.NonSpatializedAmount, &LeftOmniMap, 0.0f, OutChannelMap);
+					MixerDevice->Get3DChannelMap(NumOutputChannels, WaveInstance, RightAzimuth, SpatializationParams.NonSpatializedAmount, &RightOmniMap, 0.0f, OutChannelMap);
+				}
+				else
+				{
+					// If we are in 5.1, we need to use the side-channel speakers
+					if (NumOutputChannels == 6)
+					{
+						// Create some omni maps for left and right channels
+						auto CreateLeftOmniMap = []() -> TMap<EAudioMixerChannel::Type, float>
+						{
+							TMap<EAudioMixerChannel::Type, float> LeftOmniMap;
+							LeftOmniMap.Add(EAudioMixerChannel::FrontLeft, 1.0f);
+							LeftOmniMap.Add(EAudioMixerChannel::SideLeft, 1.0f);
+
+							return LeftOmniMap;
+						};
 
-				MixerDevice->Get3DChannelMap(NumOutputChannels, WaveInstance, LeftAzimuth, SpatializationParams.NormalizedOmniRadius, OutChannelMap);
-				MixerDevice->Get3DChannelMap(NumOutputChannels, WaveInstance, RightAzimuth, SpatializationParams.NormalizedOmniRadius, OutChannelMap);
+						auto CreateRightOmniMap = []() -> TMap<EAudioMixerChannel::Type, float>
+						{
+							TMap<EAudioMixerChannel::Type, float> RightOmniMap;
+							RightOmniMap.Add(EAudioMixerChannel::FrontRight, 1.0f);
+							RightOmniMap.Add(EAudioMixerChannel::SideRight, 1.0f);
+
+							return RightOmniMap;
+						};
+
+						static const TMap<EAudioMixerChannel::Type, float> LeftOmniMap = CreateLeftOmniMap();
+						static const TMap<EAudioMixerChannel::Type, float> RightOmniMap = CreateRightOmniMap();
+
+						MixerDevice->Get3DChannelMap(NumOutputChannels, WaveInstance, LeftAzimuth, SpatializationParams.NonSpatializedAmount, &LeftOmniMap, 0.0f, OutChannelMap);
+						MixerDevice->Get3DChannelMap(NumOutputChannels, WaveInstance, RightAzimuth, SpatializationParams.NonSpatializedAmount, &RightOmniMap, 0.0f, OutChannelMap);
+
+					}
+					// If we are in 7.1 we need to use the back-channel speakers
+					else if (NumOutputChannels == 8)
+					{
+						// Create some omni maps for left and right channels
+						auto CreateLeftOmniMap = []() -> TMap<EAudioMixerChannel::Type, float>
+						{
+							TMap<EAudioMixerChannel::Type, float> LeftOmniMap;
+							LeftOmniMap.Add(EAudioMixerChannel::FrontLeft, 1.0f);
+							LeftOmniMap.Add(EAudioMixerChannel::BackLeft, 1.0f);
+
+							return LeftOmniMap;
+						};
+
+						auto CreateRightOmniMap = []() -> TMap<EAudioMixerChannel::Type, float>
+						{
+							TMap<EAudioMixerChannel::Type, float> RightOmniMap;
+							RightOmniMap.Add(EAudioMixerChannel::FrontRight, 1.0f);
+							RightOmniMap.Add(EAudioMixerChannel::BackRight, 1.0f);
+
+							return RightOmniMap;
+						};
+
+						static const TMap<EAudioMixerChannel::Type, float> LeftOmniMap = CreateLeftOmniMap();
+						static const TMap<EAudioMixerChannel::Type, float> RightOmniMap = CreateRightOmniMap();
+
+						MixerDevice->Get3DChannelMap(NumOutputChannels, WaveInstance, LeftAzimuth, SpatializationParams.NonSpatializedAmount, &LeftOmniMap, 0.0f, OutChannelMap);
+						MixerDevice->Get3DChannelMap(NumOutputChannels, WaveInstance, RightAzimuth, SpatializationParams.NonSpatializedAmount, &RightOmniMap, 0.0f, OutChannelMap);
+					}
+				}		
 
 				return true;
 			}
diff --git a/Engine/Source/Runtime/AudioMixer/Public/AudioMixerDevice.h b/Engine/Source/Runtime/AudioMixer/Public/AudioMixerDevice.h
--- a/Engine/Source/Runtime/AudioMixer/Public/AudioMixerDevice.h
+++ b/Engine/Source/Runtime/AudioMixer/Public/AudioMixerDevice.h
@@ -239,7 +239,7 @@ namespace Audio
 		IAudioMixerPlatformInterface* GetAudioMixerPlatform() const { return AudioMixerPlatform; }
 
 		// Builds a 3D channel map for a spatialized source.
-		void Get3DChannelMap(const int32 InSubmixNumChannels, const FWaveInstance* InWaveInstance, const float EmitterAzimuth, const float NormalizedOmniRadius, Audio::FAlignedFloatBuffer& OutChannelMap);
+		void Get3DChannelMap(const int32 InSubmixNumChannels, const FWaveInstance* InWaveInstance, const float EmitterAzimuth, const float NonSpatiliazedFactor, const TMap<EAudioMixerChannel::Type, float>* InOmniMap, float InDefaultOmniValue, Audio::FAlignedFloatBuffer& OutChannelMap);
 
 		// Builds a channel gain matrix for a non-spatialized source. The non-static variation of this function queries AudioMixerDevice->NumOutputChannels directly which may not be thread safe.
 		void Get2DChannelMap(bool bIsVorbis, const int32 NumSourceChannels, const bool bIsCenterChannelOnly, Audio::FAlignedFloatBuffer& OutChannelMap) const;
@@ -396,9 +396,6 @@ namespace Audio
 
 		int32 DeviceOutputChannels;
 
-		/** Channel type arrays for submix channel types. */
-		TArray<EAudioMixerChannel::Type> DeviceChannelArray;
-
 		/** What upmix method to use for mono channel upmixing. */
 		EMonoChannelUpmixMethod MonoChannelUpmixMethod;
 
diff --git a/Engine/Source/Runtime/Core/Public/UObject/FortniteMainBranchObjectVersion.h b/Engine/Source/Runtime/Core/Public/UObject/FortniteMainBranchObjectVersion.h
--- a/Engine/Source/Runtime/Core/Public/UObject/FortniteMainBranchObjectVersion.h
+++ b/Engine/Source/Runtime/Core/Public/UObject/FortniteMainBranchObjectVersion.h
@@ -274,6 +274,9 @@ struct CORE_API FFortniteMainBranchObjectVersion
 		// Add FWorldPartitionActorFilter to FLevelInstanceActorDesc/FDataLayerInstanceDesc
 		WorldPartitionActorFilter,
 
+		// Change the non-spatialized radius to blend to a pure 2D spatialized sound vs omnidirectional
+		AudioAttenuationNonSpatializedRadiusBlend,
+
 		// -----<new versions can be added above this line>-------------------------------------------------
 		VersionPlusOne,
 		LatestVersion = VersionPlusOne - 1
diff --git a/Engine/Source/Runtime/Engine/Classes/Sound/SoundAttenuation.h b/Engine/Source/Runtime/Engine/Classes/Sound/SoundAttenuation.h
--- a/Engine/Source/Runtime/Engine/Classes/Sound/SoundAttenuation.h
+++ b/Engine/Source/Runtime/Engine/Classes/Sound/SoundAttenuation.h
@@ -150,6 +150,21 @@ struct ENGINE_API FAttenuationSubmixSendSettings
 	FRuntimeFloatCurve CustomSubmixSendCurve;
 };
 
+
+// Defines how to speaker map the sound when using the non-spatialized radius feature
+UENUM(BlueprintType)
+enum class ENonSpatializedRadiusSpeakerMapMode : uint8
+{
+	// Will blend the 3D sound to an omni-directional sound (equal output mapping in all directions)
+	OmniDirectional,
+
+	// Will blend the 3D source to the same representation speaker map used when playing the asset 2D
+	Direct2D,
+
+	// Will blend the 3D source to a multichannel 2D version (i.e. upmix stereo to quad) if rendering in surround
+	Surround2D,
+};
+
 /*
 The settings for attenuating.
 */
@@ -253,11 +268,22 @@ struct ENGINE_API FSoundAttenuationSettings : public FBaseAttenuationSettings
 #if WITH_EDITORONLY_DATA
 	UPROPERTY()
 	TEnumAsByte<enum ESoundDistanceCalc> DistanceType_DEPRECATED;
+
+ 	UPROPERTY()
+ 	float OmniRadius_DEPRECATED;
 #endif
 
-	/** The distance below which a sound is non-spatialized (2D). This prevents near-field audio from flipping as audio crosses the listener's position. This does not apply when using a 3rd party binaural plugin (audio will remain spatialized). */
-	UPROPERTY(EditAnywhere, BlueprintReadWrite, Category = AttenuationSpatialization, meta=(ClampMin = "0", EditCondition="bSpatialize", DisplayName="Non-Spatialized Radius"))
-	float OmniRadius;
+	/** The distance below which a sound begins to be non-spatialized (2D). This prevents near-field audio from flipping as audio crosses the listener's position. This does not apply when using a 3rd party binaural plugin (audio will remain spatialized). */
+	UPROPERTY(EditAnywhere, BlueprintReadWrite, Category = AttenuationSpatialization, meta = (ClampMin = "0", EditCondition = "bSpatialize"))
+	float NonSpatializedRadiusStart;
+
+	/** The distance below which a sound is fully non-spatialized (2D). */
+	UPROPERTY(EditAnywhere, BlueprintReadWrite, Category = AttenuationSpatialization, meta = (ClampMin = "0", EditCondition = "bSpatialize"))
+	float NonSpatializedRadiusEnd;
+
+	/** Defines how to blend to the 2D sound when using the non-spatialized radius. */
+	UPROPERTY(EditAnywhere, BlueprintReadWrite, Category = AttenuationSpatialization, meta = (ClampMin = "0", EditCondition = "bSpatialize"))
+	ENonSpatializedRadiusSpeakerMapMode NonSpatializedRadiusMode;
 
 	/** The world-space distance between left and right stereo channels when stereo assets are 3D spatialized. */
 	UPROPERTY(EditAnywhere, BlueprintReadWrite, Category = AttenuationSpatialization, meta = (ClampMin = "0", EditCondition = "bSpatialize", DisplayName = "3D Stereo Spread"))
@@ -438,8 +464,11 @@ struct ENGINE_API FSoundAttenuationSettings : public FBaseAttenuationSettings
 		, PriorityAttenuationMethod(EPriorityAttenuationMethod::Linear)
 #if WITH_EDITORONLY_DATA
 		, DistanceType_DEPRECATED(SOUNDDISTANCE_Normal)
+		, OmniRadius_DEPRECATED(0.0f)
 #endif
-		, OmniRadius(0.0f)
+		, NonSpatializedRadiusStart(0.0f)
+		, NonSpatializedRadiusEnd(0.0f)
+		, NonSpatializedRadiusMode(ENonSpatializedRadiusSpeakerMapMode::OmniDirectional)
 		, StereoSpread(200.0f)
 #if WITH_EDITORONLY_DATA
 		, SpatializationPluginSettings_DEPRECATED(nullptr)
diff --git a/Engine/Source/Runtime/Engine/Private/ActiveSound.cpp b/Engine/Source/Runtime/Engine/Private/ActiveSound.cpp
--- a/Engine/Source/Runtime/Engine/Private/ActiveSound.cpp
+++ b/Engine/Source/Runtime/Engine/Private/ActiveSound.cpp
@@ -1849,7 +1849,9 @@ void FActiveSound::UpdateAttenuation(float DeltaTime, FSoundParseParameters& Par
 		}
 	}
 
-	ParseParams.OmniRadius = Settings->OmniRadius;
+	ParseParams.NonSpatializedRadiusStart = Settings->NonSpatializedRadiusStart;
+	ParseParams.NonSpatializedRadiusEnd = Settings->NonSpatializedRadiusEnd;
+	ParseParams.NonSpatializedRadiusMode = Settings->NonSpatializedRadiusMode;
 	ParseParams.StereoSpread = Settings->StereoSpread;
 	ParseParams.bApplyNormalizationToStereoSounds = Settings->bApplyNormalizationToStereoSounds;
 	ParseParams.bUseSpatialization |= Settings->bSpatialize;
diff --git a/Engine/Source/Runtime/Engine/Private/Audio.cpp b/Engine/Source/Runtime/Engine/Private/Audio.cpp
--- a/Engine/Source/Runtime/Engine/Private/Audio.cpp
+++ b/Engine/Source/Runtime/Engine/Private/Audio.cpp
@@ -531,21 +531,47 @@ FSpatializationParams FSoundSource::GetSpatializationParams()
 		// Independently retrieve the attenuation used for distance in case it was overridden
 		Params.AttenuationDistance = AudioDevice->GetDistanceToNearestListener(WaveInstance->Location);
 		
-		// If we are using the OmniRadius feature
-		if (WaveInstance->OmniRadius > 0.0f)
+		// If we are using the Non-spatialized radius feature
+		if (WaveInstance->NonSpatializedRadiusStart > 0.0f)
 		{
-			// Initialize to full omni-directionality (bigger value, more omni)
-			static const float MaxNormalizedRadius = 1000000.0f;
-			Params.NormalizedOmniRadius = MaxNormalizedRadius;
+			float NonSpatializedRadiusEnd = FMath::Min(WaveInstance->NonSpatializedRadiusStart, WaveInstance->NonSpatializedRadiusEnd);
 
-			if (Params.Distance > 0)
+			if (Params.Distance > 0.0f)
 			{
-				Params.NormalizedOmniRadius = FMath::Clamp(WaveInstance->OmniRadius / Params.Distance, 0.0f, MaxNormalizedRadius);
+				// If the user specified a distance below which to be fully 2D
+				if (NonSpatializedRadiusEnd > 0.0f)
+				{
+					// We're in the non-spatialized domain
+					if (Params.Distance < WaveInstance->NonSpatializedRadiusStart)
+					{
+						float NonSpatializationRange = WaveInstance->NonSpatializedRadiusStart - NonSpatializedRadiusEnd;
+						NonSpatializationRange = FMath::Max(NonSpatializationRange, 1.0f);
+						Params.NonSpatializedAmount = FMath::Clamp((WaveInstance->NonSpatializedRadiusStart - Params.Distance)/ NonSpatializationRange, 0.0f, 1.0f);
+					}					
+				}
+				else
+				{
+					// Initialize to full omni-directionality (bigger value, more omni)
+					static const float MaxNormalizedRadius = 1000000.0f;
+					float NormalizedOmniRadus = FMath::Clamp(WaveInstance->NonSpatializedRadiusStart / Params.Distance, 0.0f, MaxNormalizedRadius);
+					if (NormalizedOmniRadus > 1.0f)
+					{
+						float NormalizedOmniRadusSquared = NormalizedOmniRadus * NormalizedOmniRadus;
+						Params.NonSpatializedAmount = 1.0f - 1.0f / NormalizedOmniRadusSquared;
+					}
+					else
+					{
+						Params.NonSpatializedAmount = 0.0f;
+					}
+				}
+
+				UE_LOG(LogTemp, Log, TEXT("Distance: %.2f, NonSpatializedRadiusStart: %.2f, NonSpatializedRadiusEnd: %.2f, NonSpatializedAmount: %.2f"), Params.Distance, WaveInstance->NonSpatializedRadiusStart, WaveInstance->NonSpatializedRadiusEnd, Params.NonSpatializedAmount);
+
 			}
 		}
 		else
 		{
-			Params.NormalizedOmniRadius = 0.0f;
+			Params.NonSpatializedAmount = 0.0f;
 		}
 
 		Params.EmitterPosition = EmitterPosition;
@@ -560,6 +586,7 @@ FSpatializationParams FSoundSource::GetSpatializationParams()
 	else
 	{
 		Params.NormalizedOmniRadius = 0.0f;
+		Params.NonSpatializedAmount = 0.0f;
 		Params.Distance = 0.0f;
 		Params.EmitterPosition = FVector::ZeroVector;
 	}
@@ -826,7 +853,9 @@ FWaveInstance::FWaveInstance(const UPTRINT InWaveInstanceHash, FActiveSound& InA
 	, AttenuationHighpassFilterFrequency(MIN_FILTER_FREQUENCY)
 	, Pitch(0.0f)
 	, Location(FVector::ZeroVector)
-	, OmniRadius(0.0f)
+	, NonSpatializedRadiusStart(0.0f)
+	, NonSpatializedRadiusEnd(0.0f)
+	, NonSpatializedRadiusMode(ENonSpatializedRadiusSpeakerMapMode::OmniDirectional)
 	, StereoSpread(0.0f)
 	, AttenuationDistance(0.0f)
 	, ListenerToSoundDistance(0.0f)
diff --git a/Engine/Source/Runtime/Engine/Private/SoundAttenuation.cpp b/Engine/Source/Runtime/Engine/Private/SoundAttenuation.cpp
--- a/Engine/Source/Runtime/Engine/Private/SoundAttenuation.cpp
+++ b/Engine/Source/Runtime/Engine/Private/SoundAttenuation.cpp
@@ -4,6 +4,7 @@
 
 #include "AudioDevice.h"
 #include "UObject/AnimPhysObjectVersion.h"
+#include "UObject/FortniteMainBranchObjectVersion.h"
 
 /*-----------------------------------------------------------------------------
 	USoundAttenuation implementation.
@@ -57,6 +58,14 @@ void FSoundAttenuationSettings::PostSerialize(const FArchive& Ar)
 			PluginSettings.ReverbPluginSettingsArray.Add(ReverbPluginSettings_DEPRECATED);
 		}
 	}
+
+	if (Ar.IsLoading() && Ar.CustomVer(FFortniteMainBranchObjectVersion::GUID) < FFortniteMainBranchObjectVersion::AudioAttenuationNonSpatializedRadiusBlend)
+	{
+		if (OmniRadius_DEPRECATED)
+		{
+			NonSpatializedRadiusStart = OmniRadius_DEPRECATED;
+		}
+	}
 }
 #endif
 
@@ -90,7 +99,9 @@ bool FSoundAttenuationSettings::operator==(const FSoundAttenuationSettings& Othe
 			&& bSpatialize			    == Other.bSpatialize
 			&& dBAttenuationAtMax	    == Other.dBAttenuationAtMax
 			&& FalloffMode				== Other.FalloffMode
-			&& OmniRadius				== Other.OmniRadius
+			&& NonSpatializedRadiusStart == Other.NonSpatializedRadiusStart
+			&& NonSpatializedRadiusEnd == Other.NonSpatializedRadiusEnd
+			&& NonSpatializedRadiusMode == Other.NonSpatializedRadiusMode
 			&& bApplyNormalizationToStereoSounds == Other.bApplyNormalizationToStereoSounds
 			&& StereoSpread				== Other.StereoSpread
 			&& DistanceAlgorithm	    == Other.DistanceAlgorithm
diff --git a/Engine/Source/Runtime/Engine/Private/SoundWave.cpp b/Engine/Source/Runtime/Engine/Private/SoundWave.cpp
--- a/Engine/Source/Runtime/Engine/Private/SoundWave.cpp
+++ b/Engine/Source/Runtime/Engine/Private/SoundWave.cpp
@@ -2854,7 +2854,9 @@ void USoundWave::Parse(FAudioDevice* AudioDevice, const UPTRINT NodeWaveInstance
 	WaveInstance->bApplyRadioFilter = ActiveSound.bApplyRadioFilter;
 	WaveInstance->StartTime = ParseParams.StartTime;
 	WaveInstance->UserIndex = ActiveSound.UserIndex;
-	WaveInstance->OmniRadius = ParseParams.OmniRadius;
+	WaveInstance->NonSpatializedRadiusStart = ParseParams.NonSpatializedRadiusStart;
+	WaveInstance->NonSpatializedRadiusEnd = ParseParams.NonSpatializedRadiusEnd;
+	WaveInstance->NonSpatializedRadiusMode = ParseParams.NonSpatializedRadiusMode;
 	WaveInstance->StereoSpread = ParseParams.StereoSpread;
 	WaveInstance->AttenuationDistance = ParseParams.AttenuationDistance;
 	WaveInstance->ListenerToSoundDistance = ParseParams.ListenerToSoundDistance;
diff --git a/Engine/Source/Runtime/Engine/Public/ActiveSound.h b/Engine/Source/Runtime/Engine/Public/ActiveSound.h
--- a/Engine/Source/Runtime/Engine/Public/ActiveSound.h
+++ b/Engine/Source/Runtime/Engine/Public/ActiveSound.h
@@ -105,8 +105,14 @@ struct FSoundParseParameters
 	// Time offset from beginning of sound to start at
 	float StartTime;
 
-	// At what distance from the source of the sound should spatialization begin
-	float OmniRadius;
+	// At what distance from the source of the sound should spatialization begin 
+	float NonSpatializedRadiusStart;
+
+	// At what distance from the source the sound is fully non-spatialized
+	float NonSpatializedRadiusEnd;
+
+	// Which mode to use for non-spatialized radius
+	ENonSpatializedRadiusSpeakerMapMode NonSpatializedRadiusMode;
 
 	// The distance over which the sound is attenuated
 	float AttenuationDistance;
@@ -227,7 +233,9 @@ struct FSoundParseParameters
 		, InteriorVolumeMultiplier(1.f)
 		, Pitch(1.f)
 		, StartTime(-1.f)
-		, OmniRadius(0.0f)
+		, NonSpatializedRadiusStart(0.0f)
+		, NonSpatializedRadiusEnd(0.0f)
+		, NonSpatializedRadiusMode(ENonSpatializedRadiusSpeakerMapMode::OmniDirectional)
 		, AttenuationDistance(0.0f)
 		, ListenerToSoundDistance(0.0f)
 		, ListenerToSoundDistanceForPanning(0.0f)
diff --git a/Engine/Source/Runtime/Engine/Public/Audio.h b/Engine/Source/Runtime/Engine/Public/Audio.h
--- a/Engine/Source/Runtime/Engine/Public/Audio.h
+++ b/Engine/Source/Runtime/Engine/Public/Audio.h
@@ -369,8 +369,14 @@ public:
 	/** Current location */
 	FVector Location;
 
-	/** At what distance we start transforming into omnidirectional soundsource */
-	float OmniRadius;
+	/** At what distance we start transforming into non-spatialized soundsource */
+	float NonSpatializedRadiusStart;
+
+	/** At what distance we are fully non-spatialized*/
+	float NonSpatializedRadiusEnd;
+
+	/** How we are doing the non-spatialized radius feature. */
+	ENonSpatializedRadiusSpeakerMapMode NonSpatializedRadiusMode;
 
 	/** Amount of spread for 3d multi-channel asset spatialization */
 	float StereoSpread;
