[Quartz] Added ability to query the progress of any musical duration as a value between 0 and 1

diff --git a/Engine/Source/Runtime/AudioMixer/Private/Quartz/AudioMixerClock.cpp b/Engine/Source/Runtime/AudioMixer/Private/Quartz/AudioMixerClock.cpp
--- a/Engine/Source/Runtime/AudioMixer/Private/Quartz/AudioMixerClock.cpp
+++ b/Engine/Source/Runtime/AudioMixer/Private/Quartz/AudioMixerClock.cpp
@@ -57,6 +57,18 @@ namespace Audio
 		return ClockPtr->GetDurationOfQuantizationTypeInSeconds(QuantizationType, Multiplier);
 	}
 
+	float FQuartzClockProxy::GetBeatProgressPercent(
+		const EQuartzCommandQuantization& QuantizationType) const
+	{
+		TSharedPtr<FQuartzClock, ESPMode::ThreadSafe> ClockPtr = ClockWeakPtr.Pin();
+		if (!ClockPtr)
+		{
+			return 0.f;
+		}
+
+		return ClockPtr->GetBeatProgressPercent(QuantizationType);
+	}
+
 	Audio::FQuartzClockTickRate FQuartzClockProxy::GetTickRate() const
 	{
 		TSharedPtr<FQuartzClock, ESPMode::ThreadSafe> ClockPtr = ClockWeakPtr.Pin();
@@ -367,6 +379,8 @@ namespace Audio
 		CachedClockState.TickRate = Metronome.GetTickRate();
 		CachedClockState.TimeStamp = Metronome.GetTimeStamp();
 		CachedClockState.RunTimeInSeconds = (float)Metronome.GetTimeSinceStart();
+
+		Metronome.CalculateDurationPhases(CachedClockState.MusicalDurationPhases);
 	}
 
 	void FQuartzClock::SetSampleRate(float InNewSampleRate)
@@ -540,6 +554,11 @@ namespace Audio
 		}
 	}
 
+	float FQuartzClock::GetBeatProgressPercent(const EQuartzCommandQuantization& QuantizationType) const
+	{
+		return CachedClockState.MusicalDurationPhases[static_cast<int32>(QuantizationType)];
+	}
+
 	FQuartzTransportTimeStamp FQuartzClock::GetCurrentTimestamp()
 	{
 		FScopeLock ScopeLock(&CachedClockStateCritSec);
diff --git a/Engine/Source/Runtime/AudioMixer/Private/Quartz/AudioMixerClockHandle.cpp b/Engine/Source/Runtime/AudioMixer/Private/Quartz/AudioMixerClockHandle.cpp
--- a/Engine/Source/Runtime/AudioMixer/Private/Quartz/AudioMixerClockHandle.cpp
+++ b/Engine/Source/Runtime/AudioMixer/Private/Quartz/AudioMixerClockHandle.cpp
@@ -308,6 +308,16 @@ float UQuartzClockHandle::GetBeatsPerMinute(const UObject* WorldContextObject) c
 	return 0.f;
 }
 
+float UQuartzClockHandle::GetBeatProgressPercent(EQuartzCommandQuantization QuantizationBoundary, float PhaseOffset)
+{
+	if(RawHandle.IsValid())
+	{
+		return FMath::Wrap(PhaseOffset + RawHandle.GetBeatProgressPercent(QuantizationBoundary), 0.f, 1.f);
+	}
+
+	return 0.f;
+}
+
 // todo: un-comment when metronome events support the offset
 // void UQuartzClockHandle::SetNotificationAnticipationAmountInMilliseconds(const UObject* WorldContextObject, UQuartzClockHandle*& ClockHandle, const double Milliseconds)
 // {
diff --git a/Engine/Source/Runtime/AudioMixer/Private/Quartz/QuartzMetronome.cpp b/Engine/Source/Runtime/AudioMixer/Private/Quartz/QuartzMetronome.cpp
--- a/Engine/Source/Runtime/AudioMixer/Private/Quartz/QuartzMetronome.cpp
+++ b/Engine/Source/Runtime/AudioMixer/Private/Quartz/QuartzMetronome.cpp
@@ -66,6 +66,7 @@ namespace Audio
 						}
 
 						FramesLeftInMusicalDuration[DurationType] += PulseDurations[PulseDurationIndex];
+						MusicalDurationsInFrames[DurationType] = PulseDurations[PulseDurationIndex];
 					}
 					while (FramesLeftInMusicalDuration[DurationType] <= 0);
 				}
@@ -284,6 +285,14 @@ namespace Audio
 		return (CurrentTimeStamp.Bars - 1) * NumPerBar + NumInThisBar;
 	}
 
+	void FQuartzMetronome::CalculateDurationPhases(float (&OutPhases)[static_cast<int32>(EQuartzCommandQuantization::Count)]) const
+	{
+		for (int i = 0; i < static_cast<int32>(EQuartzCommandQuantization::Count); ++i)
+		{
+			OutPhases[i] = 1.f - FramesLeftInMusicalDuration[i] / static_cast<float>(MusicalDurationsInFrames[i]);
+		}
+	}
+
 	void FQuartzMetronome::SubscribeToTimeDivision(MetronomeCommandQueuePtr InListenerQueue, EQuartzCommandQuantization InQuantizationBoundary)
 	{
 		MetronomeSubscriptionMatrix[(int32)InQuantizationBoundary].AddUnique(InListenerQueue);
diff --git a/Engine/Source/Runtime/AudioMixer/Public/Quartz/AudioMixerClock.h b/Engine/Source/Runtime/AudioMixer/Public/Quartz/AudioMixerClock.h
--- a/Engine/Source/Runtime/AudioMixer/Public/Quartz/AudioMixerClock.h
+++ b/Engine/Source/Runtime/AudioMixer/Public/Quartz/AudioMixerClock.h
@@ -70,6 +70,8 @@ namespace Audio
 
 		float GetDurationOfQuantizationTypeInSeconds(const EQuartzCommandQuantization& QuantizationType, float Multiplier) const;
 
+		float GetBeatProgressPercent(const EQuartzCommandQuantization& QuantizationType) const;
+
 		// returns false if the clock is not valid or has shut down
 		bool SendCommandToClock(TFunction<void(FQuartzClock*)> InCommand);
 
@@ -154,6 +156,8 @@ namespace Audio
 
 		float GetDurationOfQuantizationTypeInSeconds(const EQuartzCommandQuantization& QuantizationType, float Multiplier);
 
+		float GetBeatProgressPercent(const EQuartzCommandQuantization& QuantizationType) const;
+
 		FQuartzTransportTimeStamp GetCurrentTimestamp();
 
 		float GetEstimatedRunTime();
@@ -217,6 +221,7 @@ namespace Audio
 			FQuartzClockTickRate TickRate;
 			FQuartzTransportTimeStamp TimeStamp;
 			float RunTimeInSeconds;
+			float MusicalDurationPhases[static_cast<int32>(EQuartzCommandQuantization::Count)];
 		} CachedClockState;
 
 		void TickInternal(int32 InNumFramesUntilNextTick, TArray<PendingCommand>& CommandsToTick, int32 FramesOfLatency = 0, int32 FramesOfDelay = 0);
diff --git a/Engine/Source/Runtime/AudioMixer/Public/Quartz/AudioMixerClockHandle.h b/Engine/Source/Runtime/AudioMixer/Public/Quartz/AudioMixerClockHandle.h
--- a/Engine/Source/Runtime/AudioMixer/Public/Quartz/AudioMixerClockHandle.h
+++ b/Engine/Source/Runtime/AudioMixer/Public/Quartz/AudioMixerClockHandle.h
@@ -3,7 +3,9 @@
 #pragma once
 
 #include "CoreMinimal.h"
-#include "Containers/Queue.h"
+#include "UObject/ObjectMacros.h"
+#include "Curves/RichCurve.h"
+#include "Curves/CurveBase.h"
 #include "Sound/QuartzQuantizationUtilities.h"
 #include "Sound/QuartzSubscription.h"
 #include "Quartz/QuartzSubsystem.h"
@@ -11,9 +13,6 @@
 
 #include "AudioMixerClockHandle.generated.h"
 
-
-
-
 /**
  *  This class is a BP / Game thread wrapper around FQuartzClockProxy
  *	(to talk to the underlying clock)
@@ -130,6 +129,13 @@ public:
 	UFUNCTION(BlueprintCallable, Category = "Quantization", meta = (WorldContext = "WorldContextObject", AutoCreateRefTerm = "InDelegate", Keywords = "BPM, Tempo"))
 	float GetBeatsPerMinute(const UObject* WorldContextObject) const;
 
+	/**
+	 * Returns the current progress until the next occurrence of the provided musical duration as a float value from 0 (previous beat) to 1 (next beat).
+	 * This is useful for indexing into curves to animate parameters to musical time.
+	 */
+	UFUNCTION(BlueprintCallable, Category = "Quantization", meta = ( AutoCreateRefTerm = "PhaseOffset", Keywords = "BPM, Tempo"))
+	float GetBeatProgressPercent(EQuartzCommandQuantization QuantizationBoundary = EQuartzCommandQuantization::Beat, float PhaseOffset = 0.f);
+
 	// todo: un-comment when metronome events support the offset
 	// Set how early we would like to receive Metronome (not yet supported) and "About To Start" Delegates. (all other command delegates will execute as normal)
 	// UFUNCTION(BlueprintCallable, Category = "Quantization", meta = (WorldContext = "WorldContextObject", AutoCreateRefTerm = "InDelegate", Keywords = "BPM, Tempo"))
diff --git a/Engine/Source/Runtime/AudioMixer/Public/Quartz/QuartzMetronome.h b/Engine/Source/Runtime/AudioMixer/Public/Quartz/QuartzMetronome.h
--- a/Engine/Source/Runtime/AudioMixer/Public/Quartz/QuartzMetronome.h
+++ b/Engine/Source/Runtime/AudioMixer/Public/Quartz/QuartzMetronome.h
@@ -49,6 +49,8 @@ namespace Audio
 
 		double GetTimeSinceStart() const { return TimeSinceStart; }
 
+		void CalculateDurationPhases(float (&OutPhases)[static_cast<int32>(EQuartzCommandQuantization::Count)]) const;
+
 		// Event Subscription
 		void SubscribeToTimeDivision(MetronomeCommandQueuePtr InListenerQueue, EQuartzCommandQuantization InQuantizationBoundary);
 
