[Quartz] New quantized command to get an event on any desired boundary.

diff --git a/Engine/Source/Runtime/AudioMixer/Private/Quartz/AudioMixerClockHandle.cpp b/Engine/Source/Runtime/AudioMixer/Private/Quartz/AudioMixerClockHandle.cpp
--- a/Engine/Source/Runtime/AudioMixer/Private/Quartz/AudioMixerClockHandle.cpp
+++ b/Engine/Source/Runtime/AudioMixer/Private/Quartz/AudioMixerClockHandle.cpp
@@ -92,6 +92,12 @@ bool UQuartzClockHandle::IsClockRunning(const UObject* WorldContextObject)
 	return RawHandle.IsClockRunning();
 }
 
+void UQuartzClockHandle::NotifyOnQuantizationBoundary(const UObject* WorldContextObject, FQuartzQuantizationBoundary InQuantizationBoundary, const FOnQuartzCommandEventBP& InDelegate)
+{
+	Audio::FQuartzQuantizedCommandInitInfo Data(UQuartzSubsystem::CreateRequestDataForQuantizedNotify(this, InQuantizationBoundary, InDelegate));
+	RawHandle.SendCommandToClock([Data](Audio::FQuartzClock* InClock) mutable { InClock->AddQuantizedCommand(Data); });
+}
+
 float UQuartzClockHandle::GetDurationOfQuantizationTypeInSeconds(const UObject* WorldContextObject, const EQuartzCommandQuantization& QuantizationType, float Multiplier)
 {
 	return RawHandle.GetDurationOfQuantizationTypeInSeconds(QuantizationType, Multiplier);
@@ -119,7 +125,7 @@ void UQuartzClockHandle::StartOtherClock(const UObject* WorldContextObject, FNam
 	RawHandle.SendCommandToClock([Data](Audio::FQuartzClock* InClock) mutable { InClock->AddQuantizedCommand(Data); });
 }
 
-// todo: Move the bulk of these functions to FQuartzTcikableObject once lightweight clock handles are spun up.
+// todo: Move the bulk of these functions to FQuartzTickableObject once lightweight clock handles are spun up.
 void UQuartzClockHandle::SubscribeToQuantizationEvent(const UObject* WorldContextObject, EQuartzCommandQuantization InQuantizationBoundary, const FOnQuartzMetronomeEventBP& OnQuantizationEvent, UQuartzClockHandle*& ClockHandle)
 {
 	ClockHandle = this;
diff --git a/Engine/Source/Runtime/AudioMixer/Private/Quartz/AudioMixerQuantizedCommands.cpp b/Engine/Source/Runtime/AudioMixer/Private/Quartz/AudioMixerQuantizedCommands.cpp
--- a/Engine/Source/Runtime/AudioMixer/Private/Quartz/AudioMixerQuantizedCommands.cpp
+++ b/Engine/Source/Runtime/AudioMixer/Private/Quartz/AudioMixerQuantizedCommands.cpp
@@ -247,4 +247,17 @@ namespace Audio
 		return StartOtherClockName;
 	}
 
+
+
+	TSharedPtr<IQuartzQuantizedCommand> FQuantizedNotify::GetDeepCopyOfDerivedObject() const
+	{
+		return MakeShared<FQuantizedNotify>();
+	}
+
+	static const FName FQuantizedNotifyName("Notify Command");
+	FName FQuantizedNotify::GetCommandName() const
+	{
+		return FQuantizedNotifyName;
+	}
+
 } // namespace Audio
diff --git a/Engine/Source/Runtime/AudioMixer/Private/Quartz/QuartzSubsystem.cpp b/Engine/Source/Runtime/AudioMixer/Private/Quartz/QuartzSubsystem.cpp
--- a/Engine/Source/Runtime/AudioMixer/Private/Quartz/QuartzSubsystem.cpp
+++ b/Engine/Source/Runtime/AudioMixer/Private/Quartz/QuartzSubsystem.cpp
@@ -387,6 +387,31 @@ Audio::FQuartzQuantizedRequestData UQuartzSubsystem::CreateRequestDataForStartOt
 	return CommandInitInfo;
 }
 
+Audio::FQuartzQuantizedRequestData UQuartzSubsystem::CreateRequestDataForQuantizedNotify(UQuartzClockHandle* InClockHandle, const FQuartzQuantizationBoundary& InQuantizationBoundary, const FOnQuartzCommandEventBP& InDelegate)
+{
+	if (!ensure(InClockHandle))
+	{
+		return { };
+	}
+
+	const TSharedPtr<Audio::FQuantizedNotify> NotifyCommandPtr = MakeShared<Audio::FQuantizedNotify>();
+
+	Audio::FQuartzQuantizedRequestData CommandInitInfo;
+
+	CommandInitInfo.ClockName = InClockHandle->GetClockName();
+	CommandInitInfo.QuantizationBoundary = InQuantizationBoundary;
+	CommandInitInfo.QuantizedCommandPtr = NotifyCommandPtr;
+	CommandInitInfo.GameThreadSubscribers.Append(InQuantizationBoundary.GameThreadSubscribers);
+	CommandInitInfo.GameThreadSubscribers.Add(InClockHandle->GetQuartzSubscriber());
+
+	if (InDelegate.IsBound())
+	{
+		CommandInitInfo.GameThreadDelegateID = InClockHandle->AddCommandDelegate(InDelegate);
+	}
+
+	return CommandInitInfo;
+}
+
 Audio::FQuartzClockManager* UQuartzSubsystem::GetClockManager(const UObject* WorldContextObject, bool bUseAudioEngineClockManager)
 {
 	// decide if the clock should be managed by the AudioDevice (audio engine) or the Subsystem (this object)
diff --git a/Engine/Source/Runtime/AudioMixer/Public/Quartz/AudioMixerClockHandle.h b/Engine/Source/Runtime/AudioMixer/Public/Quartz/AudioMixerClockHandle.h
--- a/Engine/Source/Runtime/AudioMixer/Public/Quartz/AudioMixerClockHandle.h
+++ b/Engine/Source/Runtime/AudioMixer/Public/Quartz/AudioMixerClockHandle.h
@@ -63,6 +63,9 @@ public:
 	UFUNCTION(BlueprintCallable, Category = "Quartz Clock", meta = (WorldContext = "WorldContextObject", Keywords = "Transport, Counter"))
 	bool IsClockRunning(const UObject* WorldContextObject);
 
+	UFUNCTION(BlueprintCallable, Category = "Quartz Clock", meta = (WorldContext = "WorldContextObject", AutoCreateRefTerm = "InDelegate", Keywords = "Transport, Counter"))
+	void NotifyOnQuantizationBoundary(const UObject* WorldContextObject, FQuartzQuantizationBoundary InQuantizationBoundary, const FOnQuartzCommandEventBP& InDelegate);
+
 	/** Returns the duration in seconds of the given Quantization Type
 	 *
 	 * @param The Quantization type to measure
diff --git a/Engine/Source/Runtime/AudioMixer/Public/Quartz/AudioMixerQuantizedCommands.h b/Engine/Source/Runtime/AudioMixer/Public/Quartz/AudioMixerQuantizedCommands.h
--- a/Engine/Source/Runtime/AudioMixer/Public/Quartz/AudioMixerQuantizedCommands.h
+++ b/Engine/Source/Runtime/AudioMixer/Public/Quartz/AudioMixerQuantizedCommands.h
@@ -139,6 +139,34 @@ namespace Audio
 		TSharedPtr<FQuartzClock> OwningClockPtr{ nullptr };
 		FName NameOfClockToStart;
 
-	}; // class FQuantizedOtherClockStart 
+	}; // class FQuantizedOtherClockStart
+
+
+	// QuartzQuantizedCommand that basically no-ops, so the game thread can get notified on a musical boundary
+	class AUDIOMIXER_API FQuantizedNotify : public IQuartzQuantizedCommand
+	{
+	public:
+		// ctor
+		FQuantizedNotify() = default;
+
+		// dtor
+		virtual ~FQuantizedNotify() override = default;
+
+		virtual TSharedPtr<IQuartzQuantizedCommand> GetDeepCopyOfDerivedObject() const override;
+
+		virtual bool RequiresAudioDevice() const override { return true; }
+
+		virtual FName GetCommandName() const override;
+
+		virtual EQuartzCommandType GetCommandType() const override { return EQuartzCommandType::Notify; };
+
+	protected:
+		TSharedPtr<FQuartzClock> OwningClockPtr{ nullptr };
+
+		int32 SourceID{ -1 };
+
+		bool bIsCanceled = false;
+
+	}; // class FQuantizedNotify
 
 } // namespace Audio
diff --git a/Engine/Source/Runtime/AudioMixer/Public/Quartz/QuartzSubsystem.h b/Engine/Source/Runtime/AudioMixer/Public/Quartz/QuartzSubsystem.h
--- a/Engine/Source/Runtime/AudioMixer/Public/Quartz/QuartzSubsystem.h
+++ b/Engine/Source/Runtime/AudioMixer/Public/Quartz/QuartzSubsystem.h
@@ -86,6 +86,7 @@ public:
 	static Audio::FQuartzQuantizedRequestData CreateRequestDataForTransportReset(UQuartzClockHandle* InClockHandle, const FQuartzQuantizationBoundary& InQuantizationBoundary, const FOnQuartzCommandEventBP& InDelegate);
 	static Audio::FQuartzQuantizedRequestData CreateRequestDataForStartOtherClock(UQuartzClockHandle* InClockHandle, FName InClockToStart, const FQuartzQuantizationBoundary& InQuantizationBoundary, const FOnQuartzCommandEventBP& InDelegate);
 	static Audio::FQuartzQuantizedRequestData CreateRequestDataForSchedulePlaySound(UQuartzClockHandle* InClockHandle, const FOnQuartzCommandEventBP& InDelegate, const FQuartzQuantizationBoundary& InQuantizationBoundary);
+	static Audio::FQuartzQuantizedRequestData CreateRequestDataForQuantizedNotify(UQuartzClockHandle* InClockHandle, const FQuartzQuantizationBoundary& InQuantizationBoundary, const FOnQuartzCommandEventBP& InDelegate);
 
 	// DEPRECATED HELPERS: non-static versions of the above CreateDataFor...() functions
 	UE_DEPRECATED(5.1, "Use the static (CreateRequestDataFor) version of this function instead")
diff --git a/Engine/Source/Runtime/Engine/Classes/Sound/QuartzQuantizationUtilities.h b/Engine/Source/Runtime/Engine/Classes/Sound/QuartzQuantizationUtilities.h
--- a/Engine/Source/Runtime/Engine/Classes/Sound/QuartzQuantizationUtilities.h
+++ b/Engine/Source/Runtime/Engine/Classes/Sound/QuartzQuantizationUtilities.h
@@ -211,6 +211,7 @@ enum class EQuartzCommandType : uint8
 	TickRateChange UMETA(DisplayName = "Tick Rate Change", ToolTip = "Quantized change of the tick-rate (i.e. BPM change)"),
 	TransportReset UMETA(DisplayName = "Transport Reset", ToolTip = "Quantized reset of the clocks transport (back to time = 0 on the boundary)"),
 	StartOtherClock UMETA(DisplayName = "Start Other Clock", ToolTip = "Quantized start of another clock. Useful for sample accurate synchronization of clocks (i.e. to handle time signature changes)"),
+	Notify UMETA(DisplayName = "Quantized Notify", ToolTip = "Command used only to get delegates for timing information (basically an empty command)"),
 	Custom UMETA(DisplayName = "Custom", ToolTip = "Quantized custom command")
 }; // EQuartzCommandType
 
