Add cpu timing to async render tasks

diff --git a/Engine/Source/Runtime/AudioMixer/Private/AudioMixerSource.cpp b/Engine/Source/Runtime/AudioMixer/Private/AudioMixerSource.cpp
--- a/Engine/Source/Runtime/AudioMixer/Private/AudioMixerSource.cpp
+++ b/Engine/Source/Runtime/AudioMixer/Private/AudioMixerSource.cpp
@@ -1037,6 +1037,8 @@ namespace Audio
 		UpdateChannelMaps();
 
 #if ENABLE_AUDIO_DEBUG
+		UpdateCPUCoreUtilization();
+
 		Audio::FAudioDebugger::DrawDebugInfo(*this);
 #endif // ENABLE_AUDIO_DEBUG
 	}
@@ -1926,6 +1928,20 @@ namespace Audio
 		bPrevAllowedSpatializationSetting = IsSpatializationCVarEnabled();
 	}
 
+#if ENABLE_AUDIO_DEBUG
+	void FMixerSource::UpdateCPUCoreUtilization()
+	{
+		if (MixerSourceVoice)
+		{
+			if (DebugInfo.IsValid())
+			{
+				FScopeLock DebugInfoLock(&DebugInfo->CS);
+				DebugInfo->CPUCoreUtilization = MixerSourceVoice->GetCPUCoreUtilization();
+			}
+		}
+	}
+#endif // if ENABLE_AUDIO_DEBUG
+
 	bool FMixerSource::ComputeMonoChannelMap(Audio::FAlignedFloatBuffer& OutChannelMap)
 	{
 		if (IsUsingObjectBasedSpatialization())
diff --git a/Engine/Source/Runtime/AudioMixer/Private/AudioMixerSource.h b/Engine/Source/Runtime/AudioMixer/Private/AudioMixerSource.h
--- a/Engine/Source/Runtime/AudioMixer/Private/AudioMixerSource.h
+++ b/Engine/Source/Runtime/AudioMixer/Private/AudioMixerSource.h
@@ -88,6 +88,10 @@ namespace Audio
 		/** Updates the channel map of the sound if its a 3d sound.*/
 		void UpdateChannelMaps();
 
+#if ENABLE_AUDIO_DEBUG
+		void UpdateCPUCoreUtilization();
+#endif // ENABLE_AUDIO_DEBUG
+
 		/** Computes the mono-channel map. */
 		bool ComputeMonoChannelMap(Audio::FAlignedFloatBuffer& OutChannelMap);
 
diff --git a/Engine/Source/Runtime/AudioMixer/Private/AudioMixerSourceBuffer.cpp b/Engine/Source/Runtime/AudioMixer/Private/AudioMixerSourceBuffer.cpp
--- a/Engine/Source/Runtime/AudioMixer/Private/AudioMixerSourceBuffer.cpp
+++ b/Engine/Source/Runtime/AudioMixer/Private/AudioMixerSourceBuffer.cpp
@@ -83,6 +83,9 @@ namespace Audio
 		, BufferType(InArgs.Buffer->GetType())
 		, NumPrecacheFrames(InArgs.SoundWave->NumPrecacheFrames)
 		, AuioDeviceID(InArgs.AudioDeviceID)
+#if ENABLE_AUDIO_DEBUG
+		, SampleRate(InArgs.SampleRate)
+#endif // ENABLE_AUDIO_DEBUG
 		, bInitialized(false)
 		, bBufferFinished(false)
 		, bPlayedCachedBuffer(false)
@@ -455,6 +458,10 @@ namespace Audio
 					FDecodeAudioTaskResults TaskResult;
 					AsyncRealtimeAudioTask->GetResult(TaskResult);
 					bIsFinishedOrLooped = TaskResult.bIsFinishedOrLooped;
+#if ENABLE_AUDIO_DEBUG
+					double AudioDuration = static_cast<double>(MONO_PCM_BUFFER_SAMPLES) / FMath::Max(1., static_cast<double>(SampleRate));
+					UpdateCPUCoreUtilization(TaskResult.CPUDuration, AudioDuration);
+#endif // ENABLE_AUDIO_DEBUG
 				}
 				break;
 
@@ -465,6 +472,10 @@ namespace Audio
 
 					SourceVoiceBuffers[CurrentBuffer]->AudioData.SetNum(TaskResult.NumSamplesWritten);
 					bIsFinishedOrLooped = TaskResult.bIsFinished;
+#if ENABLE_AUDIO_DEBUG
+					double AudioDuration = static_cast<double>(TaskResult.NumSamplesWritten) / static_cast<double>(FMath::Max(1, NumChannels * SampleRate));
+					UpdateCPUCoreUtilization(TaskResult.CPUDuration, AudioDuration);
+#endif // ENABLE_AUDIO_DEBUG
 				}
 				break;
 			}
@@ -580,6 +591,36 @@ namespace Audio
 		return bProcedural && SoundGenerator.IsValid() && SoundGenerator->IsFinished();
 	}
 
+#if ENABLE_AUDIO_DEBUG
+	double FMixerSourceBuffer::GetCPUCoreUtilization() const
+	{
+		return CPUCoreUtilization.load(std::memory_order_relaxed);
+	}
+
+	void FMixerSourceBuffer::UpdateCPUCoreUtilization(double InCPUTime, double InAudioTime) 
+	{
+		constexpr double AnalysisTime = 1.0;
+
+		if (InAudioTime > 0.0)
+		{
+			double NewUtilization = InCPUTime / InAudioTime;
+			
+			// Determine smoothing coefficients based upon duration of audio being rendered.
+			const double DigitalCutoff = 1.0 / FMath::Max(1., AnalysisTime / InAudioTime);
+			const double SmoothingBeta = FMath::Clamp(FMath::Exp(-UE_PI * DigitalCutoff), 0.0, 1.0 - UE_DOUBLE_SMALL_NUMBER);
+
+			double PriorUtilization = CPUCoreUtilization.load(std::memory_order_relaxed);
+			
+			// Smooth value if utilization has been initialized.
+			if (PriorUtilization > 0.0)
+			{
+				NewUtilization = (1.0 - SmoothingBeta) * NewUtilization + SmoothingBeta * PriorUtilization;
+			}
+			CPUCoreUtilization.store(NewUtilization, std::memory_order_relaxed);
+		}
+	}
+#endif // ENABLE_AUDIO_DEBUG
+
 	void FMixerSourceBuffer::EnsureAsyncTaskFinishes()
 	{
 		FScopeLock Lock(&DecodeTaskCritSec);
diff --git a/Engine/Source/Runtime/AudioMixer/Private/AudioMixerSourceBuffer.h b/Engine/Source/Runtime/AudioMixer/Private/AudioMixerSourceBuffer.h
--- a/Engine/Source/Runtime/AudioMixer/Private/AudioMixerSourceBuffer.h
+++ b/Engine/Source/Runtime/AudioMixer/Private/AudioMixerSourceBuffer.h
@@ -113,6 +113,9 @@ namespace Audio
 
 		// Returns whether or not generator is finished (returns false if generator is invalid)
 		bool IsGeneratorFinished() const;
+#if ENABLE_AUDIO_DEBUG
+		double GetCPUCoreUtilization() const;
+#endif // ENABLE_AUDIO_DEBUG
 
 	private:
 		FMixerSourceBuffer(FMixerSourceBufferInitArgs& InArgs, TArray<FAudioParameter>&& InDefaultParams);
@@ -142,6 +145,11 @@ namespace Audio
 		int32 NumPrecacheFrames;
 		Audio::FDeviceId AuioDeviceID;
 		TArray<uint8> CachedRealtimeFirstBuffer;
+#if ENABLE_AUDIO_DEBUG
+		int32 SampleRate = 0;
+		std::atomic<double> CPUCoreUtilization = 0.0;
+		void UpdateCPUCoreUtilization(double InCPUTime, double InAudioTime);
+#endif // ENABLE_AUDIO_DEBUG
 
 		mutable FCriticalSection SoundWaveCritSec;
 		mutable FCriticalSection DecodeTaskCritSec;
diff --git a/Engine/Source/Runtime/AudioMixer/Private/AudioMixerSourceDecode.cpp b/Engine/Source/Runtime/AudioMixer/Private/AudioMixerSourceDecode.cpp
--- a/Engine/Source/Runtime/AudioMixer/Private/AudioMixerSourceDecode.cpp
+++ b/Engine/Source/Runtime/AudioMixer/Private/AudioMixerSourceDecode.cpp
@@ -32,6 +32,29 @@ namespace Audio
 
 class FAsyncDecodeWorker : public FNonAbandonableTask
 {
+#if ENABLE_AUDIO_DEBUG
+	struct FScopeDecodeTimer
+	{
+		FScopeDecodeTimer(double* OutResultSeconds)
+		: Result(OutResultSeconds)
+		{
+			StartCycle = FPlatformTime::Cycles64();
+		}
+		~FScopeDecodeTimer()
+		{
+			uint64 EndCycle = FPlatformTime::Cycles64();
+			if (Result)
+			{
+				*Result = static_cast<double>(EndCycle - StartCycle) * FPlatformTime::GetSecondsPerCycle64();
+			}
+		}
+
+		double* Result = nullptr;
+		uint64 StartCycle = 0;
+	};
+
+#endif // if ENABLE_AUDIO_DEBUG
+
 public:
 	FAsyncDecodeWorker(const FHeaderParseAudioTaskData& InTaskData)
 		: HeaderParseAudioData(InTaskData)
@@ -64,6 +87,9 @@ public:
 		{
 			case EAudioTaskType::Procedural:
 			{
+#if ENABLE_AUDIO_DEBUG
+				FScopeDecodeTimer Timer(&ProceduralResult.CPUDuration);
+#endif // if ENABLE_AUDIO_DEBUG
 				QUICK_SCOPE_CYCLE_COUNTER(STAT_FAsyncDecodeWorker_Procedural);
 				if (ProceduralTaskData.SoundGenerator.IsValid())
 				{
@@ -126,6 +152,9 @@ public:
 
 			case EAudioTaskType::Decode:
 			{
+#if ENABLE_AUDIO_DEBUG
+				FScopeDecodeTimer Timer(&DecodeResult.CPUDuration);
+#endif // if ENABLE_AUDIO_DEBUG
 				QUICK_SCOPE_CYCLE_COUNTER(STAT_FAsyncDecodeWorker_Decode);
 				int32 NumChannels = DecodeTaskData.NumChannels;
 				int32 ByteSize = NumChannels * DecodeTaskData.NumFramesToDecode * sizeof(int16);
diff --git a/Engine/Source/Runtime/AudioMixer/Private/AudioMixerSourceDecode.h b/Engine/Source/Runtime/AudioMixer/Private/AudioMixerSourceDecode.h
--- a/Engine/Source/Runtime/AudioMixer/Private/AudioMixerSourceDecode.h
+++ b/Engine/Source/Runtime/AudioMixer/Private/AudioMixerSourceDecode.h
@@ -111,6 +111,10 @@ namespace Audio
 		int32 NumSamplesWritten;
 		bool bIsFinished;
 
+#if ENABLE_AUDIO_DEBUG
+		double CPUDuration = 0.0;
+#endif // if ENABLE_AUDIO_DEBUG
+
 		FProceduralAudioTaskResults()
 			: NumSamplesWritten(0)
 			, bIsFinished(false)
@@ -120,9 +124,14 @@ namespace Audio
 	// Results from decode audio task
 	struct FDecodeAudioTaskResults
 	{
+
 		// Whether or not the audio buffer looped
 		bool bIsFinishedOrLooped;
 
+#if ENABLE_AUDIO_DEBUG
+		double CPUDuration = 0;
+#endif // if ENABLE_AUDIO_DEBUG
+
 		FDecodeAudioTaskResults()
 			: bIsFinishedOrLooped(false)
 		{}
diff --git a/Engine/Source/Runtime/AudioMixer/Private/AudioMixerSourceManager.cpp b/Engine/Source/Runtime/AudioMixer/Private/AudioMixerSourceManager.cpp
--- a/Engine/Source/Runtime/AudioMixer/Private/AudioMixerSourceManager.cpp
+++ b/Engine/Source/Runtime/AudioMixer/Private/AudioMixerSourceManager.cpp
@@ -358,6 +358,9 @@ namespace Audio
 		GameThreadInfo.bNeedsSpeakerMap.AddDefaulted(NumTotalSources);
 		GameThreadInfo.bIsDebugMode.AddDefaulted(NumTotalSources);
 		GameThreadInfo.bIsUsingHRTFSpatializer.AddDefaulted(NumTotalSources);
+#if ENABLE_AUDIO_DEBUG
+		GameThreadInfo.CPUCoreUtilization.AddZeroed(NumTotalSources);
+#endif // if ENABLE_AUDIO_DEBUG
 		GameThreadInfo.FreeSourceIndices.Reset(NumTotalSources);
 		for (int32 i = NumTotalSources - 1; i >= 0; --i)
 		{
@@ -1896,6 +1899,15 @@ namespace Audio
 		return SourceInfos[SourceId].SourceEnvelopeValue;
 	}
 
+#if ENABLE_AUDIO_DEBUG
+	double FMixerSourceManager::GetCPUCoreUtilization(const int32 SourceId) const
+	{
+		AUDIO_MIXER_CHECK_GAME_THREAD(MixerDevice);
+		AUDIO_MIXER_CHECK(SourceId < NumTotalSources);
+		return GameThreadInfo.CPUCoreUtilization[SourceId];
+	}
+#endif // if ENABLE_AUDIO_DEBUG
+
 	bool FMixerSourceManager::IsUsingHRTFSpatializer(const int32 SourceId) const
 	{
 		AUDIO_MIXER_CHECK_GAME_THREAD(MixerDevice);
@@ -1952,6 +1964,11 @@ namespace Audio
 
 				if (ensure(SourceInfo.MixerSourceBuffer.IsValid()))
 				{
+#if ENABLE_AUDIO_DEBUG
+					// Writing to this value is a read/write race condition on the CPUCoreUtilization value. Calling this
+					// out as an acceptable race condition given that it is utilized for debug purposes only. 
+					GameThreadInfo.CPUCoreUtilization[SourceId] = SourceInfo.MixerSourceBuffer->GetCPUCoreUtilization();
+#endif // if ENABLE_AUDIO_DEBUG
 					SourceInfo.MixerSourceBuffer->OnBufferEnd();
 				}
 			}
diff --git a/Engine/Source/Runtime/AudioMixer/Private/AudioMixerSourceManager.h b/Engine/Source/Runtime/AudioMixer/Private/AudioMixerSourceManager.h
--- a/Engine/Source/Runtime/AudioMixer/Private/AudioMixerSourceManager.h
+++ b/Engine/Source/Runtime/AudioMixer/Private/AudioMixerSourceManager.h
@@ -231,6 +231,9 @@ namespace Audio
 
 		int64 GetNumFramesPlayed(const int32 SourceId) const;
 		float GetEnvelopeValue(const int32 SourceId) const;
+#if ENABLE_AUDIO_DEBUG
+		double GetCPUCoreUtilization(const int32 SourceId) const;
+#endif // ENABLE_AUDIO_DEBUG
 		bool IsUsingHRTFSpatializer(const int32 SourceId) const;
 		bool NeedsSpeakerMap(const int32 SourceId) const;
 		void ComputeNextBlockOfSamples();
@@ -571,6 +574,9 @@ namespace Audio
 			TArray<bool> bNeedsSpeakerMap;
 			TArray<bool> bIsDebugMode;
 			TArray<bool> bIsUsingHRTFSpatializer;
+#if ENABLE_AUDIO_DEBUG
+			TArray<double> CPUCoreUtilization;
+#endif // #if ENABLE_AUDIO_DEBUG
 		} GameThreadInfo;
 
 		int32 NumActiveSources;
diff --git a/Engine/Source/Runtime/AudioMixer/Private/AudioMixerSourceVoice.cpp b/Engine/Source/Runtime/AudioMixer/Private/AudioMixerSourceVoice.cpp
--- a/Engine/Source/Runtime/AudioMixer/Private/AudioMixerSourceVoice.cpp
+++ b/Engine/Source/Runtime/AudioMixer/Private/AudioMixerSourceVoice.cpp
@@ -333,6 +333,15 @@ namespace Audio
 		return SourceManager->GetEnvelopeValue(SourceId);
 	}
 
+#if ENABLE_AUDIO_DEBUG
+	double FMixerSourceVoice::GetCPUCoreUtilization() const
+	{
+		AUDIO_MIXER_CHECK_GAME_THREAD(MixerDevice);
+
+		return SourceManager->GetCPUCoreUtilization(SourceId);
+	}
+#endif // ENABLE_AUDIO_DEBUG
+
 	void FMixerSourceVoice::MixOutputBuffers(int32 InNumOutputChannels, const float SendLevel, EMixerSourceSubmixSendStage InSubmixSendStage, FAlignedFloatBuffer& OutWetBuffer) const
 	{
 		AUDIO_MIXER_CHECK_AUDIO_PLAT_THREAD(MixerDevice);
diff --git a/Engine/Source/Runtime/AudioMixer/Private/AudioMixerSourceVoice.h b/Engine/Source/Runtime/AudioMixer/Private/AudioMixerSourceVoice.h
--- a/Engine/Source/Runtime/AudioMixer/Private/AudioMixerSourceVoice.h
+++ b/Engine/Source/Runtime/AudioMixer/Private/AudioMixerSourceVoice.h
@@ -117,6 +117,10 @@ namespace Audio
 		// Retrieves the envelope value of the source.
 		float GetEnvelopeValue() const;
 
+#if ENABLE_AUDIO_DEBUG
+		double GetCPUCoreUtilization() const;
+#endif // ENABLE_AUDIO_DEBUG
+
 		// Mixes the dry and wet buffer audio into the given buffers.
 		void MixOutputBuffers(int32 InNumChannels, const float SendLevel, EMixerSourceSubmixSendStage InSubmixSendStage, FAlignedFloatBuffer& OutWetBuffer) const;
 
diff --git a/Engine/Source/Runtime/Engine/Private/AudioDebug.cpp b/Engine/Source/Runtime/Engine/Private/AudioDebug.cpp
--- a/Engine/Source/Runtime/Engine/Private/AudioDebug.cpp
+++ b/Engine/Source/Runtime/Engine/Private/AudioDebug.cpp
@@ -2107,9 +2107,17 @@ namespace Audio
 			UAudioComponent* AudioComponent = UAudioComponent::GetAudioComponentFromID(WaveInstanceInfo.Value->AudioComponentID);
 			AActor* SoundOwner = AudioComponent ? AudioComponent->GetOwner() : nullptr;
 
-			FString TheString = *FString::Printf(TEXT("%4i.    %6.2f  %s   Owner: %s   SoundClass: %s"),
+			double CPUPercent = 0.;
+			if (WaveInstanceInfo.Key->DebugInfo.IsValid())
+			{
+				FScopeLock Lock(&WaveInstanceInfo.Key->DebugInfo->CS);
+				CPUPercent = 100. * WaveInstanceInfo.Key->DebugInfo->CPUCoreUtilization;
+			}
+			
+			FString TheString = *FString::Printf(TEXT("%4i.    %6.2f  CPU:%5.2f%% %s   Owner: %s   SoundClass: %s"),
 				WaveInstanceInfo.Key->InstanceIndex,
 				WaveInstanceInfo.Key->Volume,
+				CPUPercent,
 				*WaveInstanceInfo.Key->WaveInstanceName.ToString(),
 				SoundOwner ? *SoundOwner->GetName() : TEXT("None"),
 				*WaveInstanceInfo.Key->SoundClassName.ToString());
diff --git a/Engine/Source/Runtime/Engine/Public/Audio.h b/Engine/Source/Runtime/Engine/Public/Audio.h
--- a/Engine/Source/Runtime/Engine/Public/Audio.h
+++ b/Engine/Source/Runtime/Engine/Public/Audio.h
@@ -844,6 +844,9 @@ public:
 		/** Reason why this sound is mute/soloed. */
 		FString MuteSoloReason;
 
+		/** Fraction of a single CPU core used to render audio. */
+		double CPUCoreUtilization = 0;
+
 		/** Basic CS so we can pass this around safely. */
 		FCriticalSection CS;
 	};
