Audio render thread hang detection
#pf 6410e659924023147d581ca0

diff --git a/Engine/Source/Runtime/AudioMixer/Private/AudioMixerSourceManager.cpp b/Engine/Source/Runtime/AudioMixer/Private/AudioMixerSourceManager.cpp
--- a/Engine/Source/Runtime/AudioMixer/Private/AudioMixerSourceManager.cpp
+++ b/Engine/Source/Runtime/AudioMixer/Private/AudioMixerSourceManager.cpp
@@ -17,6 +17,7 @@
 #include "ProfilingDebugging/CsvProfiler.h"
 #include "Async/Async.h"
 #include "ProfilingDebugging/CountersTrace.h"
+#include "HAL/PlatformStackWalk.h"
 #include "Stats/Stats.h"
 #include "Trace/Trace.h"
 
@@ -3278,12 +3279,22 @@ namespace Audio
 		else
 		{
 			UE_LOG(LogAudioMixer, Error,  TEXT("Command buffer grown to %umb, preventing any more adds! Likely cause, the h/w has stopped consuming data."), CurrentBufferSizeInBytes >>20);
+			
+			if (AudioRenderThreadId != INVALID_AUDIO_RENDER_THREAD_ID)
+			{
+				const SIZE_T StackTraceSize = 65536;
+				ANSICHAR StackTrace[StackTraceSize] = {0};
+				FPlatformStackWalk::ThreadStackWalkAndDump(StackTrace, StackTraceSize, 0, AudioRenderThreadId);
+				UE_LOG(LogAudioMixer, Fatal, TEXT("***** ThreadStackWalkAndDump for ThreadId(%lu) ******\n%s"), AudioRenderThreadId, ANSI_TO_TCHAR(StackTrace));
+			}
+			
 		}
 	}
 
 	void FMixerSourceManager::PumpCommandQueue()
 	{
 		TRACE_CPUPROFILER_EVENT_SCOPE(AudioMixerThreadCommands::PumpCommandQueue)
+		AudioRenderThreadId = FPlatformTLS::GetCurrentThreadId();
 		
 		// If we're already triggered, we need to wait for the audio thread to reset it before pumping
 		if (FPlatformProcess::SupportsMultithreading())
diff --git a/Engine/Source/Runtime/AudioMixer/Private/AudioMixerSourceManager.h b/Engine/Source/Runtime/AudioMixer/Private/AudioMixerSourceManager.h
--- a/Engine/Source/Runtime/AudioMixer/Private/AudioMixerSourceManager.h
+++ b/Engine/Source/Runtime/AudioMixer/Private/AudioMixerSourceManager.h
@@ -283,7 +283,8 @@ namespace Audio
 		void AudioMixerThreadMPSCCommand(TFunction<void()> InCommand);
 		
 	private:
-		
+#define INVALID_AUDIO_RENDER_THREAD_ID static_cast<uint32>(-1)
+		uint32 AudioRenderThreadId = INVALID_AUDIO_RENDER_THREAD_ID;
 		void ReleaseSource(const int32 SourceId);
 		void BuildSourceEffectChain(const int32 SourceId, FSoundEffectSourceInitData& InitData, const TArray<FSourceEffectChainEntry>& SourceEffectChain, TArray<TSoundEffectSourcePtr>& OutSourceEffects);
 		void ResetSourceEffectChain(const int32 SourceId);
