Added smoothing to cached quartz clock state queries from game thread.

diff --git a/Engine/Source/Runtime/AudioMixer/Private/Quartz/AudioMixerClock.cpp b/Engine/Source/Runtime/AudioMixer/Private/Quartz/AudioMixerClock.cpp
--- a/Engine/Source/Runtime/AudioMixer/Private/Quartz/AudioMixerClock.cpp
+++ b/Engine/Source/Runtime/AudioMixer/Private/Quartz/AudioMixerClock.cpp
@@ -4,6 +4,7 @@
 #include "Quartz/AudioMixerClockManager.h"
 #include "AudioMixerSourceManager.h"
 #include "Sound/QuartzSubscription.h"
+#include "HAL/UnrealMemory.h" // Memcpy
 
 
 static float HeadlessClockSampleRateCvar = 100000.f;
@@ -380,7 +381,22 @@ namespace Audio
 		CachedClockState.TimeStamp = Metronome.GetTimeStamp();
 		CachedClockState.RunTimeInSeconds = (float)Metronome.GetTimeSinceStart();
 
+		const uint64 TempLastCacheTimestamp = CachedClockState.LastCacheTimestamp;
+		CachedClockState.LastCacheTimestamp = Metronome.GetLastTickCpuCycles64();
+		CachedClockState.LastCacheTimeDelta = CachedClockState.LastCacheTimestamp - TempLastCacheTimestamp;
+
+		// copy previous phases (as temp values)
+		FMemory::Memcpy(CachedClockState.MusicalDurationPhaseDeltas, CachedClockState.MusicalDurationPhases);
+		
+		// update current phases
 		Metronome.CalculateDurationPhases(CachedClockState.MusicalDurationPhases);
+		
+		// convert temp copy to deltas
+		constexpr int32 NumDurations = static_cast<int32>(EQuartzCommandQuantization::Count);
+		for(int32 i = 0; i < NumDurations; ++i)
+		{
+			CachedClockState.MusicalDurationPhaseDeltas[i] = FMath::Wrap(CachedClockState.MusicalDurationPhases[i] - CachedClockState.MusicalDurationPhaseDeltas[i], 0.f, 1.f);
+		}
 	}
 
 	void FQuartzClock::SetSampleRate(float InNewSampleRate)
@@ -556,7 +572,18 @@ namespace Audio
 
 	float FQuartzClock::GetBeatProgressPercent(const EQuartzCommandQuantization& QuantizationType) const
 	{
-		return CachedClockState.MusicalDurationPhases[static_cast<int32>(QuantizationType)];
+		if(CachedClockState.LastCacheTimeDelta == 0)
+		{
+			return CachedClockState.MusicalDurationPhases[static_cast<int32>(QuantizationType)];
+		}
+
+		// anticipate beat progress based on the amount of wall clock time that has passed since the last audio engine update
+		const float LastPhase = CachedClockState.MusicalDurationPhases[static_cast<int32>(QuantizationType)];
+		const float PhaseDelta = CachedClockState.MusicalDurationPhaseDeltas[static_cast<int32>(QuantizationType)];
+		const uint32 CyclesSinceLastTick = FPlatformTime::Cycles() - CachedClockState.LastCacheTimestamp;
+		const float EstimatedPercentToNextTick = static_cast<float>(CyclesSinceLastTick) / static_cast<float>(CachedClockState.LastCacheTimeDelta);
+
+		return LastPhase + PhaseDelta * EstimatedPercentToNextTick;
 	}
 
 	FQuartzTransportTimeStamp FQuartzClock::GetCurrentTimestamp()
diff --git a/Engine/Source/Runtime/AudioMixer/Private/Quartz/AudioMixerClockHandle.cpp b/Engine/Source/Runtime/AudioMixer/Private/Quartz/AudioMixerClockHandle.cpp
--- a/Engine/Source/Runtime/AudioMixer/Private/Quartz/AudioMixerClockHandle.cpp
+++ b/Engine/Source/Runtime/AudioMixer/Private/Quartz/AudioMixerClockHandle.cpp
@@ -314,10 +314,17 @@ float UQuartzClockHandle::GetBeatsPerMinute(const UObject* WorldContextObject) c
 	return 0.f;
 }
 
-float UQuartzClockHandle::GetBeatProgressPercent(EQuartzCommandQuantization QuantizationBoundary, float PhaseOffset)
+float UQuartzClockHandle::GetBeatProgressPercent(EQuartzCommandQuantization QuantizationBoundary, float PhaseOffset, float MsOffset)
 {
 	if(RawHandle.IsValid())
 	{
+		constexpr float ToMilliseconds = 1000.f;
+	    const float MsInQuantizationType = ToMilliseconds * RawHandle.GetDurationOfQuantizationTypeInSeconds(QuantizationBoundary, 1.f);
+	    if(!FMath::IsNearlyZero(MsInQuantizationType))
+	    {
+		    PhaseOffset += MsOffset / MsInQuantizationType;
+	    }
+
 		return FMath::Wrap(PhaseOffset + RawHandle.GetBeatProgressPercent(QuantizationBoundary), 0.f, 1.f);
 	}
 
diff --git a/Engine/Source/Runtime/AudioMixer/Private/Quartz/QuartzMetronome.cpp b/Engine/Source/Runtime/AudioMixer/Private/Quartz/QuartzMetronome.cpp
--- a/Engine/Source/Runtime/AudioMixer/Private/Quartz/QuartzMetronome.cpp
+++ b/Engine/Source/Runtime/AudioMixer/Private/Quartz/QuartzMetronome.cpp
@@ -23,7 +23,8 @@ namespace Audio
 
 	void FQuartzMetronome::Tick(int32 InNumSamples, int32 FramesOfLatency)
 	{
-		// TODO: update latency
+		LastTickCpuCycles64 = FPlatformTime::Cycles64();
+		
 		static bool bHasWarned = false;
 		if (!bHasWarned && (MusicalDurationsInFrames[EQuartzCommandQuantization::ThirtySecondNote] < InNumSamples))
 		{
diff --git a/Engine/Source/Runtime/AudioMixer/Public/Quartz/AudioMixerClock.h b/Engine/Source/Runtime/AudioMixer/Public/Quartz/AudioMixerClock.h
--- a/Engine/Source/Runtime/AudioMixer/Public/Quartz/AudioMixerClock.h
+++ b/Engine/Source/Runtime/AudioMixer/Public/Quartz/AudioMixerClock.h
@@ -221,7 +221,11 @@ namespace Audio
 			FQuartzClockTickRate TickRate;
 			FQuartzTransportTimeStamp TimeStamp;
 			float RunTimeInSeconds;
-			float MusicalDurationPhases[static_cast<int32>(EQuartzCommandQuantization::Count)];
+			float MusicalDurationPhases[static_cast<int32>(EQuartzCommandQuantization::Count)] { 0 };
+			float MusicalDurationPhaseDeltas[static_cast<int32>(EQuartzCommandQuantization::Count)] { 0 };
+			uint64 LastCacheTimestamp = 0;
+			uint64 LastCacheTimeDelta = 0;
+			
 		} CachedClockState;
 
 		void TickInternal(int32 InNumFramesUntilNextTick, TArray<PendingCommand>& CommandsToTick, int32 FramesOfLatency = 0, int32 FramesOfDelay = 0);
diff --git a/Engine/Source/Runtime/AudioMixer/Public/Quartz/AudioMixerClockHandle.h b/Engine/Source/Runtime/AudioMixer/Public/Quartz/AudioMixerClockHandle.h
--- a/Engine/Source/Runtime/AudioMixer/Public/Quartz/AudioMixerClockHandle.h
+++ b/Engine/Source/Runtime/AudioMixer/Public/Quartz/AudioMixerClockHandle.h
@@ -137,7 +137,7 @@ public:
 	 * This is useful for indexing into curves to animate parameters to musical time.
 	 */
 	UFUNCTION(BlueprintCallable, Category = "Quantization", meta = ( AutoCreateRefTerm = "PhaseOffset", Keywords = "BPM, Tempo"))
-	float GetBeatProgressPercent(EQuartzCommandQuantization QuantizationBoundary = EQuartzCommandQuantization::Beat, float PhaseOffset = 0.f);
+	float GetBeatProgressPercent(EQuartzCommandQuantization QuantizationBoundary = EQuartzCommandQuantization::Beat, float PhaseOffset = 0.f, float MsOffset = 0.f);
 
 	// todo: un-comment when metronome events support the offset
 	// Set how early we would like to receive Metronome (not yet supported) and "About To Start" Delegates. (all other command delegates will execute as normal)
diff --git a/Engine/Source/Runtime/AudioMixer/Public/Quartz/QuartzMetronome.h b/Engine/Source/Runtime/AudioMixer/Public/Quartz/QuartzMetronome.h
--- a/Engine/Source/Runtime/AudioMixer/Public/Quartz/QuartzMetronome.h
+++ b/Engine/Source/Runtime/AudioMixer/Public/Quartz/QuartzMetronome.h
@@ -49,6 +49,8 @@ namespace Audio
 
 		double GetTimeSinceStart() const { return TimeSinceStart; }
 
+		uint64 GetLastTickCpuCycles64() const { return LastTickCpuCycles64; }
+
 		void CalculateDurationPhases(float (&OutPhases)[static_cast<int32>(EQuartzCommandQuantization::Count)]) const;
 
 		// Event Subscription
@@ -74,6 +76,8 @@ namespace Audio
 
 		float CountNumSubdivisionsSinceStart(EQuartzCommandQuantization InSubdivision) const;
 
+		uint64 LastTickCpuCycles64{ 0 };
+
 		int32 ListenerFlags{ 0 };
 
 		FQuartzTransportTimeStamp CurrentTimeStamp;
