[Quartz] Follow up tweaks for GT smoothing in GetBeatProgressPercent()

diff --git a/Engine/Source/Runtime/AudioMixer/Private/Quartz/AudioMixerClock.cpp b/Engine/Source/Runtime/AudioMixer/Private/Quartz/AudioMixerClock.cpp
--- a/Engine/Source/Runtime/AudioMixer/Private/Quartz/AudioMixerClock.cpp
+++ b/Engine/Source/Runtime/AudioMixer/Private/Quartz/AudioMixerClock.cpp
@@ -381,9 +381,9 @@ namespace Audio
 		CachedClockState.TimeStamp = Metronome.GetTimeStamp();
 		CachedClockState.RunTimeInSeconds = (float)Metronome.GetTimeSinceStart();
 
-		const uint64 TempLastCacheTimestamp = CachedClockState.LastCacheTimestamp;
-		CachedClockState.LastCacheTimestamp = Metronome.GetLastTickCpuCycles64();
-		CachedClockState.LastCacheTimeDelta = CachedClockState.LastCacheTimestamp - TempLastCacheTimestamp;
+		const uint64 TempLastCacheTimestamp = CachedClockState.LastCacheTickCpuCycles64;
+		CachedClockState.LastCacheTickCpuCycles64 = Metronome.GetLastTickCpuCycles64();
+		CachedClockState.LastCacheTickDeltaCpuCycles64 = CachedClockState.LastCacheTickCpuCycles64 - TempLastCacheTimestamp;
 
 		// copy previous phases (as temp values)
 		FMemory::Memcpy(CachedClockState.MusicalDurationPhaseDeltas, CachedClockState.MusicalDurationPhases);
@@ -572,7 +572,7 @@ namespace Audio
 
 	float FQuartzClock::GetBeatProgressPercent(const EQuartzCommandQuantization& QuantizationType) const
 	{
-		if(CachedClockState.LastCacheTimeDelta == 0)
+		if(CachedClockState.LastCacheTickDeltaCpuCycles64 == 0)
 		{
 			return CachedClockState.MusicalDurationPhases[static_cast<int32>(QuantizationType)];
 		}
@@ -580,8 +580,8 @@ namespace Audio
 		// anticipate beat progress based on the amount of wall clock time that has passed since the last audio engine update
 		const float LastPhase = CachedClockState.MusicalDurationPhases[static_cast<int32>(QuantizationType)];
 		const float PhaseDelta = CachedClockState.MusicalDurationPhaseDeltas[static_cast<int32>(QuantizationType)];
-		const uint32 CyclesSinceLastTick = FPlatformTime::Cycles() - CachedClockState.LastCacheTimestamp;
-		const float EstimatedPercentToNextTick = static_cast<float>(CyclesSinceLastTick) / static_cast<float>(CachedClockState.LastCacheTimeDelta);
+		const uint64 CyclesSinceLastTick = FPlatformTime::Cycles64() - CachedClockState.LastCacheTickCpuCycles64;
+		const float EstimatedPercentToNextTick = static_cast<float>(CyclesSinceLastTick) / static_cast<float>(CachedClockState.LastCacheTickDeltaCpuCycles64);
 
 		return LastPhase + PhaseDelta * EstimatedPercentToNextTick;
 	}
diff --git a/Engine/Source/Runtime/AudioMixer/Public/Quartz/AudioMixerClock.h b/Engine/Source/Runtime/AudioMixer/Public/Quartz/AudioMixerClock.h
--- a/Engine/Source/Runtime/AudioMixer/Public/Quartz/AudioMixerClock.h
+++ b/Engine/Source/Runtime/AudioMixer/Public/Quartz/AudioMixerClock.h
@@ -223,8 +223,8 @@ namespace Audio
 			float RunTimeInSeconds;
 			float MusicalDurationPhases[static_cast<int32>(EQuartzCommandQuantization::Count)] { 0 };
 			float MusicalDurationPhaseDeltas[static_cast<int32>(EQuartzCommandQuantization::Count)] { 0 };
-			uint64 LastCacheTimestamp = 0;
-			uint64 LastCacheTimeDelta = 0;
+			uint64 LastCacheTickCpuCycles64 = 0;
+			uint64 LastCacheTickDeltaCpuCycles64 = 0;
 			
 		} CachedClockState;
 
diff --git a/Engine/Source/Runtime/AudioMixer/Public/Quartz/AudioMixerClockHandle.h b/Engine/Source/Runtime/AudioMixer/Public/Quartz/AudioMixerClockHandle.h
--- a/Engine/Source/Runtime/AudioMixer/Public/Quartz/AudioMixerClockHandle.h
+++ b/Engine/Source/Runtime/AudioMixer/Public/Quartz/AudioMixerClockHandle.h
@@ -135,6 +135,7 @@ public:
 	/**
 	 * Returns the current progress until the next occurrence of the provided musical duration as a float value from 0 (previous beat) to 1 (next beat).
 	 * This is useful for indexing into curves to animate parameters to musical time.
+	 * Ms and Phase offsets are combined internally.
 	 */
 	UFUNCTION(BlueprintCallable, Category = "Quantization", meta = ( AutoCreateRefTerm = "PhaseOffset", Keywords = "BPM, Tempo"))
 	float GetBeatProgressPercent(EQuartzCommandQuantization QuantizationBoundary = EQuartzCommandQuantization::Beat, float PhaseOffset = 0.f, float MsOffset = 0.f);
