Fixed mismatched copy constructor/assignment operator for FQuantizedQueueCommand.

diff --git a/Engine/Source/Runtime/AudioMixer/Public/Quartz/AudioMixerQuantizedCommands.h b/Engine/Source/Runtime/AudioMixer/Public/Quartz/AudioMixerQuantizedCommands.h
--- a/Engine/Source/Runtime/AudioMixer/Public/Quartz/AudioMixerQuantizedCommands.h
+++ b/Engine/Source/Runtime/AudioMixer/Public/Quartz/AudioMixerQuantizedCommands.h
@@ -60,11 +60,6 @@ namespace Audio
 
 		FQuantizedQueueCommand() {}
 
-		FQuantizedQueueCommand(const FQuantizedQueueCommand& Other)
-			: OwningClockPtr(Other.OwningClockPtr)
-			, AudioComponentData(Other.AudioComponentData)
-		{}
-
 	private:
 		TSharedPtr<FQuartzClock> OwningClockPtr{ nullptr };
 
