Use HashCombine of AudioComponent ID + NodeWaveInstanceHash as address for parameter routing instead of just AudioComponentID alone. Allows components that play multiple voices to route parameters to each voice.

diff --git a/Engine/Source/Runtime/AudioMixer/Private/AudioMixerSource.cpp b/Engine/Source/Runtime/AudioMixer/Private/AudioMixerSource.cpp
--- a/Engine/Source/Runtime/AudioMixer/Private/AudioMixerSource.cpp
+++ b/Engine/Source/Runtime/AudioMixer/Private/AudioMixerSource.cpp
@@ -1117,12 +1117,14 @@ namespace Audio
 
 		// Active sound instance ID is the audio component ID of active sound.
 		uint64 InstanceID = 0;
+		uint32 PlayOrder = 0;
 		bool bActiveSoundIsPreviewSound = false;
 		TArray<FAudioParameter> DefaultParameters;
 		FActiveSound* ActiveSound = WaveInstance->ActiveSound;
 		if (ActiveSound)
 		{
 			InstanceID = ActiveSound->GetAudioComponentID();
+			PlayOrder = ActiveSound->GetPlayOrder();
 			bActiveSoundIsPreviewSound = ActiveSound->bIsPreviewSound;
 			if (Audio::IParameterTransmitter* Transmitter = ActiveSound->GetTransmitter())
 			{
@@ -1132,7 +1134,7 @@ namespace Audio
 
 		FMixerSourceBufferInitArgs BufferInitArgs;
 		BufferInitArgs.AudioDeviceID = AudioDevice->DeviceID;
-		BufferInitArgs.InstanceID = InstanceID;
+		BufferInitArgs.InstanceID = GetTransmitterID(InstanceID, WaveInstance->WaveInstanceHash, PlayOrder);
 		BufferInitArgs.SampleRate = AudioDevice->GetSampleRate();
 		BufferInitArgs.AudioMixerNumOutputFrames = MixerDevice->GetNumOutputFrames();
 		BufferInitArgs.Buffer = MixerBuffer;
diff --git a/Engine/Source/Runtime/Engine/Private/AudioDevice.cpp b/Engine/Source/Runtime/Engine/Private/AudioDevice.cpp
--- a/Engine/Source/Runtime/Engine/Private/AudioDevice.cpp
+++ b/Engine/Source/Runtime/Engine/Private/AudioDevice.cpp
@@ -248,6 +248,11 @@ namespace Audio
 		ensureMsgf(false, TEXT("Unknown runtime sound asset format type."));
 		return nullptr;
 	}
+
+	uint64 GetTransmitterID(uint64 ComponentID, UPTRINT WaveInstanceHash, uint32 PlayOrder)
+	{
+		return HashCombineFast(static_cast<uint32>(ComponentID % TNumericLimits<uint32>::Max()), PlayOrder + WaveInstanceHash);
+	}
 }
 
 namespace AudioDeviceUtils
@@ -5131,7 +5136,7 @@ void FAudioDevice::AddNewActiveSoundInternal(const FActiveSound& InNewActiveSoun
 		{
 			Audio::FParameterTransmitterInitParams TransmitterInitParams
 			{
-				OutActiveSound.GetAudioComponentID(),
+				Audio::GetTransmitterID(OutActiveSound.GetAudioComponentID(), 0, OutActiveSound.GetPlayOrder()),
 				GetSampleRate(),
 				MoveTemp(DefaultParams),
 				DeviceID
diff --git a/Engine/Source/Runtime/Engine/Private/SoundNodeWavePlayer.cpp b/Engine/Source/Runtime/Engine/Private/SoundNodeWavePlayer.cpp
--- a/Engine/Source/Runtime/Engine/Private/SoundNodeWavePlayer.cpp
+++ b/Engine/Source/Runtime/Engine/Private/SoundNodeWavePlayer.cpp
@@ -146,7 +146,7 @@ void USoundNodeWavePlayer::ParseNodes( FAudioDevice* AudioDevice, const UPTRINT
 			
 			Audio::FParameterTransmitterInitParams Params;
 			Params.DefaultParams = ActiveSound.GetTransmitter()->GetParameters();
-			Params.InstanceID = ActiveSound.GetAudioComponentID();
+			Params.InstanceID = Audio::GetTransmitterID(ActiveSound.GetAudioComponentID(), NodeWaveInstanceHash, ActiveSound.GetPlayOrder()); 
 			Params.SampleRate = AudioDevice->GetSampleRate();
 			Params.AudioDeviceID = AudioDevice->DeviceID;
 
diff --git a/Engine/Source/Runtime/Engine/Public/AudioDevice.h b/Engine/Source/Runtime/Engine/Public/AudioDevice.h
--- a/Engine/Source/Runtime/Engine/Public/AudioDevice.h
+++ b/Engine/Source/Runtime/Engine/Public/AudioDevice.h
@@ -97,6 +97,9 @@ namespace Audio
 
 	/** Returns a decoder for the sound assets's given compression type. Will return nullptr if the compression type is platform-dependent. */
 	ENGINE_API ICompressedAudioInfo* CreateSoundAssetDecoder(const FName& InRuntimeFormat);
+
+	/** Creates an ID for use by Parameter Transmitters that can differentiate between multiple voices playing on the same Audio Component. */
+	ENGINE_API uint64 GetTransmitterID(uint64 ComponentID, UPTRINT WaveInstanceHash, uint32 PlayOrder);
 }
 
 /**
