Add support for copying wave file metadata chunks to USoundWave payload when converting bit depth of source file.

diff --git a/Engine/Source/Editor/AudioEditor/Private/Factories/SoundFactory.cpp b/Engine/Source/Editor/AudioEditor/Private/Factories/SoundFactory.cpp
--- a/Engine/Source/Editor/AudioEditor/Private/Factories/SoundFactory.cpp
+++ b/Engine/Source/Editor/AudioEditor/Private/Factories/SoundFactory.cpp
@@ -374,8 +374,6 @@ UObject* USoundFactory::CreateObject
 			// Attempt to convert to 16 bit audio
 			if (Audio::ConvertAudioToWav(RawWaveData, ConvertedRawWaveData))
 			{
-				// Icky, Reencoding with SNDFILE will strip the timecode info, so back it up.
-				auto CachedTimeCodeInfo = MoveTemp(WaveInfo.TimecodeInfo);
 				WaveInfo = FWaveModInfo();				
 				if (!WaveInfo.ReadWaveInfo(ConvertedRawWaveData.GetData(), ConvertedRawWaveData.Num(), &ErrorMessage))
 				{
@@ -383,7 +381,6 @@ UObject* USoundFactory::CreateObject
 					GEditor->GetEditorSubsystem<UImportSubsystem>()->BroadcastAssetPostImport(this, nullptr);
 					return nullptr;
 				}
-				WaveInfo.TimecodeInfo = MoveTemp(CachedTimeCodeInfo);
 			}
 
 			// Copy over the data
diff --git a/Engine/Source/Runtime/AudioMixer/Private/SoundFileIO/SoundFile.h b/Engine/Source/Runtime/AudioMixer/Private/SoundFileIO/SoundFile.h
--- a/Engine/Source/Runtime/AudioMixer/Private/SoundFileIO/SoundFile.h
+++ b/Engine/Source/Runtime/AudioMixer/Private/SoundFileIO/SoundFile.h
@@ -60,7 +60,91 @@ namespace Audio
 			return MoveTemp(Default);
 		}
 	};
+	
+	/**
+	 * FSoundFileChunkInfo which maps to libsndfile SF_CHUNK_INFO struct.
+	 */
+
+	struct FSoundFileChunkInfo
+	{
+		/** Chunk Id **/
+		ANSICHAR ChunkId[64];
+
+		/** Size of the Chunk Id **/
+		uint32 ChunkIdSize = 0;
+
+		/** Size of the data in this chunk **/
+		uint32 DataLength = 0;
+
+		/** Pointer to chunk data **/
+		void* DataPtr = nullptr;
+	};
 
+	/**
+	 * FSoundFileChunkInfoWrapper wraps FSoundFileChunkInfo and manages
+	 * chunk data memory.
+	 */
+	class FSoundFileChunkInfoWrapper
+	{
+	public:
+		FSoundFileChunkInfoWrapper() = default;
+		~FSoundFileChunkInfoWrapper() = default;
+		FSoundFileChunkInfoWrapper(const FSoundFileChunkInfoWrapper& Other) = delete;
+		FSoundFileChunkInfoWrapper& operator=(const FSoundFileChunkInfoWrapper& Other) = delete;
+
+		FSoundFileChunkInfoWrapper(FSoundFileChunkInfoWrapper&& Other) noexcept
+		{
+			if (Other.ChunkInfo.ChunkIdSize)
+			{
+				FMemory::Memcpy(ChunkInfo.ChunkId, Other.ChunkInfo.ChunkId, sizeof(ChunkInfo.ChunkId));
+			}
+			ChunkInfo.ChunkIdSize = Other.ChunkInfo.ChunkIdSize;
+			ChunkInfo.DataLength = Other.ChunkInfo.DataLength;
+			ChunkInfo.DataPtr = Other.ChunkInfo.DataPtr;
+
+			ChunkData = MoveTemp(Other.ChunkData);
+		}
+
+		FSoundFileChunkInfoWrapper& operator=(FSoundFileChunkInfoWrapper&& Other) noexcept
+		{
+			if (Other.ChunkInfo.ChunkIdSize)
+			{
+				FMemory::Memcpy(ChunkInfo.ChunkId, Other.ChunkInfo.ChunkId, sizeof(ChunkInfo.ChunkId));
+			}
+			ChunkInfo.ChunkIdSize = Other.ChunkInfo.ChunkIdSize;
+			ChunkInfo.DataLength = Other.ChunkInfo.DataLength;
+			ChunkInfo.DataPtr = Other.ChunkInfo.DataPtr;
+
+			ChunkData = MoveTemp(Other.ChunkData);
+			
+			return *this;
+		}
+
+		void AllocateChunkData()
+		{
+			if (ChunkInfo.DataLength > 0 && ensure(ChunkInfo.DataPtr == nullptr))
+			{
+				ChunkData = MakeUnique<uint8[]>(ChunkInfo.DataLength);
+				ChunkInfo.DataPtr = ChunkData.Get();
+			}
+		}
+
+		FSoundFileChunkInfo* GetPtr()
+		{
+			return &ChunkInfo;
+		}
+
+		const FSoundFileChunkInfo* GetPtr() const
+		{
+			return &ChunkInfo;
+		}
+
+	private:
+		FSoundFileChunkInfo	ChunkInfo;
+		TUniquePtr<uint8[]>	ChunkData;
+	};
+	typedef TArray<FSoundFileChunkInfoWrapper> FSoundFileChunkArray;
+	
 	/**
 	 * ISoundFile
 	 */
@@ -93,6 +177,7 @@ namespace Audio
 		virtual ESoundFileError::Type ReadSamples(float* DataPtr, SoundFileCount NumSamples, SoundFileCount& OutNumSamplesRead) = 0;
 		virtual ESoundFileError::Type ReadSamples(double* DataPtr, SoundFileCount NumSamples, SoundFileCount& OutNumSamplesRead) = 0;
 		virtual ESoundFileError::Type GetDescription(FSoundFileDescription& OutputDescription, TArray<ESoundFileChannelMap::Type>& OutChannelMap) = 0;
+		virtual ESoundFileError::Type GetOptionalChunks(FSoundFileChunkArray& OutChunkInfoArray) = 0;
 	};
 
 	class ISoundFileWriter
@@ -108,5 +193,6 @@ namespace Audio
 		virtual ESoundFileError::Type WriteSamples(const float* DataPtr, SoundFileCount NumSamples, SoundFileCount& OutNumSampleWritten) = 0;
 		virtual ESoundFileError::Type WriteSamples(const double* DataPtr, SoundFileCount NumSamples, SoundFileCount& OutNumSampleWritten) = 0;
 		virtual ESoundFileError::Type GetData(TArray<uint8>** OutData) = 0;
+		virtual ESoundFileError::Type WriteOptionalChunks(const FSoundFileChunkArray& ChunkInfoArray) = 0;
 	};
 } // namespace Audio
\ No newline at end of file
diff --git a/Engine/Source/Runtime/AudioMixer/Private/SoundFileIO/SoundFileIO.cpp b/Engine/Source/Runtime/AudioMixer/Private/SoundFileIO/SoundFileIO.cpp
--- a/Engine/Source/Runtime/AudioMixer/Private/SoundFileIO/SoundFileIO.cpp
+++ b/Engine/Source/Runtime/AudioMixer/Private/SoundFileIO/SoundFileIO.cpp
@@ -155,6 +155,22 @@ namespace Audio
 		Error = SoundFileWriter->WriteSamples((const float*)ProcessBuffer.GetData(), InputSamplesRead, SamplesWritten);
 		check(Error == ESoundFileError::Type::NONE);
 
+		// The required chunks ('fmt ' and 'data') have been written, now copy over the optional chunks.
+		FSoundFileChunkArray OptionalChunks;
+		Error = InputSoundDataReader->GetOptionalChunks(OptionalChunks);
+		if (Error != ESoundFileError::Type::NONE)
+		{
+			UE_LOG(LogAudioMixer, Error, TEXT("Error encountered while reading optional chunk data...skipping"));
+		}
+		else
+		{
+			Error = SoundFileWriter->WriteOptionalChunks(OptionalChunks);
+			if (Error != ESoundFileError::Type::NONE)
+			{
+				UE_LOG(LogAudioMixer, Error, TEXT("Error encountered while writing optional chunk data...skipping"));
+			}
+		}
+		
 		// Release the sound file handles as soon as we finished converting the file
 		InputSoundDataReader->Release();
 		SoundFileWriter->Release();
diff --git a/Engine/Source/Runtime/AudioMixer/Private/SoundFileIO/SoundFileIOEnums.h b/Engine/Source/Runtime/AudioMixer/Private/SoundFileIO/SoundFileIOEnums.h
--- a/Engine/Source/Runtime/AudioMixer/Private/SoundFileIO/SoundFileIOEnums.h
+++ b/Engine/Source/Runtime/AudioMixer/Private/SoundFileIO/SoundFileIOEnums.h
@@ -26,6 +26,7 @@ namespace Audio
 			ALREADY_INITIALIZED,
 			LOADING,
 			INVALID_STATE,
+			INVALID_CHUNK,
 			UNKNOWN
 		};
 
diff --git a/Engine/Source/Runtime/AudioMixer/Private/SoundFileIO/SoundFileIOManagerImpl.cpp b/Engine/Source/Runtime/AudioMixer/Private/SoundFileIO/SoundFileIOManagerImpl.cpp
--- a/Engine/Source/Runtime/AudioMixer/Private/SoundFileIO/SoundFileIOManagerImpl.cpp
+++ b/Engine/Source/Runtime/AudioMixer/Private/SoundFileIO/SoundFileIOManagerImpl.cpp
@@ -4,6 +4,7 @@
 
 #include "CoreMinimal.h"
 
+#include "Audio.h"
 #include "AudioMixer.h"
 #include "HAL/PlatformProcess.h"
 #include "HAL/PlatformFileManager.h"
@@ -15,6 +16,7 @@
 namespace Audio
 {
 	typedef struct SoundFileHandleOpaque LibSoundFileHandle;
+	typedef struct SoundFileChunkIteratorOpaque LibSoundFileChunkIterator;
 
 	//	// Virtual Sound File Function Pointers
 	typedef SoundFileCount(*VirtualSoundFileGetLengthFuncPtr)(void* UserData);
@@ -57,6 +59,11 @@ namespace Audio
 	typedef SoundFileCount(*SoundFileReadSamplesDoubleFuncPtr)(LibSoundFileHandle* FileHandle, double* Buffer, SoundFileCount NumSamples);
 	typedef SoundFileCount(*SoundFileWriteSamplesFloatFuncPtr)(LibSoundFileHandle* FileHandle, const float* Buffer, SoundFileCount NumSamples);
 	typedef SoundFileCount(*SoundFileWriteSamplesDoubleFuncPtr)(LibSoundFileHandle* FileHandle, const double* Buffer, SoundFileCount NumSamples);
+	typedef int32(*SoundFileGetChunkSizeFuncPtr)(const LibSoundFileChunkIterator* ChunkIterator, FSoundFileChunkInfo* ChunkInfo);
+	typedef int32(*SoundFileGetChunkDataFuncPtr)(const LibSoundFileChunkIterator* ChunkIterator, FSoundFileChunkInfo* ChunkInfo);
+	typedef LibSoundFileChunkIterator*(*SoundFileGetChunkIteratorFuncPtr)(LibSoundFileHandle* FileHandle, const FSoundFileChunkInfo* ChunkInfo);
+	typedef LibSoundFileChunkIterator*(*SoundFileNextChunkIteratorFuncPtr)(LibSoundFileChunkIterator* ChunkIterator);
+	typedef int32(*SoundFileSetChunkFuncPtr)(LibSoundFileHandle* FileHandle, const FSoundFileChunkInfo* ChunkInfo);
 
 	SoundFileOpenFuncPtr SoundFileOpen = nullptr;
 	SoundFileOpenVirtualFuncPtr SoundFileOpenVirtual = nullptr;
@@ -76,6 +83,11 @@ namespace Audio
 	SoundFileReadSamplesDoubleFuncPtr SoundFileReadSamplesDouble = nullptr;
 	SoundFileWriteSamplesFloatFuncPtr SoundFileWriteSamplesFloat = nullptr;
 	SoundFileWriteSamplesDoubleFuncPtr SoundFileWriteSamplesDouble = nullptr;
+	SoundFileGetChunkSizeFuncPtr SoundFileGetChunkSize = nullptr;
+	SoundFileGetChunkDataFuncPtr SoundFileGetChunkData = nullptr;
+	SoundFileGetChunkIteratorFuncPtr SoundFileGetChunkIterator = nullptr;
+	SoundFileNextChunkIteratorFuncPtr SoundFileNextChunkIterator = nullptr;
+	SoundFileSetChunkFuncPtr SoundFileSetChunk = nullptr;
 
 	void* SoundFileDllHandle;
 	static void* GetSoundFileDllHandle()
@@ -147,6 +159,11 @@ namespace Audio
 		SoundFileReadSamplesDouble = (SoundFileReadSamplesDoubleFuncPtr)GetSoundFileDllExport(TEXT("sf_read_double"), bSuccess);
 		SoundFileWriteSamplesFloat = (SoundFileWriteSamplesFloatFuncPtr)GetSoundFileDllExport(TEXT("sf_write_float"), bSuccess);
 		SoundFileWriteSamplesDouble = (SoundFileWriteSamplesDoubleFuncPtr)GetSoundFileDllExport(TEXT("sf_write_double"), bSuccess);
+		SoundFileGetChunkSize = (SoundFileGetChunkSizeFuncPtr)GetSoundFileDllExport(TEXT("sf_get_chunk_size"), bSuccess);
+		SoundFileGetChunkData = (SoundFileGetChunkDataFuncPtr)GetSoundFileDllExport(TEXT("sf_get_chunk_data"), bSuccess);
+		SoundFileGetChunkIterator = (SoundFileGetChunkIteratorFuncPtr)GetSoundFileDllExport(TEXT("sf_get_chunk_iterator"), bSuccess);
+		SoundFileNextChunkIterator = (SoundFileNextChunkIteratorFuncPtr)GetSoundFileDllExport(TEXT("sf_next_chunk_iterator"), bSuccess);
+		SoundFileSetChunk = (SoundFileSetChunkFuncPtr)GetSoundFileDllExport(TEXT("sf_set_chunk"), bSuccess);
 
 		// make sure we're successful
 		check(bSuccess);
@@ -177,6 +194,11 @@ namespace Audio
 			SoundFileReadSamplesDouble = nullptr;
 			SoundFileWriteSamplesFloat = nullptr;
 			SoundFileWriteSamplesDouble = nullptr;
+			SoundFileGetChunkSize = nullptr;
+			SoundFileGetChunkData = nullptr;
+			SoundFileGetChunkIterator = nullptr;
+			SoundFileNextChunkIterator = nullptr;
+			SoundFileSetChunk = nullptr;
 		}
 		return true;
 	}
@@ -345,6 +367,66 @@ namespace Audio
 		return ESoundFileError::Type::NONE;
 	}
 
+	static ESoundFileError::Type GetOptionalChunksInternal(LibSoundFileHandle* FileHandle, FSoundFileChunkArray& OutChunkInfoArray)
+	{
+		const TArray<uint32>& OptionalChunkIds = FWaveModInfo::GetOptionalWaveChunkIds();
+
+		for (const uint32 Id : OptionalChunkIds)
+		{
+			FSoundFileChunkInfo ChunkLookup;
+			ChunkLookup.ChunkIdSize = sizeof(Id);
+			FCStringAnsi::Strncpy(ChunkLookup.ChunkId, (ANSICHAR*)&(INTEL_ORDER32(Id)), ChunkLookup.ChunkIdSize + 1);
+
+			// Lookup chunk of given Id. Multiple chunks can exist of a given type
+			// so we loop here.
+			LibSoundFileChunkIterator* ChunkItr = SoundFileGetChunkIterator(FileHandle, &ChunkLookup);
+			while (ChunkItr)
+			{
+				FSoundFileChunkInfoWrapper ChunkInfo;
+				// SoundFileGetChunkSize retrieves the chunk data size. Oddly,
+				// it does not fill in the chunk Id.
+				int32 Result = SoundFileGetChunkSize(ChunkItr, ChunkInfo.GetPtr());
+				if (Result == 0 && ChunkInfo.GetPtr()->DataLength > 0)
+				{
+					ChunkInfo.AllocateChunkData();
+					// SoundFileGetChunkData copies in the chunk data and fills
+					// in the ChunkId.
+					Result = SoundFileGetChunkData(ChunkItr, ChunkInfo.GetPtr());
+
+					if (Result == 0)
+					{
+						OutChunkInfoArray.Add(MoveTemp(ChunkInfo));
+					}
+					else
+					{
+						UE_LOG(LogAudioMixer, Error, TEXT("LibSoundFile unable to read invalid chunk: %s"), ChunkLookup.ChunkId);
+						return ESoundFileError::Type::INVALID_CHUNK;
+					}
+				}
+
+				ChunkItr = SoundFileNextChunkIterator(ChunkItr);
+			}
+		}
+
+		return ESoundFileError::Type::NONE;
+	}
+
+	static ESoundFileError::Type WriteOptionalChunksInternal(LibSoundFileHandle* FileHandle, const FSoundFileChunkArray& ChunkInfoArray)
+	{
+		for (const FSoundFileChunkInfoWrapper& ChunkInfo : ChunkInfoArray)
+		{
+			// Note, libsndfile uses 4-byte pad when writing chunk data
+			int32 Result = SoundFileSetChunk(FileHandle, ChunkInfo.GetPtr());
+			if (Result)
+			{
+				UE_LOG(LogAudioMixer, Error, TEXT("LibSoundFile failed to write chunk data; Result = %d"), Result);
+				return ESoundFileError::Type::INVALID_STATE;
+			}
+		}
+
+		return ESoundFileError::Type::NONE;
+	}
+	
 	static SoundFileCount OnSoundFileGetLengthBytes(void* UserData)
 	{
 		SoundFileCount Length = 0;
@@ -638,6 +720,11 @@ namespace Audio
 			return ESoundFileError::Type::NONE;
 		}
 
+		ESoundFileError::Type GetOptionalChunks(FSoundFileChunkArray& OutChunkInfoArray) override
+		{
+			return GetOptionalChunksInternal(FileHandle, OutChunkInfoArray);
+		}
+
 	private:
 
 		ESoundFileError::Type InitLoaded(TSharedPtr<ISoundFile> InSoundFileData)
@@ -1096,6 +1183,11 @@ namespace Audio
 			return ESoundFileError::Type::NONE;
 		}
 
+		ESoundFileError::Type GetOptionalChunks(FSoundFileChunkArray& OutChunkInfoArray) override
+		{
+			return GetOptionalChunksInternal(FileHandle, OutChunkInfoArray);
+		}
+
 	private:
 
 		ESoundFileError::Type SetError(ESoundFileError::Type InError)
@@ -1431,6 +1523,11 @@ namespace Audio
 			return ESoundFileError::Type::NONE;
 		}
 
+		ESoundFileError::Type WriteOptionalChunks(const FSoundFileChunkArray& ChunkInfoArray) override
+		{
+			return WriteOptionalChunksInternal(FileHandle, ChunkInfoArray);
+		}
+
 	private:
 
 		ESoundFileError::Type SetError(ESoundFileError::Type InError)
diff --git a/Engine/Source/Runtime/Engine/Private/Audio.cpp b/Engine/Source/Runtime/Engine/Private/Audio.cpp
--- a/Engine/Source/Runtime/Engine/Private/Audio.cpp
+++ b/Engine/Source/Runtime/Engine/Private/Audio.cpp
@@ -1332,16 +1332,25 @@ struct FExtendedFormatChunk
 #pragma pack(pop)
 #endif
 
-bool IsKnownChunkId(uint32 ChunkId)
+const TArray<uint32>& FWaveModInfo::GetRequiredWaveChunkIds()
 {
-	uint32 KnownIds[] =
+	static TArray<uint32> RequiredChunkIds =
 	{
 		UE_mmioFOURCC('f', 'm', 't', ' '),
-		UE_mmioFOURCC('d', 'a', 't', 'a'),
+		UE_mmioFOURCC('d', 'a', 't', 'a')
+	};
+	
+	return RequiredChunkIds;
+}
+
+const TArray<uint32>& FWaveModInfo::GetOptionalWaveChunkIds()
+{
+	static TArray<uint32> OptionalChunkIds =
+	{
 		UE_mmioFOURCC('f', 'a', 'c', 't'),
 		UE_mmioFOURCC('c', 'u', 'e', ' '),
 		UE_mmioFOURCC('p', 'l', 's', 't'),
-		UE_mmioFOURCC('l', 'i', 's', 't'),
+		UE_mmioFOURCC('L', 'I', 'S', 'T'),
 		UE_mmioFOURCC('l', 'a', 'b', 'l'),
 		UE_mmioFOURCC('l', 't', 'x', 't'),
 		UE_mmioFOURCC('n', 'o', 't', 'e'),
@@ -1349,10 +1358,18 @@ bool IsKnownChunkId(uint32 ChunkId)
 		UE_mmioFOURCC('i', 'n', 's', 't'),
 		UE_mmioFOURCC('a', 'c', 'i', 'd'),
 		UE_mmioFOURCC('b', 'e', 'x', 't'),
-		UE_mmioFOURCC('i', 'X', 'M', 'L'),
+		UE_mmioFOURCC('i', 'X', 'M', 'L')
 	};
 
-	return Algo::Find(KnownIds, ChunkId) != nullptr;
+	return OptionalChunkIds;
+}
+
+bool IsKnownChunkId(uint32 ChunkId)
+{
+	bool bIsKnown = FWaveModInfo::GetRequiredWaveChunkIds().Contains(ChunkId) ||
+					FWaveModInfo::GetOptionalWaveChunkIds().Contains(ChunkId);
+
+	return bIsKnown;
 }
 
 FRiffChunkOld* FindRiffChunk(FRiffChunkOld* RiffChunkStart, const uint8* RiffChunkEnd, uint32 ChunkId)
diff --git a/Engine/Source/Runtime/Engine/Public/Audio.h b/Engine/Source/Runtime/Engine/Public/Audio.h
--- a/Engine/Source/Runtime/Engine/Public/Audio.h
+++ b/Engine/Source/Runtime/Engine/Public/Audio.h
@@ -921,6 +921,10 @@ public:
 		return ((InDW + 1)& ~1);
 	}
 
+	/** Wave Chunk Id utils */
+	ENGINE_API static const TArray<uint32>& GetRequiredWaveChunkIds();
+	ENGINE_API static const TArray<uint32>& GetOptionalWaveChunkIds();
+
 	// Read headers and load all info pointers in WaveModInfo. 
 	// Returns 0 if invalid data encountered.
 	ENGINE_API bool ReadWaveInfo(const uint8* WaveData, int32 WaveDataSize, FString* ErrorMessage = NULL, bool InHeaderDataOnly = false, void** OutFormatHeader = NULL );
