Log highest Max Channels value instead of 0
#10112

diff --git a/Engine/Source/Runtime/AudioMixer/Private/AudioMixerDevice.cpp b/Engine/Source/Runtime/AudioMixer/Private/AudioMixerDevice.cpp
--- a/Engine/Source/Runtime/AudioMixer/Private/AudioMixerDevice.cpp
+++ b/Engine/Source/Runtime/AudioMixer/Private/AudioMixerDevice.cpp
@@ -1144,12 +1144,13 @@ namespace Audio
  	{
 		FAudioPlatformSettings Settings = AudioMixerPlatform->GetPlatformSettings();
 
+		const int32 DefaultMaxChannels = GetDefault<UAudioSettings>()->GetHighestMaxChannels();
 		UE_LOG(LogAudioMixer, Display, TEXT("Audio Mixer Platform Settings:"));
 		UE_LOG(LogAudioMixer, Display, TEXT("	Sample Rate:						  %d"), Settings.SampleRate);
 		UE_LOG(LogAudioMixer, Display, TEXT("	Callback Buffer Frame Size Requested: %d"), Settings.CallbackBufferFrameSize);
 		UE_LOG(LogAudioMixer, Display, TEXT("	Callback Buffer Frame Size To Use:	  %d"), AudioMixerPlatform->GetNumFrames(Settings.CallbackBufferFrameSize));
 		UE_LOG(LogAudioMixer, Display, TEXT("	Number of buffers to queue:			  %d"), Settings.NumBuffers);
-		UE_LOG(LogAudioMixer, Display, TEXT("	Max Channels (voices):				  %d"), Settings.MaxChannels);
+		UE_LOG(LogAudioMixer, Display, TEXT("	Max Channels (voices):				  %d"), (Settings.MaxChannels > 0) ? Settings.MaxChannels : DefaultMaxChannels);
 		UE_LOG(LogAudioMixer, Display, TEXT("	Number of Async Source Workers:		  %d"), Settings.NumSourceWorkers);
 
  		return Settings;
diff --git a/Engine/Source/Runtime/Engine/Private/AudioDevice.cpp b/Engine/Source/Runtime/Engine/Private/AudioDevice.cpp
--- a/Engine/Source/Runtime/Engine/Private/AudioDevice.cpp
+++ b/Engine/Source/Runtime/Engine/Private/AudioDevice.cpp
@@ -493,15 +493,11 @@ bool FAudioDevice::Init(Audio::FDeviceId InDeviceID, int32 InMaxSources, int32 I
 	const int32 PlatformMaxSources = PlatformSettings.MaxChannels > 0 ? PlatformSettings.MaxChannels : InMaxSources;
 	MaxSources = FMath::Max(PlatformMaxSources, AudioChannelCountCVar);
 	MaxSources = FMath::Max(MaxSources, 1);
-
-	// Ensure and not assert so if in editor, user can change quality setting and re-serialize if so desired.
-	ensureMsgf(MaxSources > 0, TEXT("Neither passed MaxSources nor platform MaxChannel setting was positive value"));
 	UE_LOG(LogAudio, Display, TEXT("AudioDevice MaxSources: %d"), MaxSources);
 
 	MaxChannels = MaxSources;
 	MaxChannels_GameThread = MaxSources;
 
-
 	// Mixed sample rate is set by the platform
 	SampleRate = PlatformSettings.SampleRate;
 
diff --git a/Engine/Source/Runtime/Engine/Private/AudioDeviceManager.cpp b/Engine/Source/Runtime/Engine/Private/AudioDeviceManager.cpp
--- a/Engine/Source/Runtime/Engine/Private/AudioDeviceManager.cpp
+++ b/Engine/Source/Runtime/Engine/Private/AudioDeviceManager.cpp
@@ -1497,8 +1497,7 @@ FAudioDeviceManager::FAudioDeviceContainer::FAudioDeviceContainer(const FAudioDe
 	// Set to highest max channels initially provided by any quality setting, so that
 	// setting to lower quality but potentially returning to higher quality later at
 	// runtime is supported.
-	const UAudioSettings* AudioSettings = GetDefault<UAudioSettings>();
-	const int32 HighestMaxChannels = AudioSettings ? AudioSettings->GetHighestMaxChannels() : 0;
+	const int32 HighestMaxChannels = GetDefault<UAudioSettings>()->GetHighestMaxChannels();
 	if (Device->Init(InDeviceID, HighestMaxChannels, InParams.BufferSizeOverride, InParams.NumBuffersOverride))
 	{
 		const FAudioQualitySettings& QualitySettings = Device->GetQualityLevelSettings();
diff --git a/Engine/Source/Runtime/Engine/Private/AudioSettings.cpp b/Engine/Source/Runtime/Engine/Private/AudioSettings.cpp
--- a/Engine/Source/Runtime/Engine/Private/AudioSettings.cpp
+++ b/Engine/Source/Runtime/Engine/Private/AudioSettings.cpp
@@ -313,10 +313,7 @@ int32 UAudioSettings::GetHighestMaxChannels() const
 	int32 HighestMaxChannels = -1;
 	for (const FAudioQualitySettings& Settings : QualityLevels)
 	{
-		if (Settings.MaxChannels > HighestMaxChannels)
-		{
-			HighestMaxChannels = Settings.MaxChannels;
-		}
+		HighestMaxChannels = FMath::Max(Settings.MaxChannels, HighestMaxChannels);
 	}
 
 	return HighestMaxChannels;
