[Metasound] Fix mismatch in param transmitter ids

diff --git a/Engine/Plugins/Runtime/Metasound/Source/MetasoundEngine/Private/MetasoundSource.cpp b/Engine/Plugins/Runtime/Metasound/Source/MetasoundEngine/Private/MetasoundSource.cpp
--- a/Engine/Plugins/Runtime/Metasound/Source/MetasoundEngine/Private/MetasoundSource.cpp
+++ b/Engine/Plugins/Runtime/Metasound/Source/MetasoundEngine/Private/MetasoundSource.cpp
@@ -603,7 +603,7 @@ ISoundGeneratorPtr UMetaSoundSource::CreateSoundGenerator(const FSoundGeneratorI
 	};
 
 	TSharedPtr<FMetasoundGenerator> Generator = MakeShared<FMetasoundGenerator>(MoveTemp(InitParams));
-	TrackGenerator(InParams.InstanceID, Generator);
+	TrackGenerator(InParams.AudioComponentId, Generator);
 
 	return ISoundGeneratorPtr(Generator);
 }
diff --git a/Engine/Source/Runtime/AudioMixer/Private/AudioMixerSourceBuffer.cpp b/Engine/Source/Runtime/AudioMixer/Private/AudioMixerSourceBuffer.cpp
--- a/Engine/Source/Runtime/AudioMixer/Private/AudioMixerSourceBuffer.cpp
+++ b/Engine/Source/Runtime/AudioMixer/Private/AudioMixerSourceBuffer.cpp
@@ -104,11 +104,12 @@ namespace Audio
 		{
 			FSoundGeneratorInitParams InitParams;
 			InitParams.AudioDeviceID = InArgs.AudioDeviceID;
+			InitParams.AudioComponentId = InArgs.AudioComponentID;
 			InitParams.SampleRate = InArgs.SampleRate;
 			InitParams.AudioMixerNumOutputFrames = InArgs.AudioMixerNumOutputFrames;
 			InitParams.NumChannels = NumChannels;
 			InitParams.NumFramesPerCallback = MONO_PCM_BUFFER_SAMPLES;
-			InitParams.InstanceID = InArgs.AudioComponentID;
+			InitParams.InstanceID = InArgs.InstanceID;
 			InitParams.bIsPreviewSound = InArgs.bIsPreviewSound;
 
 			SoundGenerator = InArgs.SoundWave->CreateSoundGenerator(InitParams, MoveTemp(InDefaultParams));
diff --git a/Engine/Source/Runtime/Engine/Classes/Sound/SoundGenerator.h b/Engine/Source/Runtime/Engine/Classes/Sound/SoundGenerator.h
--- a/Engine/Source/Runtime/Engine/Classes/Sound/SoundGenerator.h
+++ b/Engine/Source/Runtime/Engine/Classes/Sound/SoundGenerator.h
@@ -16,6 +16,7 @@ namespace Audio
 struct FSoundGeneratorInitParams
 {
 	Audio::FDeviceId AudioDeviceID;
+	uint64 AudioComponentId = 0;
 	float SampleRate = 0.0f;
 	int32 AudioMixerNumOutputFrames = 0;
 	int32 NumChannels = 0;
