Reduce thread hang detection and stack walk from a Fatal log to a Warning.

diff --git a/Engine/Source/Runtime/AudioMixer/Private/AudioMixerSourceManager.cpp b/Engine/Source/Runtime/AudioMixer/Private/AudioMixerSourceManager.cpp
--- a/Engine/Source/Runtime/AudioMixer/Private/AudioMixerSourceManager.cpp
+++ b/Engine/Source/Runtime/AudioMixer/Private/AudioMixerSourceManager.cpp
@@ -3285,7 +3285,7 @@ namespace Audio
 				const SIZE_T StackTraceSize = 65536;
 				ANSICHAR StackTrace[StackTraceSize] = {0};
 				FPlatformStackWalk::ThreadStackWalkAndDump(StackTrace, StackTraceSize, 0, AudioRenderThreadId);
-				UE_LOG(LogAudioMixer, Fatal, TEXT("***** ThreadStackWalkAndDump for ThreadId(%lu) ******\n%s"), AudioRenderThreadId, ANSI_TO_TCHAR(StackTrace));
+				UE_LOG(LogAudioMixer, Warning, TEXT("***** ThreadStackWalkAndDump for ThreadId(%lu) ******\n%s"), AudioRenderThreadId, ANSI_TO_TCHAR(StackTrace));
 			}
 			
 		}
