Adding better logging around inflight async tasks for renderer hangs/stalls.

diff --git a/Engine/Source/Runtime/AudioMixer/Private/AudioMixerSourceBuffer.cpp b/Engine/Source/Runtime/AudioMixer/Private/AudioMixerSourceBuffer.cpp
--- a/Engine/Source/Runtime/AudioMixer/Private/AudioMixerSourceBuffer.cpp
+++ b/Engine/Source/Runtime/AudioMixer/Private/AudioMixerSourceBuffer.cpp
@@ -83,6 +83,7 @@ namespace Audio
 		, BufferType(InArgs.Buffer->GetType())
 		, NumPrecacheFrames(InArgs.SoundWave->NumPrecacheFrames)
 		, AuioDeviceID(InArgs.AudioDeviceID)
+		, WaveName(InArgs.SoundWave->GetFName())
 #if ENABLE_AUDIO_DEBUG
 		, SampleRate(InArgs.SampleRate)
 #endif // ENABLE_AUDIO_DEBUG
@@ -365,6 +366,7 @@ namespace Audio
 				NewTaskData.AudioData = SourceVoiceBuffers[BufferIndex]->AudioData.GetData();
 				NewTaskData.NumSamples = MaxSamples;
 				NewTaskData.NumChannels = NumChannels;
+				AsyncTaskStartTimeInCycles = FPlatformTime::Cycles64();
 				check(!AsyncRealtimeAudioTask);
 				AsyncRealtimeAudioTask = CreateAudioTask(AuioDeviceID, NewTaskData);
 			}
@@ -406,7 +408,8 @@ namespace Audio
 		NewTaskData.NumPrecacheFrames = NumPrecacheFrames;
 		NewTaskData.bForceSyncDecode = bForceSyncDecode;
 
-		FScopeLock Lock(&DecodeTaskCritSec);
+		AsyncTaskStartTimeInCycles = FPlatformTime::Cycles64();
+		FScopeLock Lock(&DecodeTaskCritSec);		
 		check(!AsyncRealtimeAudioTask);
 		AsyncRealtimeAudioTask = CreateAudioTask(AuioDeviceID, NewTaskData);
 
@@ -483,6 +486,7 @@ namespace Audio
 
 			delete AsyncRealtimeAudioTask;
 			AsyncRealtimeAudioTask = nullptr;
+			AsyncTaskStartTimeInCycles = 0;
 
 			SubmitRealTimeSourceData(bIsFinishedOrLooped);
 		}
@@ -622,6 +626,17 @@ namespace Audio
 	}
 #endif // ENABLE_AUDIO_DEBUG
 
+	void FMixerSourceBuffer::GetDiagnosticState(FDiagnosticState& OutState)
+	{
+		// Query without a lock!
+		OutState.bInFlight = AsyncRealtimeAudioTask != nullptr;
+		OutState.WaveName = WaveName;
+		OutState.bProcedural = bProcedural;
+		OutState.RunTimeInSecs = OutState.bInFlight ?
+			FPlatformTime::ToSeconds(FPlatformTime::Cycles64() - this->AsyncTaskStartTimeInCycles) :
+			0.f;
+	}
+
 	void FMixerSourceBuffer::EnsureAsyncTaskFinishes()
 	{
 		FScopeLock Lock(&DecodeTaskCritSec);
diff --git a/Engine/Source/Runtime/AudioMixer/Private/AudioMixerSourceBuffer.h b/Engine/Source/Runtime/AudioMixer/Private/AudioMixerSourceBuffer.h
--- a/Engine/Source/Runtime/AudioMixer/Private/AudioMixerSourceBuffer.h
+++ b/Engine/Source/Runtime/AudioMixer/Private/AudioMixerSourceBuffer.h
@@ -104,6 +104,16 @@ namespace Audio
 		// Returns true if the async task is done
 		bool IsAsyncTaskDone() const;
 
+		// Returns some diagnostic state
+		struct FDiagnosticState
+		{
+			FName WaveName;
+			float RunTimeInSecs=0.f;
+			bool bInFlight=false;
+			bool bProcedural=false;
+		};
+		void GetDiagnosticState(FDiagnosticState& OutState);
+
 		// Ensures the async task finishes
 		void EnsureAsyncTaskFinishes();
 
@@ -146,6 +156,9 @@ namespace Audio
 		int32 NumPrecacheFrames;
 		Audio::FDeviceId AuioDeviceID;
 		TArray<uint8> CachedRealtimeFirstBuffer;
+		FName WaveName;
+		uint64 AsyncTaskStartTimeInCycles=0;
+
 #if ENABLE_AUDIO_DEBUG
 		int32 SampleRate = 0;
 		std::atomic<double> CPUCoreUtilization = 0.0;
diff --git a/Engine/Source/Runtime/AudioMixer/Private/AudioMixerSourceManager.cpp b/Engine/Source/Runtime/AudioMixer/Private/AudioMixerSourceManager.cpp
--- a/Engine/Source/Runtime/AudioMixer/Private/AudioMixerSourceManager.cpp
+++ b/Engine/Source/Runtime/AudioMixer/Private/AudioMixerSourceManager.cpp
@@ -21,7 +21,11 @@
 #include "Stats/Stats.h"
 #include "Trace/Trace.h"
 
-
+#if WITH_AUDIO_MIXER_THREAD_COMMAND_DEBUG
+	#define AUDIO_MIXER_THREAD_COMMAND_STRING(X) (X)
+#else //WITH_AUDIO_MIXER_THREAD_COMMAND_DEBUG
+	#define AUDIO_MIXER_THREAD_COMMAND_STRING(X) ("")
+#endif //WITH_AUDIO_MIXER_THREAD_COMMAND_DEBUG
 
 // Link to "Audio" profiling category
 CSV_DECLARE_CATEGORY_MODULE_EXTERN(AUDIOMIXERCORE_API, Audio);
@@ -208,6 +212,20 @@ UE_TRACE_EVENT_BEGIN(Audio, MixerSourceEnvelope)
 UE_TRACE_EVENT_END()
 #endif // UE_AUDIO_PROFILERTRACE_ENABLED
 
+
+#ifndef CASE_ENUM_TO_TEXT
+	#define CASE_ENUM_TO_TEXT(TXT) case TXT: return TEXT(#TXT);
+#endif
+
+const TCHAR* LexToString(ESourceManagerRenderThreadPhase InPhase)
+{
+	switch(InPhase)
+	{
+		FOREACH_ENUM_ESOURCEMANAGERRENDERTHREADPHASE(CASE_ENUM_TO_TEXT)
+	} 
+	return TEXT("Unknown");
+}
+
 namespace Audio
 {
 	/*************************************************************************
@@ -422,7 +440,7 @@ namespace Audio
 				} Payload;
 				for (int32 i = 0; i < 65536; ++i)
 				{
-					AudioMixerThreadCommand([Payload] {});
+					AudioMixerThreadCommand([Payload] {}, AUDIO_MIXER_THREAD_COMMAND_STRING("SpamNopsCmd() -- Console command"));
 				}
 			})
 		);
@@ -434,13 +452,25 @@ namespace Audio
 			TEXT(""),
 			FConsoleCommandDelegate::CreateLambda([this]()
 			{
-				struct FSpamPayload
+				AudioMixerThreadCommand([] {while(true){}}, AUDIO_MIXER_THREAD_COMMAND_STRING("ChokeCommandQueue() -- Console command"));
+			})
+		);
+
+		// submit a MPSC command that has an endless loop
+		static FAutoConsoleCommand SpamEndlessCmdMPSC(
+			TEXT("au.AudioThreadCommand.ChokeMPSCCommandQueue"),
+			TEXT(""),
+			FConsoleCommandDelegate::CreateLambda([this]()
 				{
-					uint8 JunkBytes[1024];
-				} Payload;
+					AudioMixerThreadMPSCCommand([] {while (true) {}}, AUDIO_MIXER_THREAD_COMMAND_STRING("ChokeMPSCCommandQueue() -- Console command"));
+				})
+		);
 
-				AudioMixerThreadCommand([Payload] {while(true){}});
-			})
+		// Test stall diagnostics.
+		static FAutoConsoleCommand StallDiagnostics(
+			TEXT("au.AudioSourceManager.HangDiagnostics"),
+			TEXT(""),
+			FConsoleCommandDelegate::CreateLambda([this]() { DoStallDiagnostics(); })
 		);
 
 		bInitialized = true;
@@ -1175,7 +1205,7 @@ namespace Audio
 
 				AudioBuses.Add(InAudioBusKey, NewBusData);
 			}
-		});
+		}, AUDIO_MIXER_THREAD_COMMAND_STRING("StartAudioBus()"));
 	}
 
 	void FMixerSourceManager::StopAudioBus(FAudioBusKey InAudioBusKey)
@@ -1208,7 +1238,7 @@ namespace Audio
 					AudioBuses.Remove(InAudioBusKey);
 				}
 			}
-		});
+		}, AUDIO_MIXER_THREAD_COMMAND_STRING("StopAudioBus()"));
 	}
 
 	bool FMixerSourceManager::IsAudioBusActive(FAudioBusKey InAudioBusKey) const
@@ -1246,7 +1276,7 @@ namespace Audio
 			AudioMixerThreadMPSCCommand([this, InAudioBusKey, InPatchOutputStrongPtr]()
 			{
 				AddPatchOutputForAudioBus(InAudioBusKey, InPatchOutputStrongPtr);
-			});
+			}, AUDIO_MIXER_THREAD_COMMAND_STRING("AddPatchOutputForAudioBus()"));
 		}
 	}
 
@@ -1255,7 +1285,7 @@ namespace Audio
 		AudioMixerThreadCommand([this, InAudioBusKey, NewPatchPtr = InPatchOutputStrongPtr]() mutable
 		{
 			AddPatchOutputForAudioBus(InAudioBusKey, NewPatchPtr);
-		});
+		}, AUDIO_MIXER_THREAD_COMMAND_STRING("AddPatchOutputForAudioBus_AudioThread()"));
 	}
 
 	void FMixerSourceManager::AddPatchInputForAudioBus(FAudioBusKey InAudioBusKey, const FPatchInput& InPatchInput)
@@ -1275,7 +1305,7 @@ namespace Audio
 			AudioMixerThreadMPSCCommand([this, InAudioBusKey, InPatchInput]()
 			{
 				AddPatchInputForAudioBus(InAudioBusKey, InPatchInput);
-			});
+			}, AUDIO_MIXER_THREAD_COMMAND_STRING("AddPatchInputForAudioBus()"));
 		}
 	}
 
@@ -1284,7 +1314,7 @@ namespace Audio
 		AudioMixerThreadCommand([this, InAudioBusKey, InPatchInput]()
 		{
 			AddPatchInputForAudioBus(InAudioBusKey, InPatchInput);
-		});
+		}, AUDIO_MIXER_THREAD_COMMAND_STRING("AddPatchInputForAudioBus_AudioThread()"));
 	}
 
 	void FMixerSourceManager::Play(const int32 SourceId)
@@ -1309,7 +1339,7 @@ namespace Audio
 			SourceInfo.StartTime = StartTime;
 
 			AUDIO_MIXER_DEBUG_LOG(SourceId, TEXT("Is playing"));
-		});
+		}, AUDIO_MIXER_THREAD_COMMAND_STRING("Play()"));
 	}
 
 	void FMixerSourceManager::CancelQuantizedSound(const int32 SourceId)
@@ -1362,7 +1392,7 @@ namespace Audio
 		AudioMixerThreadCommand([this, SourceId]()
 		{
 			StopInternal(SourceId);
-		});
+		}, AUDIO_MIXER_THREAD_COMMAND_STRING("Stop()"));
 	}
 
 	void FMixerSourceManager::StopInternal(const int32 SourceId)
@@ -1429,7 +1459,7 @@ namespace Audio
 			}
 
 			AUDIO_MIXER_DEBUG_LOG(SourceId, TEXT("Is stopping with fade"));
-		});
+		}, AUDIO_MIXER_THREAD_COMMAND_STRING("StopFade()"));
 	}
 	
 	void FMixerSourceManager::Pause(const int32 SourceId)
@@ -1446,7 +1476,7 @@ namespace Audio
 
 			SourceInfo.bIsPaused = true;
 			SourceInfo.bIsActive = false;
-		});
+		}, AUDIO_MIXER_THREAD_COMMAND_STRING("Pause()"));
 	}
 
 	void FMixerSourceManager::SetPitch(const int32 SourceId, const float Pitch)
@@ -1460,7 +1490,7 @@ namespace Audio
 			check(NumOutputFrames > 0);
 
 			SourceInfos[SourceId].PitchSourceParam.SetValue(Pitch, NumOutputFrames);
-		});
+		}, AUDIO_MIXER_THREAD_COMMAND_STRING("SetPitch()"));
 	}
 
 	void FMixerSourceManager::SetVolume(const int32 SourceId, const float Volume)
@@ -1487,7 +1517,7 @@ namespace Audio
 
 				SourceInfos[SourceId].VolumeSourceDestination = Volume;
 			}
-		});
+		}, AUDIO_MIXER_THREAD_COMMAND_STRING("SetVolume()"));
 	}
 
 	void FMixerSourceManager::SetDistanceAttenuation(const int32 SourceId, const float DistanceAttenuation)
@@ -1508,7 +1538,7 @@ namespace Audio
 			}
 
 			SourceInfos[SourceId].DistanceAttenuationSourceDestination = DistanceAttenuation;
-		});
+		}, AUDIO_MIXER_THREAD_COMMAND_STRING("SetDistanceAttenuation()"));
 	}
 
 	void FMixerSourceManager::SetSpatializationParams(const int32 SourceId, const FSpatializationParams& InParams)
@@ -1522,7 +1552,7 @@ namespace Audio
 			AUDIO_MIXER_CHECK_AUDIO_PLAT_THREAD(MixerDevice);
 
 			SourceInfos[SourceId].SpatParams = InParams;
-		});
+		}, AUDIO_MIXER_THREAD_COMMAND_STRING("SetSpatializationParams()"));
 	}
 
 	void FMixerSourceManager::SetChannelMap(const int32 SourceId, const uint32 NumInputChannels, const Audio::FAlignedFloatBuffer& ChannelMap, const bool bInIs3D, const bool bInIsCenterChannelOnly)
@@ -1553,7 +1583,7 @@ namespace Audio
 
 			bool bNeedsSpeakerMap = SourceSubmixOutput.SetChannelMap(ChannelMap, bInIsCenterChannelOnly);
 			GameThreadInfo.bNeedsSpeakerMap[SourceId] = bNeedsSpeakerMap;
-		});
+		}, AUDIO_MIXER_THREAD_COMMAND_STRING("SetChannelMap()"));
 	}
 
 	void FMixerSourceManager::SetLPFFrequency(const int32 SourceId, const float InLPFFrequency)
@@ -1573,7 +1603,7 @@ namespace Audio
 			// On callback, the client version may be overridden with the modulation LPF value depending on which is more aggressive.  
 			SourceInfo.LowPassFreq = InLPFFrequency;
 			SourceInfo.LowPassFilter.StartFrequencyInterpolation(InLPFFrequency, NumOutputFrames);
-		});
+		}, AUDIO_MIXER_THREAD_COMMAND_STRING("SetLPFFrequency()"));
 	}
 
 	void FMixerSourceManager::SetHPFFrequency(const int32 SourceId, const float InHPFFrequency)
@@ -1592,7 +1622,7 @@ namespace Audio
 			// On callback, the client version may be overridden with the modulation HPF value depending on which is more aggressive.  
 			SourceInfo.HighPassFreq = InHPFFrequency;
 			SourceInfo.HighPassFilter.StartFrequencyInterpolation(InHPFFrequency, NumOutputFrames);
-		});
+		}, AUDIO_MIXER_THREAD_COMMAND_STRING("SetHPFFrequency()"));
 	}
 
 	void FMixerSourceManager::SetModLPFFrequency(const int32 SourceId, const float InLPFFrequency)
@@ -1608,7 +1638,7 @@ namespace Audio
 			FSourceInfo& SourceInfo = SourceInfos[SourceId];
 			SourceInfo.LowpassModulationBase = InLPFFrequency;
 			SourceInfo.bModFiltersUpdated = true;
-		});
+		}, AUDIO_MIXER_THREAD_COMMAND_STRING("SetModLPFFrequency()"));
 	}
 
 	void FMixerSourceManager::SetModHPFFrequency(const int32 SourceId, const float InHPFFrequency)
@@ -1624,7 +1654,7 @@ namespace Audio
 			FSourceInfo& SourceInfo = SourceInfos[SourceId];
 			SourceInfo.HighpassModulationBase = InHPFFrequency;
 			SourceInfo.bModFiltersUpdated = true;
-		});
+		}, AUDIO_MIXER_THREAD_COMMAND_STRING("SetModHPFFrequency()"));
 	}
 
 	void FMixerSourceManager::SetModulationRouting(const int32 SourceId, FSoundModulationDefaultSettings& ModulationSettings)
@@ -1663,7 +1693,8 @@ namespace Audio
 				SourceInfo.PitchModulation = MoveTemp(PitchModulation);
 				SourceInfo.LowpassModulation = MoveTemp(LowpassModulation);
 				SourceInfo.HighpassModulation = MoveTemp(HighpassModulation);
-			});
+			}, AUDIO_MIXER_THREAD_COMMAND_STRING("SetModulationRouting()")
+		);
 	}
 
 	void FMixerSourceManager::SetSourceBufferListener(const int32 SourceId, FSharedISourceBufferListenerPtr& InSourceBufferListener, bool InShouldSourceBufferListenerZeroBuffer)
@@ -1679,7 +1710,7 @@ namespace Audio
 			FSourceInfo& SourceInfo = SourceInfos[SourceId];
 			SourceInfo.SourceBufferListener = InSourceBufferListener;
 			SourceInfo.bShouldSourceBufferListenerZeroBuffer = InShouldSourceBufferListenerZeroBuffer;
-		});
+		}, AUDIO_MIXER_THREAD_COMMAND_STRING("SetSourceBufferListener()"));
 	}
 
 	void FMixerSourceManager::SetModVolume(const int32 SourceId, const float InModVolume)
@@ -1694,7 +1725,7 @@ namespace Audio
 
 			FSourceInfo& SourceInfo = SourceInfos[SourceId];
 			SourceInfo.VolumeModulationBase = InModVolume;
-		});
+		}, AUDIO_MIXER_THREAD_COMMAND_STRING("SetModVolume()"));
 	}
 
 	void FMixerSourceManager::SetModPitch(const int32 SourceId, const float InModPitch)
@@ -1709,7 +1740,7 @@ namespace Audio
 
 			FSourceInfo& SourceInfo = SourceInfos[SourceId];
 			SourceInfo.PitchModulationBase = InModPitch;
-		});
+		}, AUDIO_MIXER_THREAD_COMMAND_STRING("SetModPitch()"));
 	}
 
 	void FMixerSourceManager::SetSubmixSendInfo(const int32 SourceId, const FMixerSourceSubmixSend& InSubmixSend)
@@ -1768,7 +1799,7 @@ namespace Audio
 					InSubmixPtr->AddOrSetSourceVoice(SourceVoice, InSubmixSend.SendLevel, InSubmixSend.SubmixSendStage);
 				}
 			}
-		});
+		}, AUDIO_MIXER_THREAD_COMMAND_STRING("SetSubmixSendInfo()"));
 	}
 
 	void FMixerSourceManager::ClearSubmixSendInfo(const int32 SourceId, const FMixerSourceSubmixSend& InSubmixSend)
@@ -1816,7 +1847,7 @@ namespace Audio
 				// Now remove the source voice from the submix send list
 				InSubmixPtr->RemoveSourceVoice(MixerSources[SourceId]);
 			}
-		});
+		}, AUDIO_MIXER_THREAD_COMMAND_STRING("ClearSubmixSendInfo()"));
 	}
 
 	void FMixerSourceManager::SetBusSendInfo(const int32 SourceId, EBusSendType InAudioBusSendType, uint32 AudioBusId, float BusSendLevel)
@@ -1871,7 +1902,7 @@ namespace Audio
 			{
 				SourceInfo.AudioBusSends[(int32)InAudioBusSendType].Add(AudioBusId);
 			}
-		});
+		}, AUDIO_MIXER_THREAD_COMMAND_STRING("SetBusSendInfo()"));
 	}
 
 	void FMixerSourceManager::SetListenerTransforms(const TArray<FTransform>& InListenerTransforms)
@@ -1879,7 +1910,7 @@ namespace Audio
 		AudioMixerThreadCommand([this, InListenerTransforms]()
 		{
 			ListenerTransforms = InListenerTransforms;
-		});
+		}, AUDIO_MIXER_THREAD_COMMAND_STRING("SetListenerTransforms()"));
 	}
 
 	const TArray<FTransform>* FMixerSourceManager::GetListenerTransforms() const
@@ -2286,6 +2317,8 @@ namespace Audio
 
 	void FMixerSourceManager::ComputeBuses()
 	{
+		RenderThreadPhase = ESourceManagerRenderThreadPhase::ComputeBusses;
+		
 		// Loop through the bus registry and mix source audio
 		for (auto& Entry : AudioBuses)
 		{
@@ -2296,6 +2329,8 @@ namespace Audio
 
 	void FMixerSourceManager::UpdateBuses()
 	{
+		RenderThreadPhase = ESourceManagerRenderThreadPhase::UpdateBusses;
+
 		// Update the bus states post mixing. This flips the current/previous buffer indices.
 		for (auto& Entry : AudioBuses)
 		{
@@ -2305,15 +2340,44 @@ namespace Audio
 	}
 
 	// ctor
-	FMixerSourceManager::FAudioMixerThreadCommand::FAudioMixerThreadCommand(TFunction<void()> InFunction, bool bInDeferExecution)
-		: Function(MoveTemp(InFunction))
-		, bDeferExecution(bInDeferExecution)
+	FMixerSourceManager::FAudioMixerThreadCommand::FAudioMixerThreadCommand(TFunction<void()>&& InFunction, const char* InDebugString, bool bInDeferExecution)
+	: Function(MoveTemp(InFunction))
+	, bDeferExecution(bInDeferExecution)		
+#if WITH_AUDIO_MIXER_THREAD_COMMAND_DEBUG
+	, DebugString(InDebugString)
+#endif //WITH_AUDIO_MIXER_THREAD_COMMAND_DEBUG
 	{
 	}
 
 	void FMixerSourceManager::FAudioMixerThreadCommand::operator()() const
 	{
+#if WITH_AUDIO_MIXER_THREAD_COMMAND_DEBUG
+		StartExecuteTimeInCycles = FPlatformTime::Cycles64();
+#endif //WITH_AUDIO_MIXER_THREAD_COMMAND_DEBUG
+
 		Function();
+
+#if WITH_AUDIO_MIXER_THREAD_COMMAND_DEBUG
+		StartExecuteTimeInCycles = 0;
+#endif //WITH_AUDIO_MIXER_THREAD_COMMAND_DEBUG
+	}
+
+	FString FMixerSourceManager::FAudioMixerThreadCommand::GetSafeDebugString() const
+	{
+#if WITH_AUDIO_MIXER_THREAD_COMMAND_DEBUG
+		return FString(DebugString ? ANSI_TO_TCHAR(DebugString): TEXT(""));
+#else  //WITH_AUDIO_MIXER_THREAD_COMMAND_DEBUG
+		return {};
+#endif //WITH_AUDIO_MIXER_THREAD_COMMAND_DEBUG
+	}
+
+	float FMixerSourceManager::FAudioMixerThreadCommand::GetExecuteTimeInSeconds() const
+	{
+#if WITH_AUDIO_MIXER_THREAD_COMMAND_DEBUG
+		return FPlatformTime::ToSeconds64(FPlatformTime::Cycles64() - StartExecuteTimeInCycles);
+#else  //WITH_AUDIO_MIXER_THREAD_COMMAND_DEBUG
+		return 0.f;
+#endif //WITH_AUDIO_MIXER_THREAD_COMMAND_DEBUG
 	}
 
 	void FMixerSourceManager::ApplyDistanceAttenuation(FSourceInfo& SourceInfo, int32 NumSamples)
@@ -2870,6 +2934,10 @@ namespace Audio
 
 	void FMixerSourceManager::GenerateSourceAudio(const bool bGenerateBuses)
 	{
+		RenderThreadPhase = bGenerateBuses ?
+			ESourceManagerRenderThreadPhase::GenerateSrcAudio_WithBusses :
+			ESourceManagerRenderThreadPhase::GenerateSrcAudio_WithoutBusses;
+					
 		// If there are no buses, don't need to do anything here
 		if (bGenerateBuses && !AudioBuses.Num())
 		{
@@ -2982,7 +3050,7 @@ namespace Audio
 
 				SourceSubmixOutputBuffer.SetChannelMap(SourceInfo.ScratchChannelMap, SourceInfo.bIsCenterChannelOnly);
 			}
-		});
+		}, AUDIO_MIXER_THREAD_COMMAND_STRING("UpdateDeviceChannelCount()"));
 	}
 
 	void FMixerSourceManager::UpdateSourceEffectChain(const uint32 InSourceEffectChainId, const TArray<FSourceEffectChainEntry>& InSourceEffectChain, const bool bPlayEffectChainTails)
@@ -3047,7 +3115,7 @@ namespace Audio
 					}
 				}
 			}
-		}, /*bDeferExecution*/true);
+		}, AUDIO_MIXER_THREAD_COMMAND_STRING("UpdateSourceEffectChain()"), /*bDeferExecution*/true);
 	}
 
 	void FMixerSourceManager::PauseSoundForQuantizationCommand(const int32 SourceId)
@@ -3145,6 +3213,8 @@ namespace Audio
 		CSV_SCOPED_TIMING_STAT(Audio, SourceManagerUpdate);
 		SCOPE_CYCLE_COUNTER(STAT_AudioMixerSourceManagerUpdate);
 
+		RenderThreadPhase = ESourceManagerRenderThreadPhase::Begin;
+
 		if (FPlatformProcess::SupportsMultithreading())
 		{
 			// Get the this blocks commands before rendering audio
@@ -3157,6 +3227,7 @@ namespace Audio
 		}
 
 		// Notify modulation interface that we are beginning to update
+		RenderThreadPhase = ESourceManagerRenderThreadPhase::ProcessModulators;
 		if (MixerDevice->IsModulationPluginEnabled() && MixerDevice->ModulationInterface.IsValid())
 		{
 			MixerDevice->ModulationInterface->ProcessModulators(MixerDevice->GetAudioClockDelta());
@@ -3180,6 +3251,7 @@ namespace Audio
 		// Let the plugin know we finished processing all sources
 		if (bUsingSpatializationPlugin)
 		{
+			RenderThreadPhase = ESourceManagerRenderThreadPhase::SpatialInterface_OnAllSourcesProcessed;
 			AUDIO_MIXER_CHECK(SpatialInterfaceInfo.SpatializationPlugin.IsValid());
 			LLM_SCOPE(ELLMTag::AudioMixerPlugins);
 			SpatialInterfaceInfo.SpatializationPlugin->OnAllSourcesProcessed();
@@ -3188,12 +3260,14 @@ namespace Audio
 		// Let the plugin know we finished processing all sources
 		if (bUsingSourceDataOverridePlugin)
 		{
+			RenderThreadPhase = ESourceManagerRenderThreadPhase::SourceDataOverride_OnAllSourcesProcessed;
 			AUDIO_MIXER_CHECK(SourceDataOverridePlugin.IsValid());
 			LLM_SCOPE(ELLMTag::AudioMixerPlugins);
 			SourceDataOverridePlugin->OnAllSourcesProcessed();
 		}
 
-		// Update the game thread copy of source donenessf
+		// Update the game thread copy of source doneness
+		RenderThreadPhase = ESourceManagerRenderThreadPhase::UpdateGameThreadCopies;
 		for (int32 SourceId = 0; SourceId < NumTotalSources; ++SourceId)
 		{		
 			FSourceInfo& SourceInfo = SourceInfos[SourceId];
@@ -3215,6 +3289,7 @@ namespace Audio
 				}
 			}
 		}
+		RenderThreadPhase = ESourceManagerRenderThreadPhase::Finished;
 	}
 
 	void FMixerSourceManager::ClearStoppingSounds()
@@ -3237,16 +3312,14 @@ namespace Audio
 		}
 	}
 
-	void FMixerSourceManager::AudioMixerThreadMPSCCommand(TFunction<void()> InCommand)
+	void FMixerSourceManager::AudioMixerThreadMPSCCommand(TFunction<void()>&& InCommand, const char* InDebugString)
 	{
-		MpscCommandQueue.Enqueue(MoveTemp(InCommand));
+		MpscCommandQueue.Enqueue( FAudioMixerMpscCommand{ MoveTemp(InCommand), InDebugString, false });
 	}
 
-
-
-	void FMixerSourceManager::AudioMixerThreadCommand(TFunction<void()> InFunction, bool bDeferExecution)
+	void FMixerSourceManager::AudioMixerThreadCommand(TFunction<void()>&& InFunction, const char* InDebugString, bool bInDeferExecution /*= false*/)
 	{
-		FAudioMixerThreadCommand AudioCommand(InFunction, bDeferExecution);
+		FAudioMixerThreadCommand AudioCommand(MoveTemp(InFunction), InDebugString, bInDeferExecution);
 
 		// Here, we make sure that we don't flip our command double buffer while we are executing this function.
 		FScopeLock ScopeLock(&CommandBufferIndexCriticalSection);
@@ -3262,10 +3335,13 @@ namespace Audio
 		if (CurrentBufferSizeInBytes > WarnSize )
 		{
 			SIZE_T Num = CommandBuffers[AudioThreadCommandIndex].SourceCommandQueue.Num();
-			UE_LOG(LogAudioMixer, Error, TEXT("Command Queue has grown to %ukb, containing %d cmds, last pump was %.2fms ago."),
-				CurrentBufferSizeInBytes >> 10, Num, FPlatformTime::ToMilliseconds64(FPlatformTime::Cycles64() - LastPumpTimeInCycles));
+			float TimeSinceLastComplete = FPlatformTime::ToSeconds64(FPlatformTime::Cycles64() - LastPumpCompleteTimeInCycles);
+
+			UE_LOG(LogAudioMixer, Error, TEXT("Command Queue has grown to %ukb, containing %d cmds, last complete pump was %2.5f seconds ago."),
+				CurrentBufferSizeInBytes >> 10, Num, TimeSinceLastComplete);
 			WarnSize *= 2;
 
+			DoStallDiagnostics();
 		}
 
 		// Before adding further commands, ensure we're not growing outside any sensible size for these buffers.
@@ -3278,19 +3354,11 @@ namespace Audio
 		}
 		else
 		{
-			UE_LOG(LogAudioMixer, Error,  TEXT("Command buffer grown to %umb, preventing any more adds! Likely cause, the h/w has stopped consuming data."), CurrentBufferSizeInBytes >>20);
-			
-			if (AudioRenderThreadId != INVALID_AUDIO_RENDER_THREAD_ID)
-			{
-				const SIZE_T StackTraceSize = 65536;
-				ANSICHAR StackTrace[StackTraceSize] = {0};
-				FPlatformStackWalk::ThreadStackWalkAndDump(StackTrace, StackTraceSize, 0, AudioRenderThreadId);
-				UE_LOG(LogAudioMixer, Warning, TEXT("***** ThreadStackWalkAndDump for ThreadId(%lu) ******\n%s"), AudioRenderThreadId, ANSI_TO_TCHAR(StackTrace));
-			}
-			
+			UE_LOG(LogAudioMixer, Error, TEXT("Command buffer grown to %umb, preventing any more adds! Likely cause the AudioRenderer has hung"), CurrentBufferSizeInBytes >>20);
 		}
 	}
 
+
 	void FMixerSourceManager::PumpCommandQueue()
 	{
 		TRACE_CPUPROFILER_EVENT_SCOPE(AudioMixerThreadCommands::PumpCommandQueue)
@@ -3306,17 +3374,26 @@ namespace Audio
 		}
 
 		// Pump the MPSC command queue
+		RenderThreadPhase = ESourceManagerRenderThreadPhase::PumpMpscCmds;
 		TOptional Opt{ MpscCommandQueue.Dequeue() };
 		while (Opt.IsSet())
 		{
-			TFunction<void()> Command = MoveTemp(Opt.GetValue());
-			Command();
+			// First copy/move out the command and keep a copy of it.
+			{
+				FWriteScopeLock Lock(CurrentlyExecutingCmdLock);
+				CurrentlyExecuteingCmd = MoveTemp(Opt.GetValue());
+			}
+			
+			// Execute the current under a read-lock.
+			{
+				FReadScopeLock Lock(CurrentlyExecutingCmdLock);
+				CurrentlyExecuteingCmd();
+			}
 				
 			Opt = MpscCommandQueue.Dequeue();
 		}
 
 		int32 CurrentRenderThreadIndex = RenderThreadCommandBufferIndex.GetValue();
-
 		FCommands& Commands = CommandBuffers[CurrentRenderThreadIndex];
 
 		const int32 NumCommandsToExecute = Commands.SourceCommandQueue.Num();
@@ -3325,25 +3402,33 @@ namespace Audio
 		// Pop and execute all the commands that came since last update tick
 		TArray<FAudioMixerThreadCommand> DelayedCommands;
 
+		RenderThreadPhase = ESourceManagerRenderThreadPhase::PumpCmds;		
 		for (int32 Id = 0; Id < NumCommandsToExecute; ++Id)
 		{
-			FAudioMixerThreadCommand& AudioCommand = Commands.SourceCommandQueue[Id];
-
-			if(AudioCommand.bDeferExecution)
-			{
-				AudioCommand.bDeferExecution = false;
-				DelayedCommands.Emplace(MoveTemp(AudioCommand));
+			// First copy/move out the command and keep a copy of it.
+			{ 
+				FWriteScopeLock Lock(CurrentlyExecutingCmdLock);
+				CurrentlyExecuteingCmd = MoveTemp(Commands.SourceCommandQueue[Id]);
 			}
-			else
+			
+			// Execute the current command or differ under a read-lock.
 			{
-				AudioCommand(); // execute
+				FReadScopeLock Lock(CurrentlyExecutingCmdLock);
+				if (CurrentlyExecuteingCmd.bDeferExecution)
+				{
+					CurrentlyExecuteingCmd.bDeferExecution = false;
+					DelayedCommands.Add(CurrentlyExecuteingCmd);
+				}
+				else
+				{
+					CurrentlyExecuteingCmd(); // execute
+				}
 			}
-
+			
 			NumCommands.Decrement();
 		}
 
-
-		LastPumpTimeInCycles = FPlatformTime::Cycles64();
+		LastPumpCompleteTimeInCycles = FPlatformTime::Cycles64();
 		Commands.SourceCommandQueue = DelayedCommands;
 
 		if (FPlatformProcess::SupportsMultithreading())
@@ -3403,6 +3488,7 @@ namespace Audio
 
 	void FMixerSourceManager::UpdatePendingReleaseData(bool bForceWait)
 	{
+		RenderThreadPhase = ESourceManagerRenderThreadPhase::UpdatePendingReleaseData;
 		
 		// Don't block, but let tasks finish naturally
 		for (int32 i = PendingSourceBuffers.Num() - 1; i >= 0; --i)
@@ -3431,4 +3517,80 @@ namespace Audio
 		return bEnableBaseSubmix || bEnableSubmixSends;
 	}
 
+	void FMixerSourceManager::DoStallDiagnostics()
+	{
+		LogRenderThreadStall();
+		LogInflightAsyncTasks();
+		LogCallstacks();
+	}
+
+	void FMixerSourceManager::LogRenderThreadStall()
+	{
+		// If we are in either of the Cmd pump phases dump the current command.
+		if (RenderThreadPhase == ESourceManagerRenderThreadPhase::PumpMpscCmds ||
+			RenderThreadPhase == ESourceManagerRenderThreadPhase::PumpCmds)
+		{
+			FReadScopeLock Lock(CurrentlyExecutingCmdLock);
+			UE_LOG(LogAudioMixer, Warning, TEXT("Stall in Cmd Queue: Cmd='%s', Executing For: %2.5f secs, AudioRenderThread='%s'"),
+				*CurrentlyExecuteingCmd.GetSafeDebugString(), CurrentlyExecuteingCmd.GetExecuteTimeInSeconds(), ToCStr(LexToString(RenderThreadPhase)));
+		}
+		else
+		{
+			UE_LOG(LogAudioMixer, Warning, TEXT("Stall in AudioRenderThread Phase: '%s'"), ToCStr(LexToString(RenderThreadPhase)));
+		}
+	}
+
+	void FMixerSourceManager::LogInflightAsyncTasks()
+	{
+		// NOTE: we iterate these lists without a lock, so this is somewhat dangerous!
+
+		// Dump all in flight decodes/procedural sources.
+		using FSrcBuffer = TSharedPtr<FMixerSourceBuffer, ESPMode::ThreadSafe>;
+		TArray<FMixerSourceBuffer::FDiagnosticState> InflightTasks;
+		for (FSourceInfo& i : SourceInfos)
+		{
+			if (i.MixerSourceBuffer.IsValid())
+			{
+				FMixerSourceBuffer::FDiagnosticState State;
+				i.MixerSourceBuffer->GetDiagnosticState(State);
+				if (State.bInFlight)
+				{
+					InflightTasks.Add(State);
+				}
+			}
+		}
+		for (FSrcBuffer& i : PendingSourceBuffers)
+		{
+			FMixerSourceBuffer::FDiagnosticState State;
+			if (i.IsValid())
+			{
+				i->GetDiagnosticState(State);
+				if (State.bInFlight)
+				{
+					InflightTasks.Add(State);
+				}
+			}
+		}
+		for (FMixerSourceBuffer::FDiagnosticState& i : InflightTasks)
+		{
+			UE_LOG(LogAudioMixer, Warning, TEXT("Inflight Task: %s, %2.2f secs, Procedural=%d"),
+				*i.WaveName.ToString(), i.RunTimeInSecs, (int32)i.bProcedural);
+		}
+	}
+
+	void FMixerSourceManager::LogCallstacks()
+	{
+		LogCallstack(AudioRenderThreadId);
+	}
+
+	void FMixerSourceManager::LogCallstack(uint32 InThreadId)
+	{
+		if (InThreadId != INVALID_AUDIO_RENDER_THREAD_ID)
+		{
+			const SIZE_T StackTraceSize = 65536;
+			ANSICHAR StackTrace[StackTraceSize] = { 0 };
+			FPlatformStackWalk::ThreadStackWalkAndDump(StackTrace, StackTraceSize, 0, InThreadId);
+			UE_LOG(LogAudioMixer, Warning, TEXT("***** ThreadStackWalkAndDump for ThreadId(%lu) ******\n%s"), InThreadId, ANSI_TO_TCHAR(StackTrace));
+		}
+	}
 }
diff --git a/Engine/Source/Runtime/AudioMixer/Private/AudioMixerSourceManager.h b/Engine/Source/Runtime/AudioMixer/Private/AudioMixerSourceManager.h
--- a/Engine/Source/Runtime/AudioMixer/Private/AudioMixerSourceManager.h
+++ b/Engine/Source/Runtime/AudioMixer/Private/AudioMixerSourceManager.h
@@ -19,6 +19,12 @@
 #include "Sound/QuartzQuantizationUtilities.h"
 #include "Stats/Stats.h"
 
+#include "AudioMixerSourceManager.generated.h"
+
+// Default this to on (it's quite a small memory footprint).
+#ifndef WITH_AUDIO_MIXER_THREAD_COMMAND_DEBUG
+	#define WITH_AUDIO_MIXER_THREAD_COMMAND_DEBUG (1)
+#endif //WITH_AUDIO_MIXER_THREAD_COMMAND_DEBUG
 
 // Tracks the time it takes to up the source manager (computes source buffers, source effects, sample rate conversion)
 DECLARE_CYCLE_STAT_EXTERN(TEXT("Source Manager Update"), STAT_AudioMixerSourceManagerUpdate, STATGROUP_AudioMixer, AUDIOMIXER_API);
@@ -35,6 +41,27 @@ DECLARE_CYCLE_STAT_EXTERN(TEXT("Source Output Buffers"), STAT_AudioMixerSourceOu
 // The time it takes to process the HRTF effect.
 DECLARE_CYCLE_STAT_EXTERN(TEXT("HRTF"), STAT_AudioMixerHRTF, STATGROUP_AudioMixer, AUDIOMIXER_API);
 
+// For diagnostics, keep track of what phase of updating the Source manager is in currently.
+UENUM()
+enum ESourceManagerRenderThreadPhase: uint8
+{
+	Begin,
+	
+	PumpMpscCmds,	
+	PumpCmds,
+	ProcessModulators,
+	UpdatePendingReleaseData,
+	GenerateSrcAudio_WithBusses,
+	ComputeBusses,
+	GenerateSrcAudio_WithoutBusses,
+	UpdateBusses,
+	SpatialInterface_OnAllSourcesProcessed,
+	SourceDataOverride_OnAllSourcesProcessed,
+	UpdateGameThreadCopies,
+	
+	Finished,
+};
+
 namespace Audio
 {
 	class FMixerSubmix;
@@ -280,7 +307,7 @@ namespace Audio
 		void FlushCommandQueue(bool bPumpCommandQueue = false);
 
 		// Pushes a TFUnction command into an MPSC queue from an arbitrary thread to the audio render thread
-		void AudioMixerThreadMPSCCommand(TFunction<void()> InCommand);
+		void AudioMixerThreadMPSCCommand(TFunction<void()>&& InCommand, const char* InDebugString=nullptr);
 		
 	private:
 #define INVALID_AUDIO_RENDER_THREAD_ID static_cast<uint32>(-1)
@@ -303,22 +330,31 @@ namespace Audio
 		struct FAudioMixerThreadCommand
 		{
 			// ctor
-			FAudioMixerThreadCommand(TFunction<void()> InFunction, bool bInDeferExecution = false);
-
+			FAudioMixerThreadCommand() = default;
+			FAudioMixerThreadCommand(TFunction<void()>&& InFunction, const char* InDebugString, bool bInDeferExecution = false);
+			
 			// function-call operator
 			void operator()() const;
 
+			// data
 			TFunction<void()> Function;
 
 			// Defers the execution by a single call to PumpCommandQueue()
 			// (used for commands that affect a playing source,
 			// and that source gets initialized after the command executes
-			bool bDeferExecution;
+			bool bDeferExecution = false;
+
+#if WITH_AUDIO_MIXER_THREAD_COMMAND_DEBUG			
+			const char* DebugString=nullptr;					// Statically defined string from macro AUDIO_MIXER_THREAD_COMMAND_STRING
+			mutable uint64_t StartExecuteTimeInCycles=0;		// Set just before Function is called, for diagnostics.
+#endif // #if WITH_AUDIO_MIXER_THREAD_COMMAND_DEBUG
+
+			FString GetSafeDebugString() const;
+			float GetExecuteTimeInSeconds() const;
 		};
 
-		void AudioMixerThreadCommand(TFunction<void()> InFunction, bool bInDeferExecution = false);
+		void AudioMixerThreadCommand(TFunction<void()>&& InFunction, const char* InDebugString = nullptr, bool bInDeferExecution = false);
 
-		
 		static const int32 NUM_BYTES_PER_SAMPLE = 2;
 
 		// Private class which perform source buffer processing in a worker task
@@ -384,8 +420,8 @@ namespace Audio
 
 		TArray<int32> DebugSoloSources;
 
-		// Command queue to communicate commands to teh source manager from arbitrary threads
-		TMpscQueue<TFunction<void()>> MpscCommandQueue;
+		using FAudioMixerMpscCommand = FAudioMixerThreadCommand;
+		TMpscQueue<FAudioMixerMpscCommand> MpscCommandQueue;
 		
 		struct FSourceInfo
 		{
@@ -548,6 +584,14 @@ namespace Audio
 		static void ApplyDistanceAttenuation(FSourceInfo& InSourceInfo, int32 NumSamples);
 		void ComputePluginAudio(FSourceInfo& InSourceInfo, FMixerSourceSubmixOutputBuffer& InSourceSubmixOutputBuffer, int32 SourceId, int32 NumSamples);
 
+		// Hang/crash diagnostics.
+		void DoStallDiagnostics();
+
+		void LogRenderThreadStall();
+		void LogInflightAsyncTasks();
+		void LogCallstacks();
+		void LogCallstack(uint32 InThreadId);
+
 		// Array of listener transforms
 		TArray<FTransform> ListenerTransforms;
 
@@ -595,7 +639,10 @@ namespace Audio
 
 		// Set to true when the audio source manager should pump the command queue
 		FThreadSafeBool bPumpQueue;
-		uint64 LastPumpTimeInCycles = 0;
+		std::atomic<uint64> LastPumpCompleteTimeInCycles=0;
+		std::atomic<ESourceManagerRenderThreadPhase> RenderThreadPhase=ESourceManagerRenderThreadPhase::Begin;
+		FRWLock CurrentlyExecutingCmdLock;						// R/W slim lock for the currently executing cmd, so we can safely query it.
+		FAudioMixerThreadCommand CurrentlyExecuteingCmd;		// Keep this as a member so we can't always peek the executing cmd.
 
 		friend class FMixerSourceVoice;
 	};
