Fix for truncated wave data on file import.

diff --git a/Engine/Source/Editor/AudioEditor/Private/Factories/SoundFactory.cpp b/Engine/Source/Editor/AudioEditor/Private/Factories/SoundFactory.cpp
--- a/Engine/Source/Editor/AudioEditor/Private/Factories/SoundFactory.cpp
+++ b/Engine/Source/Editor/AudioEditor/Private/Factories/SoundFactory.cpp
@@ -371,6 +371,8 @@ UObject* USoundFactory::CreateObject
 		if (*WaveInfo.pBitsPerSample != 16)
 		{
 #if WITH_SNDFILE_IO
+			const uint32 OrigNumSamples = WaveInfo.GetNumSamples();
+
 			// Attempt to convert to 16 bit audio
 			if (Audio::ConvertAudioToWav(RawWaveData, ConvertedRawWaveData))
 			{
@@ -381,6 +383,10 @@ UObject* USoundFactory::CreateObject
 					GEditor->GetEditorSubsystem<UImportSubsystem>()->BroadcastAssetPostImport(this, nullptr);
 					return nullptr;
 				}
+
+				// Sanity check that the same number of samples exist in the converted file as the original
+				const uint32 ConvertedNumSamples = WaveInfo.GetNumSamples();
+				ensure(ConvertedNumSamples == OrigNumSamples);
 			}
 
 			// Copy over the data
diff --git a/Engine/Source/Runtime/AudioMixer/Private/SoundFileIO/SoundFileIO.cpp b/Engine/Source/Runtime/AudioMixer/Private/SoundFileIO/SoundFileIO.cpp
--- a/Engine/Source/Runtime/AudioMixer/Private/SoundFileIO/SoundFileIO.cpp
+++ b/Engine/Source/Runtime/AudioMixer/Private/SoundFileIO/SoundFileIO.cpp
@@ -18,6 +18,8 @@
 
 namespace Audio
 {
+	static void CopyOptionalWavChunks(TSharedPtr<ISoundFileReader>& InSoundDataReader, const int32 InInputFormat, TSharedPtr<ISoundFileWriter>& InSoundFileWriter, const int32 InOutputFormat);
+
 	bool AUDIOMIXER_API InitSoundFileIOManager()
 	{
 		return Audio::SoundFileIOManagerInit();
@@ -61,6 +63,9 @@ namespace Audio
 			return false;
 		}
 
+		// Copy optional chunks before writing data chunk which libsndfile assumes will be the last chunk
+		CopyOptionalWavChunks(InputSoundDataReader, InputDescription.FormatFlags, SoundFileWriter, NewSoundFileDescription.FormatFlags);
+
 		// Create a buffer to do the processing 
 		SoundFileCount ProcessBufferSamples = static_cast<SoundFileCount>(1024) * NewSoundFileDescription.NumChannels;
 		TArray<float> ProcessBuffer;
@@ -155,22 +160,6 @@ namespace Audio
 		Error = SoundFileWriter->WriteSamples((const float*)ProcessBuffer.GetData(), InputSamplesRead, SamplesWritten);
 		check(Error == ESoundFileError::Type::NONE);
 
-		// The required chunks ('fmt ' and 'data') have been written, now copy over the optional chunks.
-		FSoundFileChunkArray OptionalChunks;
-		Error = InputSoundDataReader->GetOptionalChunks(OptionalChunks);
-		if (Error != ESoundFileError::Type::NONE)
-		{
-			UE_LOG(LogAudioMixer, Error, TEXT("Error encountered while reading optional chunk data...skipping"));
-		}
-		else
-		{
-			Error = SoundFileWriter->WriteOptionalChunks(OptionalChunks);
-			if (Error != ESoundFileError::Type::NONE)
-			{
-				UE_LOG(LogAudioMixer, Error, TEXT("Error encountered while writing optional chunk data...skipping"));
-			}
-		}
-		
 		// Release the sound file handles as soon as we finished converting the file
 		InputSoundDataReader->Release();
 		SoundFileWriter->Release();
@@ -184,4 +173,28 @@ namespace Audio
 
 		return true;
 	}
+
+	void CopyOptionalWavChunks(TSharedPtr<ISoundFileReader>& InSoundDataReader, const int32 InInputFormat, TSharedPtr<ISoundFileWriter>& InSoundFileWriter, const int32 InOutputFormat)
+	{
+		// libsndfile only supports chunk operations with wave file formats
+		if ((InInputFormat & ESoundFileFormat::WAV) && (InOutputFormat & ESoundFileFormat::WAV))
+		{
+			// Get the optional chunks from the input data
+			FSoundFileChunkArray OptionalChunks;
+			ESoundFileError::Type Error = InSoundDataReader->GetOptionalChunks(OptionalChunks);
+			if (Error != ESoundFileError::Type::NONE)
+			{
+				UE_LOG(LogAudioMixer, Error, TEXT("Error encountered while reading optional chunk data...skipping"));
+			}
+			else
+			{
+				// Copy any chunks found over to the output file
+				Error = InSoundFileWriter->WriteOptionalChunks(OptionalChunks);
+				if (Error != ESoundFileError::Type::NONE)
+				{
+					UE_LOG(LogAudioMixer, Error, TEXT("Error encountered while writing optional chunk data...skipping"));
+				}
+			}
+		}
+	}
 }
diff --git a/Engine/Source/Runtime/AudioMixer/Private/SoundFileIO/SoundFileIOManagerImpl.cpp b/Engine/Source/Runtime/AudioMixer/Private/SoundFileIO/SoundFileIOManagerImpl.cpp
--- a/Engine/Source/Runtime/AudioMixer/Private/SoundFileIO/SoundFileIOManagerImpl.cpp
+++ b/Engine/Source/Runtime/AudioMixer/Private/SoundFileIO/SoundFileIOManagerImpl.cpp
@@ -39,6 +39,7 @@ namespace Audio
 	static const int32 SET_ENCODING_QUALITY = 0x1300;
 	static const int32 SET_CHANNEL_MAP_INFO = 0x1101;
 	static const int32 GET_CHANNEL_MAP_INFO = 0x1100;
+	static const int32 UPDATE_HEADER_NOW = 0x1060;
 
 	// Exported SoundFile Functions
 	typedef LibSoundFileHandle*(*SoundFileOpenFuncPtr)(const char* Path, int32 Mode, FSoundFileDescription* Description);
@@ -369,60 +370,89 @@ namespace Audio
 
 	static ESoundFileError::Type GetOptionalChunksInternal(LibSoundFileHandle* FileHandle, FSoundFileChunkArray& OutChunkInfoArray)
 	{
-		const TArray<uint32>& OptionalChunkIds = FWaveModInfo::GetOptionalWaveChunkIds();
-
-		for (const uint32 Id : OptionalChunkIds)
+		// Verify that the necessary library function pointers have been properly set
+		if (SoundFileGetChunkIterator != nullptr && SoundFileGetChunkSize != nullptr &&
+			SoundFileGetChunkData != nullptr && SoundFileNextChunkIterator != nullptr)
 		{
-			FSoundFileChunkInfo ChunkLookup;
-			ChunkLookup.ChunkIdSize = sizeof(Id);
-			FCStringAnsi::Strncpy(ChunkLookup.ChunkId, (ANSICHAR*)&(INTEL_ORDER32(Id)), ChunkLookup.ChunkIdSize + 1);
+			const TArray<uint32>& OptionalChunkIds = FWaveModInfo::GetOptionalWaveChunkIds();
 
-			// Lookup chunk of given Id. Multiple chunks can exist of a given type
-			// so we loop here.
-			LibSoundFileChunkIterator* ChunkItr = SoundFileGetChunkIterator(FileHandle, &ChunkLookup);
-			while (ChunkItr)
+			for (const uint32 Id : OptionalChunkIds)
 			{
-				FSoundFileChunkInfoWrapper ChunkInfo;
-				// SoundFileGetChunkSize retrieves the chunk data size. Oddly,
-				// it does not fill in the chunk Id.
-				int32 Result = SoundFileGetChunkSize(ChunkItr, ChunkInfo.GetPtr());
-				if (Result == 0 && ChunkInfo.GetPtr()->DataLength > 0)
-				{
-					ChunkInfo.AllocateChunkData();
-					// SoundFileGetChunkData copies in the chunk data and fills
-					// in the ChunkId.
-					Result = SoundFileGetChunkData(ChunkItr, ChunkInfo.GetPtr());
+				FSoundFileChunkInfo ChunkLookup;
+				ChunkLookup.ChunkIdSize = sizeof(Id);
+				FCStringAnsi::Strncpy(ChunkLookup.ChunkId, (ANSICHAR*)&(INTEL_ORDER32(Id)), ChunkLookup.ChunkIdSize + 1);
 
-					if (Result == 0)
-					{
-						OutChunkInfoArray.Add(MoveTemp(ChunkInfo));
-					}
-					else
+				// Lookup chunk of given Id. Multiple chunks can exist of a given type
+				// so we loop here.
+				LibSoundFileChunkIterator* ChunkItr = SoundFileGetChunkIterator(FileHandle, &ChunkLookup);
+				while (ChunkItr)
+				{
+					FSoundFileChunkInfoWrapper ChunkInfo;
+					// SoundFileGetChunkSize retrieves the chunk data size. Oddly,
+					// it does not fill in the chunk Id.
+					int32 Result = SoundFileGetChunkSize(ChunkItr, ChunkInfo.GetPtr());
+					if (Result == 0 && ChunkInfo.GetPtr()->DataLength > 0)
 					{
-						UE_LOG(LogAudioMixer, Error, TEXT("LibSoundFile unable to read invalid chunk: %s"), ChunkLookup.ChunkId);
-						return ESoundFileError::Type::INVALID_CHUNK;
+						ChunkInfo.AllocateChunkData();
+						// SoundFileGetChunkData copies in the chunk data and fills
+						// in the ChunkId.
+						Result = SoundFileGetChunkData(ChunkItr, ChunkInfo.GetPtr());
+
+						if (Result == 0)
+						{
+							OutChunkInfoArray.Add(MoveTemp(ChunkInfo));
+						}
+						else
+						{
+							UE_LOG(LogAudioMixer, Error, TEXT("LibSoundFile unable to read invalid chunk: %s"), *FString((ANSICHAR*)ChunkLookup.ChunkId));
+							return ESoundFileError::Type::INVALID_CHUNK;
+						}
 					}
-				}
 
-				ChunkItr = SoundFileNextChunkIterator(ChunkItr);
+					ChunkItr = SoundFileNextChunkIterator(ChunkItr);
+				}
 			}
 		}
+		else
+		{
+			UE_LOG(LogAudioMixer, Error, TEXT("LibSoundFile wasn't properly loaded with symbols for accessing wav chunk data."));
+		}
 
 		return ESoundFileError::Type::NONE;
 	}
 
 	static ESoundFileError::Type WriteOptionalChunksInternal(LibSoundFileHandle* FileHandle, const FSoundFileChunkArray& ChunkInfoArray)
 	{
-		for (const FSoundFileChunkInfoWrapper& ChunkInfo : ChunkInfoArray)
+		// Verify that the necessary library function pointers have been properly set
+		if (SoundFileSetChunk != nullptr && SoundFileCommand != nullptr)
 		{
-			// Note, libsndfile uses 4-byte pad when writing chunk data
-			int32 Result = SoundFileSetChunk(FileHandle, ChunkInfo.GetPtr());
-			if (Result)
+			for (const FSoundFileChunkInfoWrapper& ChunkInfo : ChunkInfoArray)
 			{
-				UE_LOG(LogAudioMixer, Error, TEXT("LibSoundFile failed to write chunk data; Result = %d"), Result);
-				return ESoundFileError::Type::INVALID_STATE;
+				// Note, libsndfile uses 4-byte pad when writing chunk data
+				int32 Result = SoundFileSetChunk(FileHandle, ChunkInfo.GetPtr());
+				if (Result)
+				{
+					UE_LOG(LogAudioMixer, Error, TEXT("LibSoundFile failed to write chunk data; Result = %d"), Result);
+					return ESoundFileError::Type::INVALID_STATE;
+				}
+				else
+				{
+					UE_LOG(LogAudioMixer, VeryVerbose, TEXT("Wrote ChunkId: %s, chunk DataLength: %d"), *FString((ANSICHAR*)ChunkInfo.GetPtr()->ChunkId), ChunkInfo.GetPtr()->DataLength);
+
+					// Update file header after adding a new chunk
+					Result = SoundFileCommand(FileHandle, UPDATE_HEADER_NOW, nullptr, 0);
+					if (Result)
+					{
+						UE_LOG(LogAudioMixer, Error, TEXT("LibSoundFile failed to update file header; Result = %d"), Result);
+						return ESoundFileError::Type::INVALID_STATE;
+					}
+				}
 			}
 		}
+		else
+		{
+			UE_LOG(LogAudioMixer, Error, TEXT("LibSoundFile wasn't properly loaded with symbols for accessing wav chunk data."));
+		}
 
 		return ESoundFileError::Type::NONE;
 	}
diff --git a/Engine/Source/Runtime/Engine/Private/Audio.cpp b/Engine/Source/Runtime/Engine/Private/Audio.cpp
--- a/Engine/Source/Runtime/Engine/Private/Audio.cpp
+++ b/Engine/Source/Runtime/Engine/Private/Audio.cpp
@@ -1615,16 +1615,6 @@ bool FWaveModInfo::ReadWaveInfo( const uint8* WaveData, int32 WaveDataSize, FStr
 	SampleDataSize = INTEL_ORDER32( RiffChunk->ChunkLen );
 	SampleDataEnd = SampleDataStart + SampleDataSize;
 
-	if (!InHeaderDataOnly && (uint8*)SampleDataEnd > (uint8*)WaveDataEnd)
-	{
-		UE_LOG(LogAudio, Warning, TEXT( "Wave data chunk is too big!" ) );
-
-		// Fix it up by clamping data chunk.
-		SampleDataEnd = (uint8*)WaveDataEnd;
-		SampleDataSize = SampleDataEnd - SampleDataStart;
-		RiffChunk->ChunkLen = INTEL_ORDER32( SampleDataSize );
-	}
-
 	if (*pFormatTag != 0x0001 // WAVE_FORMAT_PCM
 		&& *pFormatTag != 0x0002 // WAVE_FORMAT_ADPCM
 		&& *pFormatTag != 0x0011 // WAVE_FORMAT_DVI_ADPCM
@@ -1639,7 +1629,7 @@ bool FWaveModInfo::ReadWaveInfo( const uint8* WaveData, int32 WaveDataSize, FStr
 	{
 		if ((uint8*)SampleDataEnd > (uint8*)WaveDataEnd)
 		{
-			UE_LOG(LogAudio, Warning, TEXT("Wave data chunk is too big!" ));
+			UE_LOG(LogAudio, Warning, TEXT("Wave data chunk exceeds end of wave file by %d bytes, truncating"), (SampleDataEnd - WaveDataEnd));
 
 			// Fix it up by clamping data chunk.
 			SampleDataEnd = (uint8*)WaveDataEnd;
@@ -1909,6 +1899,16 @@ void FWaveModInfo::ReportImportFailure() const
 	}
 }
 
+uint32 FWaveModInfo::GetNumSamples() const
+{
+	if (*pBitsPerSample != 0)
+	{
+		return SampleDataSize / (*pBitsPerSample / 8);
+	}
+
+	return 0;
+}
+
 static void WriteUInt32ToByteArrayLE(TArray<uint8>& InByteArray, int32& Index, const uint32 Value)
 {
 	InByteArray[Index++] = (uint8)(Value >> 0);
diff --git a/Engine/Source/Runtime/Engine/Public/Audio.h b/Engine/Source/Runtime/Engine/Public/Audio.h
--- a/Engine/Source/Runtime/Engine/Public/Audio.h
+++ b/Engine/Source/Runtime/Engine/Public/Audio.h
@@ -935,6 +935,10 @@ public:
 	ENGINE_API bool ReadWaveHeader(const uint8* RawWaveData, int32 Size, int32 Offset );
 
 	ENGINE_API void ReportImportFailure() const;
+
+	/** Return total number of samples */
+	ENGINE_API uint32 GetNumSamples() const;
+
 };
 
 /** Utility to serialize raw PCM data into a wave file. */
