Guards around memory corruption.

diff --git a/Engine/Source/Runtime/AudioMixer/Private/AudioMixerSubmix.cpp b/Engine/Source/Runtime/AudioMixer/Private/AudioMixerSubmix.cpp
--- a/Engine/Source/Runtime/AudioMixer/Private/AudioMixerSubmix.cpp
+++ b/Engine/Source/Runtime/AudioMixer/Private/AudioMixerSubmix.cpp
@@ -740,7 +740,7 @@ namespace Audio
 		{
 			if (FadeInfo.bIsCurrentChain)
 			{
-				if (InIndex < FadeInfo.EffectChain.Num())
+				if (FadeInfo.EffectChain.IsValidIndex(InIndex))
 				{
 					FadeInfo.EffectChain[InIndex] = InEffectInstance;
 				}
