[Quartz] Fixed non-sample accurate issue w/ "QuantizedOtherClockStart".

diff --git a/Engine/Source/Runtime/AudioMixer/Private/Quartz/AudioMixerClock.cpp b/Engine/Source/Runtime/AudioMixer/Private/Quartz/AudioMixerClock.cpp
--- a/Engine/Source/Runtime/AudioMixer/Private/Quartz/AudioMixerClock.cpp
+++ b/Engine/Source/Runtime/AudioMixer/Private/Quartz/AudioMixerClock.cpp
@@ -307,17 +307,18 @@ namespace Audio
 		{
 			TickInternal(InNumFramesUntilNextTick, ClockAlteringPendingCommands, FramesOfLatency); // (process things like BPM changes first)
 			TickInternal(InNumFramesUntilNextTick, PendingCommands, FramesOfLatency);
+			Metronome.Tick(InNumFramesUntilNextTick, FramesOfLatency);
 		}
 		else
 		{
-			TickInternal(TickDelayLengthInFrames, ClockAlteringPendingCommands, FramesOfLatency);
-			TickInternal(TickDelayLengthInFrames, PendingCommands, FramesOfLatency);
+			const int32 FramesToTick = InNumFramesUntilNextTick - TickDelayLengthInFrames; 
+			TickInternal(FramesToTick, ClockAlteringPendingCommands, FramesOfLatency, TickDelayLengthInFrames);
+			TickInternal(FramesToTick, PendingCommands, FramesOfLatency, TickDelayLengthInFrames);
+			Metronome.Tick(FramesToTick, FramesOfLatency);
 
-			TickInternal(InNumFramesUntilNextTick - TickDelayLengthInFrames, ClockAlteringPendingCommands, FramesOfLatency, TickDelayLengthInFrames);
-			TickInternal(InNumFramesUntilNextTick - TickDelayLengthInFrames, PendingCommands, FramesOfLatency, TickDelayLengthInFrames);
+			TickDelayLengthInFrames = 0;
 		}
 
-		Metronome.Tick(InNumFramesUntilNextTick, FramesOfLatency);
 
 		UpdateCachedState();
 	}
diff --git a/Engine/Source/Runtime/AudioMixer/Private/Quartz/AudioMixerClockManager.cpp b/Engine/Source/Runtime/AudioMixer/Private/Quartz/AudioMixerClockManager.cpp
--- a/Engine/Source/Runtime/AudioMixer/Private/Quartz/AudioMixerClockManager.cpp
+++ b/Engine/Source/Runtime/AudioMixer/Private/Quartz/AudioMixerClockManager.cpp
@@ -458,7 +458,7 @@ namespace Audio
 			{
 				// if this clock is earlier in the array than the last clock we ticked,
 				// then it has already been ticked this update
-				if (&ClockPtr < &ActiveClocks[LastClockTickedIndex.GetValue()])
+				if (i < LastClockTickedIndex.GetValue())
 				{
 					return true;
 				}
