Zero all fields when constructing FSoundFileDescription because libsndfile's file open calls read some of them. Leaving them uninitialized can cause unexpected, intermittent behavior.

diff --git a/Engine/Source/Runtime/AudioMixer/Private/SoundFileIO/SoundFile.h b/Engine/Source/Runtime/AudioMixer/Private/SoundFileIO/SoundFile.h
--- a/Engine/Source/Runtime/AudioMixer/Private/SoundFileIO/SoundFile.h
+++ b/Engine/Source/Runtime/AudioMixer/Private/SoundFileIO/SoundFile.h
@@ -11,27 +11,30 @@ namespace Audio
 	using SoundFileCount = int64;
 
 	/**
- * Specifies a sound file description.
- */
+	 * Specifies a sound file description.
+	 * 
+	 * Note that libsndfile reads some of these fields (noteably FormatFlags and bIsSeekable)
+	 * at file open time so we zero them out at construction time to avoid unexpected/intermintent issues.
+	 */
 	struct FSoundFileDescription
 	{
 		/** The number of frames (interleaved samples) in the sound file. */
-		int64 NumFrames;
+		int64 NumFrames = 0;
 
 		/** The sample rate of the sound file. */
-		int32 SampleRate;
+		int32 SampleRate = 0;
 
 		/** The number of channels of the sound file. */
-		int32 NumChannels;
+		int32 NumChannels = 0;
 
 		/** The format flags of the sound file. */
-		int32 FormatFlags;
+		int32 FormatFlags = 0;
 
 		/** The number of sections of the sound file. */
-		int32 NumSections;
+		int32 NumSections = 0;
 
 		/** Whether or not the sound file is seekable. */
-		int32 bIsSeekable;
+		int32 bIsSeekable = 0;
 	};
 
 	struct FSoundFileConvertFormat
