Adding an Audio CSV stat for a few audio subsystems.

diff --git a/Engine/Source/Runtime/AudioMixer/Private/Quartz/QuartzSubsystem.cpp b/Engine/Source/Runtime/AudioMixer/Private/Quartz/QuartzSubsystem.cpp
--- a/Engine/Source/Runtime/AudioMixer/Private/Quartz/QuartzSubsystem.cpp
+++ b/Engine/Source/Runtime/AudioMixer/Private/Quartz/QuartzSubsystem.cpp
@@ -8,6 +8,7 @@
 #include "Quartz/AudioMixerClockManager.h"
 #include "Sound/QuartzQuantizationUtilities.h"
 #include "ProfilingDebugging/CountersTrace.h"
+#include "ProfilingDebugging/CsvProfiler.h"
 #include "Stats/Stats.h"
 
 #include "AudioDevice.h"
@@ -194,6 +195,7 @@ bool UQuartzSubsystem::DoesSupportWorldType(EWorldType::Type WorldType) const
 
 void UQuartzSubsystem::Tick(float DeltaTime)
 {
+	CSV_SCOPED_TIMING_STAT_EXCLUSIVE(Audio);
 	Super::Tick(DeltaTime);
 	TRACE_CPUPROFILER_EVENT_SCOPE(QuartzSubsystem::Tick);
 
