Add nullptr checks to audio bus subsystem

#rnx

diff --git a/Engine/Source/Runtime/AudioMixer/Private/AudioBusSubsystem.cpp b/Engine/Source/Runtime/AudioMixer/Private/AudioBusSubsystem.cpp
--- a/Engine/Source/Runtime/AudioMixer/Private/AudioBusSubsystem.cpp
+++ b/Engine/Source/Runtime/AudioMixer/Private/AudioBusSubsystem.cpp
@@ -47,16 +47,22 @@ void UAudioBusSubsystem::StartAudioBus(Audio::FAudioBusKey InAudioBusKey, int32
 
 		FAudioThread::RunCommandOnAudioThread([this, InAudioBusKey, InNumChannels, bInIsAutomatic]()
 		{
-			GetMutableSourceManager()->StartAudioBus(InAudioBusKey, InNumChannels, bInIsAutomatic);
+			if (Audio::FMixerSourceManager* MixerSourceManager = GetMutableSourceManager())
+			{
+				MixerSourceManager->StartAudioBus(InAudioBusKey, InNumChannels, bInIsAutomatic);
+			}
 		});
 	}
 	else
 	{
 		// If we're not the game thread, this needs to be on the game thread, so queue up a command to execute it on the game thread
-		GetMutableMixerDevice()->GameThreadMPSCCommand([this, InAudioBusKey, InNumChannels, bInIsAutomatic]
+		if (Audio::FMixerDevice* MixerDevice = GetMutableMixerDevice())
 		{
-			StartAudioBus(InAudioBusKey, InNumChannels, bInIsAutomatic);
-		});
+			MixerDevice->GameThreadMPSCCommand([this, InAudioBusKey, InNumChannels, bInIsAutomatic]
+			{
+				StartAudioBus(InAudioBusKey, InNumChannels, bInIsAutomatic);
+			});
+		}
 	}
 }
 
@@ -73,16 +79,22 @@ void UAudioBusSubsystem::StopAudioBus(Audio::FAudioBusKey InAudioBusKey)
 
 		FAudioThread::RunCommandOnAudioThread([this, InAudioBusKey]()
 		{
-			GetMutableSourceManager()->StopAudioBus(InAudioBusKey);
+			if (Audio::FMixerSourceManager* MixerSourceManager = GetMutableSourceManager())
+			{
+				MixerSourceManager->StopAudioBus(InAudioBusKey);
+			}
 		});
 	}
 	else
 	{
 		// If we're not the game thread, this needs to be on the game thread, so queue up a command to execute it on the game thread
-		GetMutableMixerDevice()->GameThreadMPSCCommand([this, InAudioBusKey]
+		if (Audio::FMixerDevice* MixerDevice = GetMutableMixerDevice())
 		{
-			StopAudioBus(InAudioBusKey);
-		});
+			MixerDevice->GameThreadMPSCCommand([this, InAudioBusKey]
+			{
+				StopAudioBus(InAudioBusKey);
+			});
+		}
 	}
 }
 
@@ -94,7 +106,11 @@ bool UAudioBusSubsystem::IsAudioBusActive(Audio::FAudioBusKey InAudioBusKey) con
 	}
 
 	check(IsInAudioThread());
-	return GetSourceManager()->IsAudioBusActive(InAudioBusKey);
+	if (const Audio::FMixerSourceManager* MixerSourceManager = GetSourceManager())
+	{
+		return MixerSourceManager->IsAudioBusActive(InAudioBusKey);
+	}
+	return false;
 }
 
 Audio::FPatchInput UAudioBusSubsystem::AddPatchInputForAudioBus(Audio::FAudioBusKey InAudioBusKey, int32 InFrames, int32 InChannels, float InGain)
@@ -108,14 +124,23 @@ Audio::FPatchInput UAudioBusSubsystem::AddPatchInputForAudioBus(Audio::FAudioBus
 	Audio::FMixerSourceManager* SourceManager = GetMutableSourceManager();
 	check(SourceManager);
 
+	Audio::FMixerDevice* MixerDevice = GetMutableMixerDevice();
+	if (!MixerDevice)
+	{
+		return FPatchInput();
+	}
+
 	if (IsInGameThread())
 	{
 		if (ActiveAudioBuses_GameThread.Find(InAudioBusKey))
 		{
-			FPatchInput PatchInput = GetMutableMixerDevice()->MakePatch(InFrames, InChannels, InGain);
+			FPatchInput PatchInput = MixerDevice->MakePatch(InFrames, InChannels, InGain);
 			FAudioThread::RunCommandOnAudioThread([this, PatchInput, InAudioBusKey]()
 			{
-				GetMutableSourceManager()->AddPatchInputForAudioBus_AudioThread(InAudioBusKey, PatchInput);
+				if (Audio::FMixerSourceManager* SourceManager = GetMutableSourceManager())
+				{
+					SourceManager->AddPatchInputForAudioBus_AudioThread(InAudioBusKey, PatchInput);
+				}
 			});
 			return PatchInput;
 		}
@@ -124,18 +149,18 @@ Audio::FPatchInput UAudioBusSubsystem::AddPatchInputForAudioBus(Audio::FAudioBus
 	}
 	else if (IsInAudioThread())
 	{
-		FPatchInput PatchInput = GetMutableMixerDevice()->MakePatch(InFrames, InChannels, InGain);
-		GetMutableSourceManager()->AddPatchInputForAudioBus_AudioThread(InAudioBusKey, PatchInput);
+		FPatchInput PatchInput = MixerDevice->MakePatch(InFrames, InChannels, InGain);
+		SourceManager->AddPatchInputForAudioBus_AudioThread(InAudioBusKey, PatchInput);
 		return PatchInput;
 	}
 	else if (GetMixerDevice()->IsAudioRenderingThread())
 	{
 		check(SourceManager);
 
-		const int32 NumChannels = GetSourceManager()->GetAudioBusNumChannels(InAudioBusKey);
+		const int32 NumChannels = SourceManager->GetAudioBusNumChannels(InAudioBusKey);
 		if (NumChannels > 0)
 		{
-			FPatchInput PatchInput = GetMutableMixerDevice()->MakePatch(InFrames, InChannels, InGain);
+			FPatchInput PatchInput = MixerDevice->MakePatch(InFrames, InChannels, InGain);
 			SourceManager->AddPatchInputForAudioBus(InAudioBusKey, PatchInput);
 			return PatchInput;
 		}
@@ -144,14 +169,17 @@ Audio::FPatchInput UAudioBusSubsystem::AddPatchInputForAudioBus(Audio::FAudioBus
 	}
 	else
 	{
-		FPatchInput PatchInput = GetMutableMixerDevice()->MakePatch(InFrames, InChannels, InGain);
-		GetMutableMixerDevice()->GameThreadMPSCCommand([this, PatchInput, InAudioBusKey]() mutable
+		FPatchInput PatchInput = MixerDevice->MakePatch(InFrames, InChannels, InGain);
+		MixerDevice->GameThreadMPSCCommand([this, PatchInput, InAudioBusKey]() mutable
 		{
 			if (ActiveAudioBuses_GameThread.Find(InAudioBusKey))
 			{
 				FAudioThread::RunCommandOnAudioThread([this, InAudioBusKey, PatchInput = MoveTemp(PatchInput)]()
 				{
-					GetMutableSourceManager()->AddPatchInputForAudioBus_AudioThread(InAudioBusKey, PatchInput);
+					if (Audio::FMixerSourceManager* SourceManager = GetMutableSourceManager())
+					{
+						SourceManager->AddPatchInputForAudioBus_AudioThread(InAudioBusKey, PatchInput);
+					}
 				});
 			}
 			else
@@ -173,14 +201,23 @@ Audio::FPatchOutputStrongPtr UAudioBusSubsystem::AddPatchOutputForAudioBus(Audio
 	Audio::FMixerSourceManager* SourceManager = GetMutableSourceManager();
 	check(SourceManager);
 
+	Audio::FMixerDevice* MixerDevice = GetMutableMixerDevice();
+	if (!MixerDevice)
+	{
+		return nullptr;
+	}
+
 	if (IsInGameThread())
 	{
 		if (ActiveAudioBuses_GameThread.Find(InAudioBusKey))
 		{
-			Audio::FPatchOutputStrongPtr PatchOutput = GetMutableMixerDevice()->MakePatch(InFrames, InChannels, InGain);
+			Audio::FPatchOutputStrongPtr PatchOutput = MixerDevice->MakePatch(InFrames, InChannels, InGain);
 			FAudioThread::RunCommandOnAudioThread([this, PatchOutput, InAudioBusKey]()
 			{
-				GetMutableSourceManager()->AddPatchOutputForAudioBus_AudioThread(InAudioBusKey, PatchOutput);
+				if (Audio::FMixerSourceManager* SourceManager = GetMutableSourceManager())
+				{
+					SourceManager->AddPatchOutputForAudioBus_AudioThread(InAudioBusKey, PatchOutput);
+				}
 			});
 			return PatchOutput;
 		}
@@ -189,18 +226,18 @@ Audio::FPatchOutputStrongPtr UAudioBusSubsystem::AddPatchOutputForAudioBus(Audio
 	}
 	else if (IsInAudioThread())
 	{
-		Audio::FPatchOutputStrongPtr PatchOutput = GetMutableMixerDevice()->MakePatch(InFrames, InChannels, InGain);
+		Audio::FPatchOutputStrongPtr PatchOutput = MixerDevice->MakePatch(InFrames, InChannels, InGain);
 		SourceManager->AddPatchOutputForAudioBus_AudioThread(InAudioBusKey, PatchOutput);
 		return PatchOutput;
 	}
-	else if (GetMixerDevice()->IsAudioRenderingThread())
+	else if (MixerDevice->IsAudioRenderingThread())
 	{
 		check(SourceManager);
 
-		const int32 NumChannels = GetMutableSourceManager()->GetAudioBusNumChannels(InAudioBusKey);
+		const int32 NumChannels = SourceManager->GetAudioBusNumChannels(InAudioBusKey);
 		if (NumChannels > 0)
 		{
-			Audio::FPatchOutputStrongPtr PatchOutput = GetMutableMixerDevice()->MakePatch(InFrames, InChannels, InGain);
+			Audio::FPatchOutputStrongPtr PatchOutput = MixerDevice->MakePatch(InFrames, InChannels, InGain);
 			SourceManager->AddPatchOutputForAudioBus(InAudioBusKey, PatchOutput);
 			return PatchOutput;
 		}
@@ -209,14 +246,17 @@ Audio::FPatchOutputStrongPtr UAudioBusSubsystem::AddPatchOutputForAudioBus(Audio
 	}
 	else
 	{
-		Audio::FPatchOutputStrongPtr PatchOutput = GetMutableMixerDevice()->MakePatch(InFrames, InChannels, InGain);
-		GetMutableMixerDevice()->GameThreadMPSCCommand([this, PatchOutput, InAudioBusKey]() mutable
+		Audio::FPatchOutputStrongPtr PatchOutput = MixerDevice->MakePatch(InFrames, InChannels, InGain);
+		MixerDevice->GameThreadMPSCCommand([this, PatchOutput, InAudioBusKey]() mutable
 		{
 			if (ActiveAudioBuses_GameThread.Find(InAudioBusKey))
 			{
 				FAudioThread::RunCommandOnAudioThread([this, InAudioBusKey, PatchOutput = MoveTemp(PatchOutput)]()
 				{
-					GetMutableSourceManager()->AddPatchOutputForAudioBus_AudioThread(InAudioBusKey, PatchOutput);
+					if (Audio::FMixerSourceManager* SourceManager = GetMutableSourceManager())
+					{
+						SourceManager->AddPatchOutputForAudioBus_AudioThread(InAudioBusKey, PatchOutput);
+					}
 				});
 			}
 			else
