[Engine]
* Moved dll export from types to methods to fix linker errors when building in debug where forceinline -> inline

diff --git a/Engine/Source/Runtime/AudioMixer/Public/AudioMixerDevice.h b/Engine/Source/Runtime/AudioMixer/Public/AudioMixerDevice.h
--- a/Engine/Source/Runtime/AudioMixer/Public/AudioMixerDevice.h
+++ b/Engine/Source/Runtime/AudioMixer/Public/AudioMixerDevice.h
@@ -87,145 +87,145 @@ namespace Audio
 	};
 
 
-	class AUDIOMIXER_API FMixerDevice :	public FAudioDevice,
-										public IAudioMixer,
-										public FGCObject
+	class FMixerDevice :	public FAudioDevice,
+							public IAudioMixer,
+							public FGCObject
 	{
 	public:
-		FMixerDevice(IAudioMixerPlatformInterface* InAudioMixerPlatform);
-		~FMixerDevice();
+		AUDIOMIXER_API FMixerDevice(IAudioMixerPlatformInterface* InAudioMixerPlatform);
+		AUDIOMIXER_API ~FMixerDevice();
 
 		//~ Begin FAudioDevice
-		virtual void UpdateDeviceDeltaTime() override;
-		virtual void GetAudioDeviceList(TArray<FString>& OutAudioDeviceNames) const override;
-		virtual bool InitializeHardware() override;
-		virtual void FadeIn() override;
-		virtual void FadeOut() override;
-		virtual void TeardownHardware() override;
-		virtual void UpdateHardwareTiming() override;
-		virtual void UpdateGameThread() override;
-		virtual void UpdateHardware() override;
-		virtual double GetAudioTime() const override;
-		virtual FAudioEffectsManager* CreateEffectsManager() override;
-		virtual FSoundSource* CreateSoundSource() override;
-		virtual FName GetRuntimeFormat(const USoundWave* SoundWave) const override;
-		virtual bool HasCompressedAudioInfoClass(USoundWave* SoundWave) override;
-		virtual bool SupportsRealtimeDecompression() const override;
-		virtual bool DisablePCMAudioCaching() const override;
-		virtual class ICompressedAudioInfo* CreateCompressedAudioInfo(const USoundWave* SoundWave) const override;
-		virtual class ICompressedAudioInfo* CreateCompressedAudioInfo(const FSoundWaveProxyPtr& SoundWave) const override;
-		virtual bool ValidateAPICall(const TCHAR* Function, uint32 ErrorCode) override;
+		AUDIOMIXER_API virtual void UpdateDeviceDeltaTime() override;
+		AUDIOMIXER_API virtual void GetAudioDeviceList(TArray<FString>& OutAudioDeviceNames) const override;
+		AUDIOMIXER_API virtual bool InitializeHardware() override;
+		AUDIOMIXER_API virtual void FadeIn() override;
+		AUDIOMIXER_API virtual void FadeOut() override;
+		AUDIOMIXER_API virtual void TeardownHardware() override;
+		AUDIOMIXER_API virtual void UpdateHardwareTiming() override;
+		AUDIOMIXER_API virtual void UpdateGameThread() override;
+		AUDIOMIXER_API virtual void UpdateHardware() override;
+		AUDIOMIXER_API virtual double GetAudioTime() const override;
+		AUDIOMIXER_API virtual FAudioEffectsManager* CreateEffectsManager() override;
+		AUDIOMIXER_API virtual FSoundSource* CreateSoundSource() override;
+		AUDIOMIXER_API virtual FName GetRuntimeFormat(const USoundWave* SoundWave) const override;
+		AUDIOMIXER_API virtual bool HasCompressedAudioInfoClass(USoundWave* SoundWave) override;
+		AUDIOMIXER_API virtual bool SupportsRealtimeDecompression() const override;
+		AUDIOMIXER_API virtual bool DisablePCMAudioCaching() const override;
+		AUDIOMIXER_API virtual class ICompressedAudioInfo* CreateCompressedAudioInfo(const USoundWave* SoundWave) const override;
+		AUDIOMIXER_API virtual class ICompressedAudioInfo* CreateCompressedAudioInfo(const FSoundWaveProxyPtr& SoundWave) const override;
+		AUDIOMIXER_API virtual bool ValidateAPICall(const TCHAR* Function, uint32 ErrorCode) override;
 #if UE_ALLOW_EXEC_COMMANDS
-		virtual bool Exec(UWorld* InWorld, const TCHAR* Cmd, FOutputDevice& Ar) override;
+		AUDIOMIXER_API virtual bool Exec(UWorld* InWorld, const TCHAR* Cmd, FOutputDevice& Ar) override;
 #endif
-		virtual void CountBytes(class FArchive& Ar) override;
-		virtual bool IsExernalBackgroundSoundActive() override;
-		virtual void ResumeContext() override;
-		virtual void SuspendContext() override;
-		virtual void EnableDebugAudioOutput() override;
-		virtual FAudioPlatformSettings GetPlatformSettings() const override;
-		virtual void RegisterSoundSubmix(USoundSubmixBase* SoundSubmix, bool bInit = true) override;
-		virtual void UnregisterSoundSubmix(const USoundSubmixBase* SoundSubmix) override;
-
-		virtual void InitSoundEffectPresets() override;
+		AUDIOMIXER_API virtual void CountBytes(class FArchive& Ar) override;
+		AUDIOMIXER_API virtual bool IsExernalBackgroundSoundActive() override;
+		AUDIOMIXER_API virtual void ResumeContext() override;
+		AUDIOMIXER_API virtual void SuspendContext() override;
+		AUDIOMIXER_API virtual void EnableDebugAudioOutput() override;
+		AUDIOMIXER_API virtual FAudioPlatformSettings GetPlatformSettings() const override;
+		AUDIOMIXER_API virtual void RegisterSoundSubmix(USoundSubmixBase* SoundSubmix, bool bInit = true) override;
+		AUDIOMIXER_API virtual void UnregisterSoundSubmix(const USoundSubmixBase* SoundSubmix) override;
+
+		AUDIOMIXER_API virtual void InitSoundEffectPresets() override;
 		UE_DEPRECATED(5.2, "The functionality for this has been moved to UAudioBusSubsystem::InitDefaultAudioBuses, which is now automatically called on subsystem creation.")
-		virtual void InitDefaultAudioBuses() override;
+		AUDIOMIXER_API virtual void InitDefaultAudioBuses() override;
 		UE_DEPRECATED(5.2, "The functionality for this has been moved to UAudioBusSubsystem::ShutdownDefaultAudioBuses.")
-		virtual void ShutdownDefaultAudioBuses() override;
-		virtual int32 GetNumActiveSources() const override;
+		AUDIOMIXER_API virtual void ShutdownDefaultAudioBuses() override;
+		AUDIOMIXER_API virtual int32 GetNumActiveSources() const override;
 
 		// Updates the source effect chain (using unique object id). 
-		virtual void UpdateSourceEffectChain(const uint32 SourceEffectChainId, const TArray<FSourceEffectChainEntry>& SourceEffectChain, const bool bPlayEffectChainTails) override;
-		virtual bool GetCurrentSourceEffectChain(const uint32 SourceEffectChainId, TArray<FSourceEffectChainEntry>& OutCurrentSourceEffectChainEntries) override;
+		AUDIOMIXER_API virtual void UpdateSourceEffectChain(const uint32 SourceEffectChainId, const TArray<FSourceEffectChainEntry>& SourceEffectChain, const bool bPlayEffectChainTails) override;
+		AUDIOMIXER_API virtual bool GetCurrentSourceEffectChain(const uint32 SourceEffectChainId, TArray<FSourceEffectChainEntry>& OutCurrentSourceEffectChainEntries) override;
 
 		// Submix dry/wet settings
-		virtual void UpdateSubmixProperties(USoundSubmixBase* InSubmix) override;
-		virtual void SetSubmixWetDryLevel(USoundSubmix* InSoundSubmix, float InOutputVolume, float InWetLevel, float InDryLevel) override;
-		virtual void SetSubmixOutputVolume(USoundSubmix* InSoundSubmix, float InOutputVolume) override;
-		virtual void SetSubmixWetLevel(USoundSubmix* InSoundSubmix, float InWetLevel) override;
-		virtual void SetSubmixDryLevel(USoundSubmix* InSoundSubmix, float InDryLevel) override;
+		AUDIOMIXER_API virtual void UpdateSubmixProperties(USoundSubmixBase* InSubmix) override;
+		AUDIOMIXER_API virtual void SetSubmixWetDryLevel(USoundSubmix* InSoundSubmix, float InOutputVolume, float InWetLevel, float InDryLevel) override;
+		AUDIOMIXER_API virtual void SetSubmixOutputVolume(USoundSubmix* InSoundSubmix, float InOutputVolume) override;
+		AUDIOMIXER_API virtual void SetSubmixWetLevel(USoundSubmix* InSoundSubmix, float InWetLevel) override;
+		AUDIOMIXER_API virtual void SetSubmixDryLevel(USoundSubmix* InSoundSubmix, float InDryLevel) override;
 
 		// Submix auto-disable setteings
-		virtual void SetSubmixAutoDisable(USoundSubmix* InSoundSubmix, bool bInAutoDisable) override;
-		virtual void SetSubmixAutoDisableTime(USoundSubmix* InSoundSubmix, float InDisableTime) override;
+		AUDIOMIXER_API virtual void SetSubmixAutoDisable(USoundSubmix* InSoundSubmix, bool bInAutoDisable) override;
+		AUDIOMIXER_API virtual void SetSubmixAutoDisableTime(USoundSubmix* InSoundSubmix, float InDisableTime) override;
 
 		// Submix Modulation Settings
-		virtual void UpdateSubmixModulationSettings(USoundSubmix* InSoundSubmix, const TSet<TObjectPtr<USoundModulatorBase>>& InOutputModulation, const TSet<TObjectPtr<USoundModulatorBase>>& InWetLevelModulation, const TSet<TObjectPtr<USoundModulatorBase>>& InDryLevelModulation) override;
-		virtual void SetSubmixModulationBaseLevels(USoundSubmix* InSoundSubmix, float InVolumeModBase, float InWetModBase, float InDryModBase) override;
+		AUDIOMIXER_API virtual void UpdateSubmixModulationSettings(USoundSubmix* InSoundSubmix, const TSet<TObjectPtr<USoundModulatorBase>>& InOutputModulation, const TSet<TObjectPtr<USoundModulatorBase>>& InWetLevelModulation, const TSet<TObjectPtr<USoundModulatorBase>>& InDryLevelModulation) override;
+		AUDIOMIXER_API virtual void SetSubmixModulationBaseLevels(USoundSubmix* InSoundSubmix, float InVolumeModBase, float InWetModBase, float InDryModBase) override;
 
 		// Submix effect chain override settings
-		virtual void SetSubmixEffectChainOverride(USoundSubmix* InSoundSubmix, const TArray<FSoundEffectSubmixPtr>& InSubmixEffectPresetChain, float InFadeTimeSec) override;
-		virtual void ClearSubmixEffectChainOverride(USoundSubmix* InSoundSubmix, float InFadeTimeSec) override;
+		AUDIOMIXER_API virtual void SetSubmixEffectChainOverride(USoundSubmix* InSoundSubmix, const TArray<FSoundEffectSubmixPtr>& InSubmixEffectPresetChain, float InFadeTimeSec) override;
+		AUDIOMIXER_API virtual void ClearSubmixEffectChainOverride(USoundSubmix* InSoundSubmix, float InFadeTimeSec) override;
 
 		// Submix recording callbacks:
-		virtual void StartRecording(USoundSubmix* InSubmix, float ExpectedRecordingDuration) override;
-		virtual Audio::FAlignedFloatBuffer& StopRecording(USoundSubmix* InSubmix, float& OutNumChannels, float& OutSampleRate) override;
+		AUDIOMIXER_API virtual void StartRecording(USoundSubmix* InSubmix, float ExpectedRecordingDuration) override;
+		AUDIOMIXER_API virtual Audio::FAlignedFloatBuffer& StopRecording(USoundSubmix* InSubmix, float& OutNumChannels, float& OutSampleRate) override;
 
-		virtual void PauseRecording(USoundSubmix* InSubmix);
-		virtual void ResumeRecording(USoundSubmix* InSubmix);
+		AUDIOMIXER_API virtual void PauseRecording(USoundSubmix* InSubmix);
+		AUDIOMIXER_API virtual void ResumeRecording(USoundSubmix* InSubmix);
 
 		// Submix envelope following
-		virtual void StartEnvelopeFollowing(USoundSubmix* InSubmix) override;
-		virtual void StopEnvelopeFollowing(USoundSubmix* InSubmix) override;
-		virtual void AddEnvelopeFollowerDelegate(USoundSubmix* InSubmix, const FOnSubmixEnvelopeBP& OnSubmixEnvelopeBP) override;
+		AUDIOMIXER_API virtual void StartEnvelopeFollowing(USoundSubmix* InSubmix) override;
+		AUDIOMIXER_API virtual void StopEnvelopeFollowing(USoundSubmix* InSubmix) override;
+		AUDIOMIXER_API virtual void AddEnvelopeFollowerDelegate(USoundSubmix* InSubmix, const FOnSubmixEnvelopeBP& OnSubmixEnvelopeBP) override;
 
 		// Submix Spectrum Analysis
-		virtual void StartSpectrumAnalysis(USoundSubmix* InSubmix, const FSoundSpectrumAnalyzerSettings& InSettings) override;
-		virtual void StopSpectrumAnalysis(USoundSubmix* InSubmix) override;
-		virtual void GetMagnitudesForFrequencies(USoundSubmix* InSubmix, const TArray<float>& InFrequencies, TArray<float>& OutMagnitudes) override;
-		virtual void GetPhasesForFrequencies(USoundSubmix* InSubmix, const TArray<float>& InFrequencies, TArray<float>& OutPhases) override;
-		virtual void AddSpectralAnalysisDelegate(USoundSubmix* InSubmix, const FSoundSpectrumAnalyzerDelegateSettings& InDelegateSettings, const FOnSubmixSpectralAnalysisBP& OnSubmixSpectralAnalysisBP) override;
-		virtual void RemoveSpectralAnalysisDelegate(USoundSubmix* InSubmix, const FOnSubmixSpectralAnalysisBP& OnSubmixSpectralAnalysisBP) override;
+		AUDIOMIXER_API virtual void StartSpectrumAnalysis(USoundSubmix* InSubmix, const FSoundSpectrumAnalyzerSettings& InSettings) override;
+		AUDIOMIXER_API virtual void StopSpectrumAnalysis(USoundSubmix* InSubmix) override;
+		AUDIOMIXER_API virtual void GetMagnitudesForFrequencies(USoundSubmix* InSubmix, const TArray<float>& InFrequencies, TArray<float>& OutMagnitudes) override;
+		AUDIOMIXER_API virtual void GetPhasesForFrequencies(USoundSubmix* InSubmix, const TArray<float>& InFrequencies, TArray<float>& OutPhases) override;
+		AUDIOMIXER_API virtual void AddSpectralAnalysisDelegate(USoundSubmix* InSubmix, const FSoundSpectrumAnalyzerDelegateSettings& InDelegateSettings, const FOnSubmixSpectralAnalysisBP& OnSubmixSpectralAnalysisBP) override;
+		AUDIOMIXER_API virtual void RemoveSpectralAnalysisDelegate(USoundSubmix* InSubmix, const FOnSubmixSpectralAnalysisBP& OnSubmixSpectralAnalysisBP) override;
 
 		// Submix buffer listener callbacks
-		virtual void RegisterSubmixBufferListener(ISubmixBufferListener* InSubmixBufferListener, USoundSubmix* InSubmix = nullptr) override;
-		virtual void UnregisterSubmixBufferListener(ISubmixBufferListener* InSubmixBufferListener, USoundSubmix* InSubmix = nullptr) override;
+		AUDIOMIXER_API virtual void RegisterSubmixBufferListener(ISubmixBufferListener* InSubmixBufferListener, USoundSubmix* InSubmix = nullptr) override;
+		AUDIOMIXER_API virtual void UnregisterSubmixBufferListener(ISubmixBufferListener* InSubmixBufferListener, USoundSubmix* InSubmix = nullptr) override;
 
-		virtual FPatchOutputStrongPtr AddPatchForSubmix(uint32 InObjectId, float InPatchGain) override;
+		AUDIOMIXER_API virtual FPatchOutputStrongPtr AddPatchForSubmix(uint32 InObjectId, float InPatchGain) override;
 
-		virtual void FlushExtended(UWorld* WorldToFlush, bool bClearActivatedReverb);
-		virtual void FlushAudioRenderingCommands(bool bPumpSynchronously = false) override;
+		AUDIOMIXER_API virtual void FlushExtended(UWorld* WorldToFlush, bool bClearActivatedReverb);
+		AUDIOMIXER_API virtual void FlushAudioRenderingCommands(bool bPumpSynchronously = false) override;
 
 		// Audio Device Properties
-		virtual bool IsNonRealtime() const override;
+		AUDIOMIXER_API virtual bool IsNonRealtime() const override;
 
 		//~ End FAudioDevice
 
 		//~ Begin IAudioMixer
-		virtual bool OnProcessAudioStream(FAlignedFloatBuffer& OutputBuffer) override;
-		virtual void OnAudioStreamShutdown() override;
+		AUDIOMIXER_API virtual bool OnProcessAudioStream(FAlignedFloatBuffer& OutputBuffer) override;
+		AUDIOMIXER_API virtual void OnAudioStreamShutdown() override;
 		//~ End IAudioMixer
 
 		//~ Begin FGCObject
-		virtual void AddReferencedObjects(FReferenceCollector& Collector) override;
+		AUDIOMIXER_API virtual void AddReferencedObjects(FReferenceCollector& Collector) override;
 		virtual FString GetReferencerName() const override
 		{
 			return TEXT("Audio::FMixerDevice");
 		}
 		//~End FGCObject
 
-		FMixerSubmixPtr FindSubmixInstanceByObjectId(uint32 InObjectId);
+		AUDIOMIXER_API FMixerSubmixPtr FindSubmixInstanceByObjectId(uint32 InObjectId);
 
-		FMixerSubmixWeakPtr GetSubmixInstance(const USoundSubmixBase* SoundSubmix);
+		AUDIOMIXER_API FMixerSubmixWeakPtr GetSubmixInstance(const USoundSubmixBase* SoundSubmix);
 
 		// If SoundSubmix is a soundfield submix, this will return the factory used to encode 
 		// source audio to it's soundfield format.
 		// Otherwise, returns nullptr.
-		ISoundfieldFactory* GetFactoryForSubmixInstance(USoundSubmix* SoundSubmix);
-		ISoundfieldFactory* GetFactoryForSubmixInstance(FMixerSubmixWeakPtr& SoundSubmixPtr);
+		AUDIOMIXER_API ISoundfieldFactory* GetFactoryForSubmixInstance(USoundSubmix* SoundSubmix);
+		AUDIOMIXER_API ISoundfieldFactory* GetFactoryForSubmixInstance(FMixerSubmixWeakPtr& SoundSubmixPtr);
 
 		// Functions which check the thread it's called on and helps make sure functions are called from correct threads
-		void CheckAudioThread() const;
-		void CheckAudioRenderingThread() const;
-		bool IsAudioRenderingThread() const;
+		AUDIOMIXER_API void CheckAudioThread() const;
+		AUDIOMIXER_API void CheckAudioRenderingThread() const;
+		AUDIOMIXER_API bool IsAudioRenderingThread() const;
 
 		// Public Functions
-		FMixerSourceVoice* GetMixerSourceVoice();
-		void ReleaseMixerSourceVoice(FMixerSourceVoice* InSourceVoice);
-		int32 GetNumSources() const;
+		AUDIOMIXER_API FMixerSourceVoice* GetMixerSourceVoice();
+		AUDIOMIXER_API void ReleaseMixerSourceVoice(FMixerSourceVoice* InSourceVoice);
+		AUDIOMIXER_API int32 GetNumSources() const;
 
 		// AudioLink
-		IAudioLinkFactory* GetAudioLinkFactory() const;
+		AUDIOMIXER_API IAudioLinkFactory* GetAudioLinkFactory() const;
 
 		const FAudioPlatformDeviceInfo& GetPlatformDeviceInfo() const { return PlatformInfo; };
 
@@ -239,55 +239,55 @@ namespace Audio
 		IAudioMixerPlatformInterface* GetAudioMixerPlatform() const { return AudioMixerPlatform; }
 
 		// Builds a 3D channel map for a spatialized source.
-		void Get3DChannelMap(const int32 InSubmixNumChannels, const FWaveInstance* InWaveInstance, const float EmitterAzimuth, const float NonSpatiliazedFactor, const TMap<EAudioMixerChannel::Type, float>* InOmniMap, float InDefaultOmniValue, Audio::FAlignedFloatBuffer& OutChannelMap);
+		AUDIOMIXER_API void Get3DChannelMap(const int32 InSubmixNumChannels, const FWaveInstance* InWaveInstance, const float EmitterAzimuth, const float NonSpatiliazedFactor, const TMap<EAudioMixerChannel::Type, float>* InOmniMap, float InDefaultOmniValue, Audio::FAlignedFloatBuffer& OutChannelMap);
 
 		// Builds a channel gain matrix for a non-spatialized source. The non-static variation of this function queries AudioMixerDevice->NumOutputChannels directly which may not be thread safe.
-		void Get2DChannelMap(bool bIsVorbis, const int32 NumSourceChannels, const bool bIsCenterChannelOnly, Audio::FAlignedFloatBuffer& OutChannelMap) const;
-		static void Get2DChannelMap(bool bIsVorbis, const int32 NumSourceChannels, const int32 NumOutputChannels, const bool bIsCenterChannelOnly, Audio::FAlignedFloatBuffer& OutChannelMap);
+		AUDIOMIXER_API void Get2DChannelMap(bool bIsVorbis, const int32 NumSourceChannels, const bool bIsCenterChannelOnly, Audio::FAlignedFloatBuffer& OutChannelMap) const;
+		AUDIOMIXER_API static void Get2DChannelMap(bool bIsVorbis, const int32 NumSourceChannels, const int32 NumOutputChannels, const bool bIsCenterChannelOnly, Audio::FAlignedFloatBuffer& OutChannelMap);
 
-		int32 GetDeviceSampleRate() const;
-		int32 GetDeviceOutputChannels() const;
+		AUDIOMIXER_API int32 GetDeviceSampleRate() const;
+		AUDIOMIXER_API int32 GetDeviceOutputChannels() const;
 
-		FMixerSourceManager* GetSourceManager();
-		const FMixerSourceManager* GetSourceManager() const;
+		AUDIOMIXER_API FMixerSourceManager* GetSourceManager();
+		AUDIOMIXER_API const FMixerSourceManager* GetSourceManager() const;
 
-		FMixerSubmixWeakPtr GetMasterSubmix(); 
-		FMixerSubmixWeakPtr GetBaseDefaultSubmix();
-		FMixerSubmixWeakPtr GetMasterReverbSubmix();
-		FMixerSubmixWeakPtr GetMasterEQSubmix();
+		AUDIOMIXER_API FMixerSubmixWeakPtr GetMasterSubmix(); 
+		AUDIOMIXER_API FMixerSubmixWeakPtr GetBaseDefaultSubmix();
+		AUDIOMIXER_API FMixerSubmixWeakPtr GetMasterReverbSubmix();
+		AUDIOMIXER_API FMixerSubmixWeakPtr GetMasterEQSubmix();
 
 		// Add submix effect to master submix
-		void AddMasterSubmixEffect(FSoundEffectSubmixPtr SoundEffect);
+		AUDIOMIXER_API void AddMasterSubmixEffect(FSoundEffectSubmixPtr SoundEffect);
 		
 		// Remove submix effect from master submix
-		void RemoveMasterSubmixEffect(uint32 SubmixEffectId);
+		AUDIOMIXER_API void RemoveMasterSubmixEffect(uint32 SubmixEffectId);
 		
 		// Clear all submix effects from master submix
-		void ClearMasterSubmixEffects();
+		AUDIOMIXER_API void ClearMasterSubmixEffects();
 
 		// Add submix effect to given submix
-		int32 AddSubmixEffect(USoundSubmix* InSoundSubmix, FSoundEffectSubmixPtr SoundEffect);
+		AUDIOMIXER_API int32 AddSubmixEffect(USoundSubmix* InSoundSubmix, FSoundEffectSubmixPtr SoundEffect);
 
 		// Remove submix effect to given submix
-		void RemoveSubmixEffect(USoundSubmix* InSoundSubmix, uint32 SubmixEffectId);
+		AUDIOMIXER_API void RemoveSubmixEffect(USoundSubmix* InSoundSubmix, uint32 SubmixEffectId);
 
 		// Remove submix effect at the given submix chain index
-		void RemoveSubmixEffectAtIndex(USoundSubmix* InSoundSubmix, int32 SubmixChainIndex);
+		AUDIOMIXER_API void RemoveSubmixEffectAtIndex(USoundSubmix* InSoundSubmix, int32 SubmixChainIndex);
 
 		// Replace the submix effect of the given submix at the submix chain index with the new submix effect id and submix instance
-		void ReplaceSoundEffectSubmix(USoundSubmix* InSoundSubmix, int32 InSubmixChainIndex, FSoundEffectSubmixPtr SoundEffect);
+		AUDIOMIXER_API void ReplaceSoundEffectSubmix(USoundSubmix* InSoundSubmix, int32 InSubmixChainIndex, FSoundEffectSubmixPtr SoundEffect);
 
 		// Clear all submix effects from given submix
-		void ClearSubmixEffects(USoundSubmix* InSoundSubmix);
+		AUDIOMIXER_API void ClearSubmixEffects(USoundSubmix* InSoundSubmix);
 
 		// Returns the channel array for the given submix channel type
-		const TArray<EAudioMixerChannel::Type>& GetChannelArray() const;
+		AUDIOMIXER_API const TArray<EAudioMixerChannel::Type>& GetChannelArray() const;
 
 		// Retrieves the listener transforms
-		const TArray<FTransform>* GetListenerTransforms();
+		AUDIOMIXER_API const TArray<FTransform>* GetListenerTransforms();
 
 		// Retrieves spherical locations of channels for a given submix format
-		const FChannelPositionInfo* GetDefaultChannelPositions() const;
+		AUDIOMIXER_API const FChannelPositionInfo* GetDefaultChannelPositions() const;
 
 		// Audio thread tick timing relative to audio render thread timing
 		double GetAudioThreadTime() const { return AudioThreadTimingData.AudioThreadTime; }
@@ -296,31 +296,31 @@ namespace Audio
 
 		EMonoChannelUpmixMethod GetMonoChannelUpmixMethod() const { return MonoChannelUpmixMethod; }
 
-		TArray<Audio::FChannelPositionInfo>* GetDefaultPositionMap(int32 NumChannels);
+		AUDIOMIXER_API TArray<Audio::FChannelPositionInfo>* GetDefaultPositionMap(int32 NumChannels);
 
-		static bool IsEndpointSubmix(const USoundSubmixBase* InSubmix);
+		static AUDIOMIXER_API bool IsEndpointSubmix(const USoundSubmixBase* InSubmix);
 
 		// Audio bus API - these are deprecated. Use corresponding calls in UAudioBusSubsystem instead
 		UE_DEPRECATED(5.2, "This function is deprecated. Use UAudioBusSubsystem::StartAudioBus instead.")
-		virtual void StartAudioBus(uint32 InAudioBusId, int32 InNumChannels, bool bInIsAutomatic) override;
+		AUDIOMIXER_API virtual void StartAudioBus(uint32 InAudioBusId, int32 InNumChannels, bool bInIsAutomatic) override;
 		UE_DEPRECATED(5.2, "This function is deprecated. Use UAudioBusSubsystem::StopAudioBus instead.")
-		virtual void StopAudioBus(uint32 InAudioBusId) override;
+		AUDIOMIXER_API virtual void StopAudioBus(uint32 InAudioBusId) override;
 		UE_DEPRECATED(5.2, "This function is deprecated. Use UAudioBusSubsystem::IsAudioBusActive instead.")
-		virtual bool IsAudioBusActive(uint32 InAudioBusId) const override;
+		AUDIOMIXER_API virtual bool IsAudioBusActive(uint32 InAudioBusId) const override;
 
 		UE_DEPRECATED(5.2, "AddPatchForAudioBus is deprecated.  Use UAudioBusSubsystem::AddPatchOutputForAudioBus.")
-		virtual FPatchOutputStrongPtr AddPatchForAudioBus(uint32 InAudioBusId, float InPatchGain = 1.0f) override;
+		AUDIOMIXER_API virtual FPatchOutputStrongPtr AddPatchForAudioBus(uint32 InAudioBusId, float InPatchGain = 1.0f) override;
 
 		UE_DEPRECATED(5.2, "AddPatchForAudioBus_GameThread is deprecated.  Use UAudioBusSubsystem::AddPatchOutputForAudioBus.")
-		virtual FPatchOutputStrongPtr AddPatchForAudioBus_GameThread(uint32 InAudioBusId, float InPatchGain = 1.0f) override;
+		AUDIOMIXER_API virtual FPatchOutputStrongPtr AddPatchForAudioBus_GameThread(uint32 InAudioBusId, float InPatchGain = 1.0f) override;
 
 		UE_DEPRECATED(5.2, "This overload of AddPatchInputForAudioBus is deprecated and non-functional.  Use the overload in UAudioBusSubsystem that takes the number of frames and channels as parameters.")
-		virtual void AddPatchInputForAudioBus(const FPatchInput& InPatchInput, uint32 InAudioBusId, float InPatchGain = 1.0f) override;
+		AUDIOMIXER_API virtual void AddPatchInputForAudioBus(const FPatchInput& InPatchInput, uint32 InAudioBusId, float InPatchGain = 1.0f) override;
 
 		UE_DEPRECATED(5.2, "AddPatchInputForAudioBus_GameThread is deprecated.  Use UAudioBusSubsystem::AddPatchInputForAudioBus instead.")
-		virtual void AddPatchInputForAudioBus_GameThread(const FPatchInput& InPatchInput, uint32 InAudioBusId, float InPatchGain = 1.0f) override;
+		AUDIOMIXER_API virtual void AddPatchInputForAudioBus_GameThread(const FPatchInput& InPatchInput, uint32 InAudioBusId, float InPatchGain = 1.0f) override;
 
-		FPatchOutputStrongPtr MakePatch(int32 InFrames, int32 InChannels, float InGain) const;
+		AUDIOMIXER_API FPatchOutputStrongPtr MakePatch(int32 InFrames, int32 InChannels, float InGain) const;
 
 		// Clock Manager for quantized event handling on Audio Render Thread
 		FQuartzClockManager QuantizedEventClockManager;
@@ -334,16 +334,16 @@ namespace Audio
 		FCriticalSection QuartzPersistentStateCritSec;
 
 		// Pushes the command to a audio render thread command queue to be executed on render thread
-		void AudioRenderThreadCommand(TFunction<void()> Command);
+		AUDIOMIXER_API void AudioRenderThreadCommand(TFunction<void()> Command);
 
 		// Pushes the command to a MPSC queue to be executed on the game thread
-		void GameThreadMPSCCommand(TFunction<void()> InCommand);
+		AUDIOMIXER_API void GameThreadMPSCCommand(TFunction<void()> InCommand);
 
 	protected:
 
-		virtual void InitSoundSubmixes() override;
+		AUDIOMIXER_API virtual void InitSoundSubmixes() override;
 
-		virtual void OnListenerUpdated(const TArray<FListener>& InListeners) override;
+		AUDIOMIXER_API virtual void OnListenerUpdated(const TArray<FListener>& InListeners) override;
 
 		TArray<FTransform> ListenerTransforms;
 
@@ -467,14 +467,14 @@ namespace Audio
 		// Creates a queue for audio decode requests with a specific Id. Tasks
 		// created with this Id will not be started immediately upon creation,
 		// but will instead be queued up to await a start "kick" later.
-		static void CreateSynchronizedAudioTaskQueue(AudioTaskQueueId QueueId);
+		AUDIOMIXER_API static void CreateSynchronizedAudioTaskQueue(AudioTaskQueueId QueueId);
 
 		// Destroys an audio decode task queue. Tasks currently queued up are 
 		// optionally started.
-		static void DestroySynchronizedAudioTaskQueue(AudioTaskQueueId QueueId, bool RunCurrentQueue = false);
+		AUDIOMIXER_API static void DestroySynchronizedAudioTaskQueue(AudioTaskQueueId QueueId, bool RunCurrentQueue = false);
 
 		// "Kicks" all of the audio decode tasks currentlyt in the queue.
-		static int KickQueuedTasks(AudioTaskQueueId QueueId);
+		AUDIOMIXER_API static int KickQueuedTasks(AudioTaskQueueId QueueId);
 	};
 }
 
