[Engine]

diff --git a/Engine/Source/Runtime/AudioExtensions/Public/ISoundfieldEndpoint.h b/Engine/Source/Runtime/AudioExtensions/Public/ISoundfieldEndpoint.h
--- a/Engine/Source/Runtime/AudioExtensions/Public/ISoundfieldEndpoint.h
+++ b/Engine/Source/Runtime/AudioExtensions/Public/ISoundfieldEndpoint.h
@@ -37,7 +37,7 @@ template <typename FuncType> class TFunctionRef;
  * This interface should be used to provide a non-uclass version of the data described in
  * your implementation of USoundfieldEndpointSettingsBase.
  */
-class AUDIOEXTENSIONS_API ISoundfieldEndpointSettingsProxy
+class ISoundfieldEndpointSettingsProxy
 {
 public:
 	virtual ~ISoundfieldEndpointSettingsProxy() {};
@@ -47,51 +47,51 @@ public:
  * This opaque class should be used for specifying settings for how audio should be
  * send to an external endpoint.
  */
-UCLASS(config = Engine, abstract, editinlinenew, BlueprintType)
-class AUDIOEXTENSIONS_API USoundfieldEndpointSettingsBase : public UObject
+UCLASS(config = Engine, abstract, editinlinenew, BlueprintType, MinimalAPI)
+class USoundfieldEndpointSettingsBase : public UObject
 {
 	GENERATED_BODY()
 
 public:
-	virtual TUniquePtr<ISoundfieldEndpointSettingsProxy> GetProxy() const PURE_VIRTUAL(USoundfieldEndpointSettingsBase::GetProxy, return nullptr;);
+	AUDIOEXTENSIONS_API virtual TUniquePtr<ISoundfieldEndpointSettingsProxy> GetProxy() const PURE_VIRTUAL(USoundfieldEndpointSettingsBase::GetProxy, return nullptr;);
 };
 
 /**
  * Class that allows soundfield-encoded audio to be sent to an arbitrary locale.
  * For endpoint types that support receiving our downmixed interleaved audio data directly, implement IAudioEndpoint instead.
  */
-class AUDIOEXTENSIONS_API ISoundfieldEndpoint
+class ISoundfieldEndpoint
 {
 public:
 	/**
 	 * @param [NumRenderCallbacksToBuffer] Maximum number of ISoundfieldPackets that should be buffered.
 	 */
-	ISoundfieldEndpoint(int32 NumRenderCallbacksToBuffer);
+	AUDIOEXTENSIONS_API ISoundfieldEndpoint(int32 NumRenderCallbacksToBuffer);
 
 	virtual ~ISoundfieldEndpoint() {};
 
 	/**
 	 * Push a soundfield packet to the buffer.
 	 */
-	bool PushAudio(TUniquePtr<ISoundfieldAudioPacket>&& InPacket);
+	AUDIOEXTENSIONS_API bool PushAudio(TUniquePtr<ISoundfieldAudioPacket>&& InPacket);
 
 	/**
 	 * Post new settings for this endpoint.
 	 * There is no type safety on this call, so make sure that you are using the correct implementation
 	 * of IAudioEndpointSettingsProxy for this implementation of IAudioEndpoint.
 	 */
-	void SetNewSettings(TUniquePtr<ISoundfieldEndpointSettingsProxy>&& InNewSettings);
+	AUDIOEXTENSIONS_API void SetNewSettings(TUniquePtr<ISoundfieldEndpointSettingsProxy>&& InNewSettings);
 
 	// Returns the amount of ISoundfieldAudioPackets currently buffered for this endpoint.
-	int32 GetNumPacketsBuffer();
+	AUDIOEXTENSIONS_API int32 GetNumPacketsBuffer();
 
 	// Returns the amount of ISoundfieldAudioPackets that can be buffered for this endpoint before reaching capacity.
-	int32 GetRemainderInPacketBuffer();
+	AUDIOEXTENSIONS_API int32 GetRemainderInPacketBuffer();
 
 	/**
 	 * If this endpoint hasn't created it's own async callback thread but still requires an explicit callback, this should be called.
 	 */
-	void ProcessAudioIfNecessary();
+	AUDIOEXTENSIONS_API void ProcessAudioIfNecessary();
 
 protected:
 
@@ -124,7 +124,7 @@ protected:
 	 * 
 	 * 
 	 */
-	TUniquePtr<ISoundfieldAudioPacket> PopAudio();
+	AUDIOEXTENSIONS_API TUniquePtr<ISoundfieldAudioPacket> PopAudio();
 
 	/**
 	 * Use this as a thread safe way to use the current settings posted to this IAudioEndpoint. Locks with IAudioEndpoint::SetSettings.
@@ -133,15 +133,15 @@ protected:
 	 *                                 to safely scope usage of the IAudioEndpointSettingsProxy pointer.
 	 *                                 Note that the resulting ISoundfieldEndpointSettingsProxy can be null.
 	 */
-	void PollSettings(TFunctionRef<void(const ISoundfieldEndpointSettingsProxy*)> SettingsCallback);
+	AUDIOEXTENSIONS_API void PollSettings(TFunctionRef<void(const ISoundfieldEndpointSettingsProxy*)> SettingsCallback);
 
 	/**
 	 * If EndpointRequiresCallback() returns true, this will be used to spawn an async thread and begin calling OnAudioCallback.
 	 */
-	void StartRunningCallback();
-	void StopRunningCallback();
+	AUDIOEXTENSIONS_API void StartRunningCallback();
+	AUDIOEXTENSIONS_API void StopRunningCallback();
 
-	void RunCallbackSynchronously();
+	AUDIOEXTENSIONS_API void RunCallbackSynchronously();
 
 private:
 	// Owns a scoped thread and runs OnAudioCallback when StartRunningCallback() is called.
@@ -164,7 +164,7 @@ private:
  * Also note that an implementation of ISoundfieldDecoder is not necessary for soundfield formats that are only used for
  * soundfield endpoints.
  */
-class AUDIOEXTENSIONS_API ISoundfieldEndpointFactory : public ISoundfieldFactory
+class ISoundfieldEndpointFactory : public ISoundfieldFactory
 {
 public:
 	/** Virtual destructor */
@@ -203,19 +203,19 @@ public:
 	/**
 	 * Get a registered endpoint factory by name.
 	 */
-	static ISoundfieldEndpointFactory* Get(const FName& InName);
+	static AUDIOEXTENSIONS_API ISoundfieldEndpointFactory* Get(const FName& InName);
 
-	static TArray<FName> GetAllSoundfieldEndpointTypes();
+	static AUDIOEXTENSIONS_API TArray<FName> GetAllSoundfieldEndpointTypes();
 
 	/**
 	 * This is the default name used when a user creates a soundfield endpoint submix.
 	 * Soundfied Endpoint submixes with this type will send their audio to the default output
 	 * with no encoding.
 	 */
-	static FName DefaultSoundfieldEndpointName();
+	static AUDIOEXTENSIONS_API FName DefaultSoundfieldEndpointName();
 
 	/** This function is not necessary to override, since audio sent to an endpoint does not need to be decoded to interleaved audio buffers. */
-	TUniquePtr<ISoundfieldDecoderStream> CreateDecoderStream(const FAudioPluginInitializationParams& InitInfo, const ISoundfieldEncodingSettingsProxy& InitialSettings) override;
+	AUDIOEXTENSIONS_API TUniquePtr<ISoundfieldDecoderStream> CreateDecoderStream(const FAudioPluginInitializationParams& InitInfo, const ISoundfieldEncodingSettingsProxy& InitialSettings) override;
 
 	/** REQUIRED OVERRIDES: */
 
@@ -249,6 +249,6 @@ public:
 
 	bool IsEndpointFormat() override { return true; }
 
-	virtual FName GetSoundfieldFormatName() override;
-	virtual bool CanTranscodeToSoundfieldFormat(FName DestinationFormat, const ISoundfieldEncodingSettingsProxy& DestinationEncodingSettings) override;
-};
\ No newline at end of file
+	AUDIOEXTENSIONS_API virtual FName GetSoundfieldFormatName() override;
+	AUDIOEXTENSIONS_API virtual bool CanTranscodeToSoundfieldFormat(FName DestinationFormat, const ISoundfieldEncodingSettingsProxy& DestinationEncodingSettings) override;
+};
diff --git a/Engine/Source/Runtime/AudioExtensions/Public/ISoundfieldFormat.h b/Engine/Source/Runtime/AudioExtensions/Public/ISoundfieldFormat.h
--- a/Engine/Source/Runtime/AudioExtensions/Public/ISoundfieldFormat.h
+++ b/Engine/Source/Runtime/AudioExtensions/Public/ISoundfieldFormat.h
@@ -100,7 +100,7 @@ ToType& DowncastSoundfieldRef(FromType& InRef)
  * your implementation of USoundfieldEncodingSettingsBase. We will then pass this proxy 
  * object to the soundfield stream classes.
  */
-class AUDIOEXTENSIONS_API ISoundfieldEncodingSettingsProxy
+class ISoundfieldEncodingSettingsProxy
 {
 public:
 	virtual ~ISoundfieldEncodingSettingsProxy() {};
@@ -119,13 +119,13 @@ public:
  * This opaque class should be used for specifying settings for how audio should be encoded
  * to your soundfield format for a given submix or file.
  */
-UCLASS(config = Engine, abstract, editinlinenew, BlueprintType)
-class AUDIOEXTENSIONS_API USoundfieldEncodingSettingsBase : public UObject
+UCLASS(config = Engine, abstract, editinlinenew, BlueprintType, MinimalAPI)
+class USoundfieldEncodingSettingsBase : public UObject
 {
 	GENERATED_BODY()
 
 public:
-	virtual TUniquePtr<ISoundfieldEncodingSettingsProxy> GetProxy() const PURE_VIRTUAL(USoundfieldEncodingSettingsBase::GetProxy, return nullptr;);
+	AUDIOEXTENSIONS_API virtual TUniquePtr<ISoundfieldEncodingSettingsProxy> GetProxy() const PURE_VIRTUAL(USoundfieldEncodingSettingsBase::GetProxy, return nullptr;);
 };
 
 struct FAudioPluginInitializationParams;
@@ -133,7 +133,7 @@ struct FAudioPluginInitializationParams;
 /**
  * This interface represents all encoded soundfield audio from a single render callback.
  */
-class AUDIOEXTENSIONS_API ISoundfieldAudioPacket
+class ISoundfieldAudioPacket
 {
 public:
 	virtual ~ISoundfieldAudioPacket() {};
@@ -214,7 +214,7 @@ struct FSoundfieldDecoderOutputData
 	Audio::FAlignedFloatBuffer& AudioBuffer;
 };
 
-class AUDIOEXTENSIONS_API ISoundfieldDecoderStream
+class ISoundfieldDecoderStream
 {
 public:
 	virtual ~ISoundfieldDecoderStream() {};
@@ -223,7 +223,7 @@ public:
 	virtual void DecodeAndMixIn(const FSoundfieldDecoderInputData& InputData, FSoundfieldDecoderOutputData& OutputData) = 0;
 };
 
-class AUDIOEXTENSIONS_API ISoundfieldTranscodeStream
+class ISoundfieldTranscodeStream
 {
 public:
 	virtual ~ISoundfieldTranscodeStream() {};
@@ -242,7 +242,7 @@ struct FSoundfieldMixerInputData
 	float SendLevel;
 };
 
-class AUDIOEXTENSIONS_API ISoundfieldMixerStream
+class ISoundfieldMixerStream
 {
 public:
 	virtual ~ISoundfieldMixerStream() {};
@@ -250,7 +250,7 @@ public:
 	virtual void MixTogether(const FSoundfieldMixerInputData& InputData, ISoundfieldAudioPacket& PacketToMixInto) = 0;
 };
 
-class AUDIOEXTENSIONS_API ISoundfieldFactory : public IModularFeature
+class ISoundfieldFactory : public IModularFeature
 {
 public:
 	/** Virtual destructor */
@@ -259,32 +259,32 @@ public:
 	}
 
 	/** When a submix has this format name, it is using interleaved, floating point audio with no metadata. */
-	static FName GetFormatNameForNoEncoding();
+	static AUDIOEXTENSIONS_API FName GetFormatNameForNoEncoding();
 
 	/** When a submix has this format name, it derives its format from the submix it sends audio to. */
-	static FName GetFormatNameForInheritedEncoding();
+	static AUDIOEXTENSIONS_API FName GetFormatNameForInheritedEncoding();
 
 	/** This is the FName used to register Soundfield Format factories with the modular feature system. */
-	static FName GetModularFeatureName();
+	static AUDIOEXTENSIONS_API FName GetModularFeatureName();
 
 	/** 
 	 * This needs to be called to make a soundfield format usable by the engine.
 	 * It can be called from a ISoundfieldFactory subclass' constructor
 	*/
-	static void RegisterSoundfieldFormat(ISoundfieldFactory* InFactory);
+	static AUDIOEXTENSIONS_API void RegisterSoundfieldFormat(ISoundfieldFactory* InFactory);
 
 	/**
 	 * This needs to be called it an implementation of ISoundfieldFactory is about to be destroyed.
 	 * It can be called from the destructor of an implementation of ISoundfieldFactory.
 	 */
-	static void UnregisterSoundfieldFormat(ISoundfieldFactory* InFactory);
+	static AUDIOEXTENSIONS_API void UnregisterSoundfieldFormat(ISoundfieldFactory* InFactory);
 
 	/**
 	 * Get a registered soundfield format factory by name.
 	 */
-	static ISoundfieldFactory* Get(const FName& InName);
+	static AUDIOEXTENSIONS_API ISoundfieldFactory* Get(const FName& InName);
 
-	static TArray<FName> GetAvailableSoundfieldFormats();
+	static AUDIOEXTENSIONS_API TArray<FName> GetAvailableSoundfieldFormats();
 
 	/** Get soundfield format name  */
 	virtual FName GetSoundfieldFormatName() = 0;
@@ -347,13 +347,13 @@ public:
 	virtual ~ISoundfieldEffectSettingsProxy() {};
 };
 
-UCLASS(config = Engine, abstract, editinlinenew, BlueprintType)
-class AUDIOEXTENSIONS_API USoundfieldEffectSettingsBase : public UObject
+UCLASS(config = Engine, abstract, editinlinenew, BlueprintType, MinimalAPI)
+class USoundfieldEffectSettingsBase : public UObject
 {
 	GENERATED_BODY()
 
 protected:
-	virtual TUniquePtr<ISoundfieldEffectSettingsProxy> GetNewProxy() const PURE_VIRTUAL(USoundfieldEffectSettingsBase::GetProxy, return nullptr;);
+	AUDIOEXTENSIONS_API virtual TUniquePtr<ISoundfieldEffectSettingsProxy> GetNewProxy() const PURE_VIRTUAL(USoundfieldEffectSettingsBase::GetProxy, return nullptr;);
 
 private:
 	// This is called by any engine system that is explicitly marked as a friend.
@@ -379,8 +379,8 @@ public:
  * This opaque class should be used for specifying settings for how audio should be encoded
  * to your soundfield format for a given submix or file.
  */
-UCLASS(config = Engine, abstract, editinlinenew, BlueprintType)
-class AUDIOEXTENSIONS_API USoundfieldEffectBase : public UObject
+UCLASS(config = Engine, abstract, editinlinenew, BlueprintType, MinimalAPI)
+class USoundfieldEffectBase : public UObject
 {
 	GENERATED_BODY()
 
@@ -396,22 +396,22 @@ protected:
 	/*
 	 * Get the implementation of USoundfieldProcessorSettingsBase that is used for this processor's settings. Will always be called on the CDO.
 	 */
-	virtual bool SupportsFormat(const FName& InFormat) const PURE_VIRTUAL(USoundfieldEncodingSettingsBase::SupportsFormat, return false;);
+	AUDIOEXTENSIONS_API virtual bool SupportsFormat(const FName& InFormat) const PURE_VIRTUAL(USoundfieldEncodingSettingsBase::SupportsFormat, return false;);
 
 	/*
 	 * Get the implementation of USoundfieldProcessorSettingsBase that is used for this processor's settings. Will always be called on the CDO.
 	 */
-	virtual const UClass* GetSettingsClass() const PURE_VIRTUAL(USoundfieldEncodingSettingsBase::GetSettingsClass, return nullptr;);
+	AUDIOEXTENSIONS_API virtual const UClass* GetSettingsClass() const PURE_VIRTUAL(USoundfieldEncodingSettingsBase::GetSettingsClass, return nullptr;);
 	
 	/**
 	 * return the default processor settings we should use when none is provided. Will always be called on the CDO.
 	 */
-	virtual const USoundfieldEffectSettingsBase* GetDefaultSettings() const PURE_VIRTUAL(USoundfieldEncodingSettingsBase::GetDefaultSettings, return nullptr;);
+	AUDIOEXTENSIONS_API virtual const USoundfieldEffectSettingsBase* GetDefaultSettings() const PURE_VIRTUAL(USoundfieldEncodingSettingsBase::GetDefaultSettings, return nullptr;);
 
 	/**
 	 * Spawn a new instance of this processor.
 	 */
-	virtual TUniquePtr<ISoundfieldEffectInstance> GetNewProcessor(const ISoundfieldEncodingSettingsProxy& EncodingSettings) const PURE_VIRTUAL(USoundfieldEncodingSettingsBase::GetProxy, return nullptr;);
+	AUDIOEXTENSIONS_API virtual TUniquePtr<ISoundfieldEffectInstance> GetNewProcessor(const ISoundfieldEncodingSettingsProxy& EncodingSettings) const PURE_VIRTUAL(USoundfieldEncodingSettingsBase::GetProxy, return nullptr;);
 
 private:
 
diff --git a/Engine/Source/Runtime/AudioExtensions/Public/IWaveformTransformation.h b/Engine/Source/Runtime/AudioExtensions/Public/IWaveformTransformation.h
--- a/Engine/Source/Runtime/AudioExtensions/Public/IWaveformTransformation.h
+++ b/Engine/Source/Runtime/AudioExtensions/Public/IWaveformTransformation.h
@@ -17,7 +17,7 @@
 namespace Audio
 {
 	// information about the current state of the wave file we are transforming
-	struct AUDIOEXTENSIONS_API FWaveformTransformationWaveInfo
+	struct FWaveformTransformationWaveInfo
 	{
 		float SampleRate = 0.f;
 		int32 NumChannels = 0;
@@ -32,7 +32,7 @@ namespace Audio
 	 *
 	 * note: WaveTransformation vs WaveformTransformation is to prevent UHT class name conflicts without having to namespace everything - remember this in derived classes!
 	 */
-	class AUDIOEXTENSIONS_API IWaveTransformation
+	class IWaveTransformation
 	{
 	public:
 
@@ -59,8 +59,8 @@ struct FWaveTransformUObjectConfiguration
 };
 
 // Base class to hold editor configurable properties for an arbitrary transformation of audio waveform data
-UCLASS(Abstract, EditInlineNew)
-class AUDIOEXTENSIONS_API UWaveformTransformationBase : public UObject
+UCLASS(Abstract, EditInlineNew, MinimalAPI)
+class UWaveformTransformationBase : public UObject
 {
 	GENERATED_BODY()
 
@@ -70,8 +70,8 @@ public:
 };
 
 // Object that holds an ordered list of transformations to perform on a sound wave
-UCLASS(EditInlineNew)
-class AUDIOEXTENSIONS_API UWaveformTransformationChain : public UObject
+UCLASS(EditInlineNew, MinimalAPI)
+class UWaveformTransformationChain : public UObject
 {
 	GENERATED_BODY()
 	
@@ -79,5 +79,5 @@ public:
 	UPROPERTY(EditAnywhere, Instanced, Category = "Transformations")
 	TArray<TObjectPtr<UWaveformTransformationBase>> Transformations;
 
-	TArray<Audio::FTransformationPtr> CreateTransformations() const;
+	AUDIOEXTENSIONS_API TArray<Audio::FTransformationPtr> CreateTransformations() const;
 };
diff --git a/Engine/Source/Runtime/AudioLink/AudioLinkCore/Public/AudioLinkSettingsAbstract.h b/Engine/Source/Runtime/AudioLink/AudioLinkCore/Public/AudioLinkSettingsAbstract.h
--- a/Engine/Source/Runtime/AudioLink/AudioLinkCore/Public/AudioLinkSettingsAbstract.h
+++ b/Engine/Source/Runtime/AudioLink/AudioLinkCore/Public/AudioLinkSettingsAbstract.h
@@ -12,7 +12,7 @@ class UAudioLinkSettingsAbstract;
   * This interface should be used to provide a non-uclass version of the data described in
   * your implementation of UAudioLinkSettingsAbstract
   */
-class AUDIOLINKCORE_API IAudioLinkSettingsProxy
+class IAudioLinkSettingsProxy
 {
 public:
 	virtual ~IAudioLinkSettingsProxy() = default;
@@ -27,15 +27,15 @@ protected:
   * This opaque class should be used for specifying settings for how audio should be
   * send to an external audio link.
   */
-UCLASS(config = Engine, abstract, editinlinenew, BlueprintType)
-class AUDIOLINKCORE_API UAudioLinkSettingsAbstract : public UObject
+UCLASS(config = Engine, abstract, editinlinenew, BlueprintType, MinimalAPI)
+class UAudioLinkSettingsAbstract : public UObject
 {
 	GENERATED_BODY()
 
 public:
 	using FSharedSettingsProxyPtr = TSharedPtr<IAudioLinkSettingsProxy, ESPMode::ThreadSafe>;
 
-	virtual FName GetFactoryName() const PURE_VIRTUAL(UAudioLinkSettingsAbstract::GetFactoryName, return {};);
+	AUDIOLINKCORE_API virtual FName GetFactoryName() const PURE_VIRTUAL(UAudioLinkSettingsAbstract::GetFactoryName, return {};);
 	
 	virtual const FSharedSettingsProxyPtr& GetProxy() const
 	{
@@ -60,7 +60,7 @@ protected:
 	}
 #endif //WITH_EDITOR
 
-	virtual FSharedSettingsProxyPtr MakeProxy() const PURE_VIRTUAL(UAudioLinkSettingsAbstract::MakeProxy, return {};);
+	AUDIOLINKCORE_API virtual FSharedSettingsProxyPtr MakeProxy() const PURE_VIRTUAL(UAudioLinkSettingsAbstract::MakeProxy, return {};);
 	mutable FSharedSettingsProxyPtr ProxyInstance;
 };
 
diff --git a/Engine/Source/Runtime/AudioLink/AudioLinkEngine/Private/BufferedListenerBase.h b/Engine/Source/Runtime/AudioLink/AudioLinkEngine/Private/BufferedListenerBase.h
--- a/Engine/Source/Runtime/AudioLink/AudioLinkEngine/Private/BufferedListenerBase.h
+++ b/Engine/Source/Runtime/AudioLink/AudioLinkEngine/Private/BufferedListenerBase.h
@@ -112,42 +112,42 @@ private:
 
 /** Common base class of Buffered Listener objects.
 */
-class AUDIOLINKENGINE_API FBufferedListenerBase : public IBufferedAudioOutput
+class FBufferedListenerBase : public IBufferedAudioOutput
 {
 protected:
-	FBufferedListenerBase(int32 InDefaultCircularBufferSize);
+	AUDIOLINKENGINE_API FBufferedListenerBase(int32 InDefaultCircularBufferSize);
 	virtual ~FBufferedListenerBase() = default;
 
 	//~ Begin IBufferedAudioOutput
-	bool PopBuffer(float* InBuffer, int32 InBufferSizeInSamples, int32& OutSamplesWritten) override;
-	bool GetFormat(IBufferedAudioOutput::FBufferFormat& OutFormat) const override;
-	void Reserve(int32 InNumSamplesToReserve, int32 InNumSamplesOfSilence) override;
-	void SetFormatKnownDelegate(FOnFormatKnown InFormatKnownDelegate) override;
+	AUDIOLINKENGINE_API bool PopBuffer(float* InBuffer, int32 InBufferSizeInSamples, int32& OutSamplesWritten) override;
+	AUDIOLINKENGINE_API bool GetFormat(IBufferedAudioOutput::FBufferFormat& OutFormat) const override;
+	AUDIOLINKENGINE_API void Reserve(int32 InNumSamplesToReserve, int32 InNumSamplesOfSilence) override;
+	AUDIOLINKENGINE_API void SetFormatKnownDelegate(FOnFormatKnown InFormatKnownDelegate) override;
 	void SetBufferStreamEndDelegate(FOnBufferStreamEnd) override {}
 	//~ End IBufferedAudioOutput
 
 	//* Common path to receive a new buffer, call from derived classes */
-	void OnBufferReceived(const FBufferFormat& InFormat, TArrayView<const float> InBuffer);
+	AUDIOLINKENGINE_API void OnBufferReceived(const FBufferFormat& InFormat, TArrayView<const float> InBuffer);
 
 	//* Reset the format of the buffer */
-	void ResetFormat();
+	AUDIOLINKENGINE_API void ResetFormat();
 
 	//* Set the format of the buffer */
-	void SetFormat(const FBufferFormat& InFormat);
+	AUDIOLINKENGINE_API void SetFormat(const FBufferFormat& InFormat);
 	
 	//* Ask if the started flag has been set. Note this is non-atomic, as it could change during the call */	
-	bool IsStartedNonAtomic() const;
+	AUDIOLINKENGINE_API bool IsStartedNonAtomic() const;
 		
 	//* Attempt to set our state to started. Not this can fail if we're already started. */	
-	bool TrySetStartedFlag(); 
+	AUDIOLINKENGINE_API bool TrySetStartedFlag(); 
 	
 	//* Attempt to set started to false. This can fail if we're already stopped. */	
-	bool TryUnsetStartedFlag();
+	AUDIOLINKENGINE_API bool TryUnsetStartedFlag();
 
-	bool TrySetStoppingFlag();
+	AUDIOLINKENGINE_API bool TrySetStoppingFlag();
 
 private:	
-	void PushSilence(int32 InNumSamplesOfSilence);
+	AUDIOLINKENGINE_API void PushSilence(int32 InNumSamplesOfSilence);
 
 	/** Buffer to hold the data for the single source we're listening to, interleaved. */
 	FLockingCircularSampleBuffer CircularBuffer;
diff --git a/Engine/Source/Runtime/AudioLink/AudioLinkEngine/Private/BufferedSourceListener.h b/Engine/Source/Runtime/AudioLink/AudioLinkEngine/Private/BufferedSourceListener.h
--- a/Engine/Source/Runtime/AudioLink/AudioLinkEngine/Private/BufferedSourceListener.h
+++ b/Engine/Source/Runtime/AudioLink/AudioLinkEngine/Private/BufferedSourceListener.h
@@ -9,11 +9,11 @@
 	FBufferedSourceListener: Concrete implementation of both a buffer listener and buffered output.
 	Contains a circular buffer.
 */
-class AUDIOLINKENGINE_API FBufferedSourceListener : public ISourceBufferListener, public FBufferedListenerBase
+class FBufferedSourceListener : public ISourceBufferListener, public FBufferedListenerBase
 {
 public:
-	FBufferedSourceListener(int32 InDefaultCircularBufferSize);
-	virtual ~FBufferedSourceListener();
+	AUDIOLINKENGINE_API FBufferedSourceListener(int32 InDefaultCircularBufferSize);
+	AUDIOLINKENGINE_API virtual ~FBufferedSourceListener();
 
 	/**
 	 * @param  InDevice Audio Device to register this listener with.
@@ -21,19 +21,19 @@ public:
 	 *
 	 * @return  success true, false otherwise.
 	 */
-	bool Start(FAudioDevice* InDevice) override;
+	AUDIOLINKENGINE_API bool Start(FAudioDevice* InDevice) override;
 
 	/** Stop the buffer by unregistering the buffer listener */
-	void Stop(FAudioDevice* InDevice) override;
+	AUDIOLINKENGINE_API void Stop(FAudioDevice* InDevice) override;
 	
 private:
 	//~ Begin ISourceBufferListener
-	void OnNewBuffer(const ISourceBufferListener::FOnNewBufferParams& InParams) override;
-	void OnSourceReleased(const int32 SourceId) override;
+	AUDIOLINKENGINE_API void OnNewBuffer(const ISourceBufferListener::FOnNewBufferParams& InParams) override;
+	AUDIOLINKENGINE_API void OnSourceReleased(const int32 SourceId) override;
 	//~ End ISourceBufferListener
 
 	//~ Begin IBufferedAudioOutput
-	void SetBufferStreamEndDelegate(FOnBufferStreamEnd InBufferStreamEndDelegate) override;
+	AUDIOLINKENGINE_API void SetBufferStreamEndDelegate(FOnBufferStreamEnd InBufferStreamEndDelegate) override;
 	//~ End IBufferedAudioOutput
 
 	std::atomic<int32> CurrentSourceId	= INDEX_NONE;		// r/w AudioMixer thread.	
diff --git a/Engine/Source/Runtime/AudioLink/AudioLinkEngine/Private/BufferedSubmixListener.h b/Engine/Source/Runtime/AudioLink/AudioLinkEngine/Private/BufferedSubmixListener.h
--- a/Engine/Source/Runtime/AudioLink/AudioLinkEngine/Private/BufferedSubmixListener.h
+++ b/Engine/Source/Runtime/AudioLink/AudioLinkEngine/Private/BufferedSubmixListener.h
@@ -7,7 +7,7 @@
 
 /** Concrete Submix Buffer Listener  
 */
-class AUDIOLINKENGINE_API FBufferedSubmixListener : public ISubmixBufferListener, 
+class FBufferedSubmixListener : public ISubmixBufferListener, 
 													public FBufferedListenerBase
 {
 public:
@@ -15,9 +15,9 @@ public:
 	 * @param  InDefaultCircularBufferSize: Size of the circular buffer in samples by default.
 	 * @param  bInZeroInputBuffer: The passed in buffer will be set zerod after we've buffered it. Not 
 	 */
-	FBufferedSubmixListener(int32 InDefaultCircularBufferSize, bool bInZeroInputBuffer);
+	AUDIOLINKENGINE_API FBufferedSubmixListener(int32 InDefaultCircularBufferSize, bool bInZeroInputBuffer);
 
-	virtual ~FBufferedSubmixListener();
+	AUDIOLINKENGINE_API virtual ~FBufferedSubmixListener();
 
 	/**
 	 * Starts the Submix buffer listener by registering it with the passed in Audio Device.
@@ -26,17 +26,17 @@ public:
 	 *
 	 * @return  success true, false otherwise.
 	 */
-	bool Start(FAudioDevice* InDevice) override;
-	void Stop(FAudioDevice* InDevice) override;
+	AUDIOLINKENGINE_API bool Start(FAudioDevice* InDevice) override;
+	AUDIOLINKENGINE_API void Stop(FAudioDevice* InDevice) override;
 
 private:
 	void RegisterWithAudioDevice(FAudioDevice* InDevice);
 	void UnregsiterWithAudioDevice(FAudioDevice* InDevice);
 	
 	//~ Begin ISubmixBufferListener
-	void OnNewSubmixBuffer(const USoundSubmix* OwningSubmix, float* AudioData, int32 InNumSamples, int32 InNumChannels, const int32 InSampleRate, double) override;
+	AUDIOLINKENGINE_API void OnNewSubmixBuffer(const USoundSubmix* OwningSubmix, float* AudioData, int32 InNumSamples, int32 InNumChannels, const int32 InSampleRate, double) override;
 	//~ End ISubmixBufferListener
 
 	Audio::FDeviceId DeviceId;
 	bool bZeroInputBuffer = false;
-};
\ No newline at end of file
+};
diff --git a/Engine/Source/Runtime/AudioLink/AudioLinkEngine/Public/IAudioLink.h b/Engine/Source/Runtime/AudioLink/AudioLinkEngine/Public/IAudioLink.h
--- a/Engine/Source/Runtime/AudioLink/AudioLinkEngine/Public/IAudioLink.h
+++ b/Engine/Source/Runtime/AudioLink/AudioLinkEngine/Public/IAudioLink.h
@@ -10,10 +10,10 @@
   *	Abstract interface for AudioLink instances. 
   *	Purely opaque.
   */
-class AUDIOLINKENGINE_API IAudioLink
+class IAudioLink
 {
 protected:
-	IAudioLink();
+	AUDIOLINKENGINE_API IAudioLink();
 public:
 	virtual ~IAudioLink() = default;
 };
@@ -21,7 +21,7 @@ public:
 /** IAudioLinkSourcePushed
   *	Where the owning object needs to push it's state
   */
-class AUDIOLINKENGINE_API IAudioLinkSourcePushed : public IAudioLink
+class IAudioLinkSourcePushed : public IAudioLink
 {
 protected:
 	IAudioLinkSourcePushed() = default;
@@ -95,4 +95,4 @@ public:
 	DECLARE_MULTICAST_DELEGATE_OneParam(FOnEndRender, const FOnRenderParams&);
 	virtual FDelegateHandle RegisterEndRenderDelegate(const FOnEndRender::FDelegate&) = 0;
 	virtual bool RemoveEndRenderDelegate(const FDelegateHandle&) = 0;
-};
\ No newline at end of file
+};
diff --git a/Engine/Source/Runtime/AudioLink/AudioLinkEngine/Public/IAudioLinkBlueprintInterface.h b/Engine/Source/Runtime/AudioLink/AudioLinkEngine/Public/IAudioLinkBlueprintInterface.h
--- a/Engine/Source/Runtime/AudioLink/AudioLinkEngine/Public/IAudioLinkBlueprintInterface.h
+++ b/Engine/Source/Runtime/AudioLink/AudioLinkEngine/Public/IAudioLinkBlueprintInterface.h
@@ -12,13 +12,13 @@ class UObject;
 class USoundBase;
 struct FFrame;
 
-UINTERFACE(BlueprintType, meta = (CannotImplementInterfaceInBlueprint))
-class AUDIOLINKENGINE_API UAudioLinkBlueprintInterface : public UInterface
+UINTERFACE(BlueprintType, meta = (CannotImplementInterfaceInBlueprint), MinimalAPI)
+class UAudioLinkBlueprintInterface : public UInterface
 {
 	GENERATED_UINTERFACE_BODY()
 };
 
-class AUDIOLINKENGINE_API IAudioLinkBlueprintInterface : public IInterface
+class IAudioLinkBlueprintInterface : public IInterface
 {
 	GENERATED_IINTERFACE_BODY()
 
diff --git a/Engine/Source/Runtime/AudioLink/AudioLinkEngine/Public/IAudioLinkFactory.h b/Engine/Source/Runtime/AudioLink/AudioLinkEngine/Public/IAudioLinkFactory.h
--- a/Engine/Source/Runtime/AudioLink/AudioLinkEngine/Public/IAudioLinkFactory.h
+++ b/Engine/Source/Runtime/AudioLink/AudioLinkEngine/Public/IAudioLinkFactory.h
@@ -17,14 +17,14 @@ class UAudioLinkComponent;
 	IAudioLinkFactory
 	Abstract interface for AudioLink factory objects. 
 */
-class AUDIOLINKENGINE_API IAudioLinkFactory : public IModularFeature
+class IAudioLinkFactory : public IModularFeature
 {
 protected:
 	/** Constructor will auto register this instance. */
-	IAudioLinkFactory();
+	AUDIOLINKENGINE_API IAudioLinkFactory();
 public:
 	/** Destructor will auto unregister this instance. */
-	virtual ~IAudioLinkFactory();
+	AUDIOLINKENGINE_API virtual ~IAudioLinkFactory();
 
 	/** 
 	 * Gets the name of this factory type. This will uniquely identity itself, so it can be found by the FindFactory call below
@@ -46,21 +46,21 @@ public:
 		TWeakObjectPtr<UAudioComponent> AudioComponent;
 		TWeakObjectPtr<USceneComponent> OwningComponent;
 	};
-	virtual FSharedBufferedOutputPtr CreateSourceBufferListener(const FSourceBufferListenerCreateParams&);
+	AUDIOLINKENGINE_API virtual FSharedBufferedOutputPtr CreateSourceBufferListener(const FSourceBufferListenerCreateParams&);
 
 	struct FPushedBufferListenerCreateParams
 	{
 		int32 SizeOfBufferInFrames = INDEX_NONE;
 		bool bShouldZeroBuffer = false;
 	};
-	virtual FSharedBufferedOutputPtr CreatePushableBufferListener(const FPushedBufferListenerCreateParams&);
+	AUDIOLINKENGINE_API virtual FSharedBufferedOutputPtr CreatePushableBufferListener(const FPushedBufferListenerCreateParams&);
 
 	struct FSubmixBufferListenerCreateParams
 	{
 		int32 SizeOfBufferInFrames = 0;
 		bool bShouldZeroBuffer = false;
 	};
-	virtual FSharedBufferedOutputPtr CreateSubmixBufferListener(const FSubmixBufferListenerCreateParams&);
+	AUDIOLINKENGINE_API virtual FSharedBufferedOutputPtr CreateSubmixBufferListener(const FSubmixBufferListenerCreateParams&);
 			
 	/**
 	 * Parameters use when creating a Submix Audio Link 
@@ -130,19 +130,19 @@ public:
 	 * Gets all registered factory instances.
 	 * @return Array of all factories.
 	 */
-	static TArray<IAudioLinkFactory*> GetAllRegisteredFactories();
+	static AUDIOLINKENGINE_API TArray<IAudioLinkFactory*> GetAllRegisteredFactories();
 
 	/**
 	 * Gets all registered factory names
 	 * @return Array of all factory names
 	 */
-	static TArray<FName> GetAllRegisteredFactoryNames();
+	static AUDIOLINKENGINE_API TArray<FName> GetAllRegisteredFactoryNames();
 	
 	/**
 	 * Gets all registered factory names
 	 * @return Array of all factory names
 	 */
-	static IAudioLinkFactory* FindFactory(const FName InFactoryName);
+	static AUDIOLINKENGINE_API IAudioLinkFactory* FindFactory(const FName InFactoryName);
 
 protected:	
 };
diff --git a/Engine/Source/Runtime/AudioMixer/Classes/Generators/AudioGenerator.h b/Engine/Source/Runtime/AudioMixer/Classes/Generators/AudioGenerator.h
--- a/Engine/Source/Runtime/AudioMixer/Classes/Generators/AudioGenerator.h
+++ b/Engine/Source/Runtime/AudioMixer/Classes/Generators/AudioGenerator.h
@@ -19,14 +19,14 @@ struct FAudioGeneratorHandle
 
 typedef TFunction<void(const float * InAudio, int32 NumSamples)> FOnAudioGenerate;
 
-UCLASS(ClassGroup = (Custom), meta = (BlueprintSpawnableComponent))
-class AUDIOMIXER_API UAudioGenerator : public UObject
+UCLASS(ClassGroup = (Custom), meta = (BlueprintSpawnableComponent), MinimalAPI)
+class UAudioGenerator : public UObject
 {
 	GENERATED_BODY()
 
 public:
-	UAudioGenerator();
-	virtual ~UAudioGenerator();
+	AUDIOMIXER_API UAudioGenerator();
+	AUDIOMIXER_API virtual ~UAudioGenerator();
 
 	// Returns the sample rate of the generator
 	int32 GetSampleRate() const { return SampleRate; }
@@ -35,21 +35,21 @@ public:
 	int32 GetNumChannels() const { return NumChannels; }
 
 	// Adds a generator delegate. Returns a handle for the generator delegate, so it can be removed.
-	FAudioGeneratorHandle AddGeneratorDelegate(FOnAudioGenerate InFunction);
+	AUDIOMIXER_API FAudioGeneratorHandle AddGeneratorDelegate(FOnAudioGenerate InFunction);
 
 	// Removes the given audio generator delegate handle
-	void RemoveGeneratorDelegate(FAudioGeneratorHandle InHandle);
+	AUDIOMIXER_API void RemoveGeneratorDelegate(FAudioGeneratorHandle InHandle);
 
 protected:
 
 	// Called by derived classes to initialize the sample rate and num channels of the generator
-	void Init(int32 InSampleRate, int32 InNumChannels);
+	AUDIOMIXER_API void Init(int32 InSampleRate, int32 InNumChannels);
 
 	// Called by derived classes when new audio is generated
-	void OnGeneratedAudio(const float* InAudio, int32 NumSamples);
+	AUDIOMIXER_API void OnGeneratedAudio(const float* InAudio, int32 NumSamples);
 
 	FCriticalSection CritSect;
 	int32 SampleRate;
 	int32 NumChannels;
 	TMap<uint32, FOnAudioGenerate> OnGeneratedMap;
-};
\ No newline at end of file
+};
diff --git a/Engine/Source/Runtime/AudioMixer/Classes/SubmixEffects/AudioMixerSubmixEffectDynamicsProcessor.h b/Engine/Source/Runtime/AudioMixer/Classes/SubmixEffects/AudioMixerSubmixEffectDynamicsProcessor.h
--- a/Engine/Source/Runtime/AudioMixer/Classes/SubmixEffects/AudioMixerSubmixEffectDynamicsProcessor.h
+++ b/Engine/Source/Runtime/AudioMixer/Classes/SubmixEffects/AudioMixerSubmixEffectDynamicsProcessor.h
@@ -137,7 +137,7 @@ public:
 };
 
 USTRUCT(BlueprintType)
-struct AUDIOMIXER_API FSubmixEffectDynamicProcessorFilterSettings
+struct FSubmixEffectDynamicProcessorFilterSettings
 {
 	GENERATED_USTRUCT_BODY()
 
@@ -163,7 +163,7 @@ struct AUDIOMIXER_API FSubmixEffectDynamicProcessorFilterSettings
 
 // Submix dynamics processor settings
 USTRUCT(BlueprintType)
-struct AUDIOMIXER_API FSubmixEffectDynamicsProcessorSettings
+struct FSubmixEffectDynamicsProcessorSettings
 {
 	GENERATED_USTRUCT_BODY()
 
@@ -266,35 +266,35 @@ struct AUDIOMIXER_API FSubmixEffectDynamicsProcessorSettings
 };
 
 
-class AUDIOMIXER_API FSubmixEffectDynamicsProcessor : public FSoundEffectSubmix
+class FSubmixEffectDynamicsProcessor : public FSoundEffectSubmix
 {
 public:
-	FSubmixEffectDynamicsProcessor();
+	AUDIOMIXER_API FSubmixEffectDynamicsProcessor();
 
-	virtual ~FSubmixEffectDynamicsProcessor();
+	AUDIOMIXER_API virtual ~FSubmixEffectDynamicsProcessor();
 
 	// Gets the effect's deviceId that owns it
-	Audio::FDeviceId GetDeviceId() const;
+	AUDIOMIXER_API Audio::FDeviceId GetDeviceId() const;
 
 	// Called on an audio effect at initialization on audio thread before audio processing begins.
-	virtual void Init(const FSoundEffectSubmixInitData& InInitData) override;
+	AUDIOMIXER_API virtual void Init(const FSoundEffectSubmixInitData& InInitData) override;
 
 	// Process the input block of audio. Called on audio render thread.
-	virtual void OnProcessAudio(const FSoundEffectSubmixInputData& InData, FSoundEffectSubmixOutputData& OutData) override;
+	AUDIOMIXER_API virtual void OnProcessAudio(const FSoundEffectSubmixInputData& InData, FSoundEffectSubmixOutputData& OutData) override;
 
 	// Called when an audio effect preset is changed
-	virtual void OnPresetChanged() override;
+	AUDIOMIXER_API virtual void OnPresetChanged() override;
 
 
 protected:
-	Audio::FMixerDevice* GetMixerDevice();
+	AUDIOMIXER_API Audio::FMixerDevice* GetMixerDevice();
 
-	void ResetKey();
-	void UpdateKeyFromSettings(const FSubmixEffectDynamicsProcessorSettings& InSettings);
-	bool UpdateKeySourcePatch();
+	AUDIOMIXER_API void ResetKey();
+	AUDIOMIXER_API void UpdateKeyFromSettings(const FSubmixEffectDynamicsProcessorSettings& InSettings);
+	AUDIOMIXER_API bool UpdateKeySourcePatch();
 
-	void OnDeviceCreated(Audio::FDeviceId InDeviceId);
-	void OnDeviceDestroyed(Audio::FDeviceId InDeviceId);
+	AUDIOMIXER_API void OnDeviceCreated(Audio::FDeviceId InDeviceId);
+	AUDIOMIXER_API void OnDeviceDestroyed(Audio::FDeviceId InDeviceId);
 	
 	Audio::FAlignedFloatBuffer AudioExternal;
 
@@ -315,41 +315,41 @@ private:
 	friend class USubmixEffectDynamicsProcessorPreset;
 };
 
-UCLASS(ClassGroup = AudioSourceEffect, meta = (BlueprintSpawnableComponent))
-class AUDIOMIXER_API USubmixEffectDynamicsProcessorPreset : public USoundEffectSubmixPreset
+UCLASS(ClassGroup = AudioSourceEffect, meta = (BlueprintSpawnableComponent), MinimalAPI)
+class USubmixEffectDynamicsProcessorPreset : public USoundEffectSubmixPreset
 {
 	GENERATED_BODY()
 
 public:
 	EFFECT_PRESET_METHODS(SubmixEffectDynamicsProcessor)
 
-	virtual void OnInit() override;
+	AUDIOMIXER_API virtual void OnInit() override;
 
-	virtual void Serialize(FStructuredArchive::FRecord Record) override;
+	AUDIOMIXER_API virtual void Serialize(FStructuredArchive::FRecord Record) override;
 
 #if WITH_EDITOR
-	virtual void PostEditChangeChainProperty(struct FPropertyChangedChainEvent& InChainEvent) override;
+	AUDIOMIXER_API virtual void PostEditChangeChainProperty(struct FPropertyChangedChainEvent& InChainEvent) override;
 #endif // WITH_EDITOR
 
 	UFUNCTION(BlueprintCallable, Category = "Audio|Effects")
-	void ResetKey();
+	AUDIOMIXER_API void ResetKey();
 
 	// Sets the source key input as the provided AudioBus' output.  If no object is provided, key is set
 	// to effect's input.
 	UFUNCTION(BlueprintCallable, Category = "Audio|Effects")
-	void SetAudioBus(UAudioBus* AudioBus);
+	AUDIOMIXER_API void SetAudioBus(UAudioBus* AudioBus);
 
 	// Sets the source key input as the provided Submix's output.  If no object is provided, key is set
 	// to effect's input.
 	UFUNCTION(BlueprintCallable, Category = "Audio|Effects")
-	void SetExternalSubmix(USoundSubmix* Submix);
+	AUDIOMIXER_API void SetExternalSubmix(USoundSubmix* Submix);
 
 	UFUNCTION(BlueprintCallable, Category = "Audio|Effects")
-	void SetSettings(const FSubmixEffectDynamicsProcessorSettings& Settings);
+	AUDIOMIXER_API void SetSettings(const FSubmixEffectDynamicsProcessorSettings& Settings);
 
 	UPROPERTY(EditAnywhere, BlueprintReadWrite, Category = SubmixEffectPreset, meta = (ShowOnlyInnerProperties))
 	FSubmixEffectDynamicsProcessorSettings Settings;
 
 private:
 	void SetKey(ESubmixEffectDynamicsKeySource InKeySource, UObject* InObject, int32 InNumChannels = 0);
-};
\ No newline at end of file
+};
diff --git a/Engine/Source/Runtime/AudioMixer/Classes/SubmixEffects/AudioMixerSubmixEffectEQ.h b/Engine/Source/Runtime/AudioMixer/Classes/SubmixEffects/AudioMixerSubmixEffectEQ.h
--- a/Engine/Source/Runtime/AudioMixer/Classes/SubmixEffects/AudioMixerSubmixEffectEQ.h
+++ b/Engine/Source/Runtime/AudioMixer/Classes/SubmixEffects/AudioMixerSubmixEffectEQ.h
@@ -14,7 +14,7 @@ DECLARE_CYCLE_STAT_EXTERN(TEXT("Submix EQ"), STAT_AudioMixerSubmixEQ, STATGROUP_
 
 // A multiband EQ submix effect.
 USTRUCT(BlueprintType)
-struct AUDIOMIXER_API FSubmixEffectEQBand
+struct FSubmixEffectEQBand
 {
 	GENERATED_USTRUCT_BODY()
 
@@ -50,19 +50,19 @@ struct FSubmixEffectSubmixEQSettings
 	TArray<FSubmixEffectEQBand> EQBands;
 };
 
-class AUDIOMIXER_API FSubmixEffectSubmixEQ : public FSoundEffectSubmix
+class FSubmixEffectSubmixEQ : public FSoundEffectSubmix
 {
 public:
-	FSubmixEffectSubmixEQ();
+	AUDIOMIXER_API FSubmixEffectSubmixEQ();
 
 	// Called on an audio effect at initialization on main thread before audio processing begins.
-	virtual void Init(const FSoundEffectSubmixInitData& InSampleRate) override;
+	AUDIOMIXER_API virtual void Init(const FSoundEffectSubmixInitData& InSampleRate) override;
 
 	// Process the input block of audio. Called on audio thread.
-	virtual void OnProcessAudio(const FSoundEffectSubmixInputData& InData, FSoundEffectSubmixOutputData& OutData) override;
+	AUDIOMIXER_API virtual void OnProcessAudio(const FSoundEffectSubmixInputData& InData, FSoundEffectSubmixOutputData& OutData) override;
 
 	// Sets the effect parameters using the old audio engine preset setting object
-	virtual bool SetParameters(const FAudioEffectParameters& InParameters) override;
+	AUDIOMIXER_API virtual bool SetParameters(const FAudioEffectParameters& InParameters) override;
 
 	virtual bool SupportsDefaultEQ() const override
 	{
@@ -70,10 +70,10 @@ public:
 	}
 
 	// Called when an audio effect preset is changed
-	virtual void OnPresetChanged() override;
+	AUDIOMIXER_API virtual void OnPresetChanged() override;
 
 protected:
-	void UpdateParameters(const int32 NumOutputChannels);
+	AUDIOMIXER_API void UpdateParameters(const int32 NumOutputChannels);
 
 	// An EQ effect is a bank of biquad filters
 	struct FEQ
@@ -105,8 +105,8 @@ protected:
 	FSubmixEffectSubmixEQSettings RenderThreadEQSettings;
 };
 
-UCLASS(ClassGroup = AudioSourceEffect, meta = (BlueprintSpawnableComponent))
-class AUDIOMIXER_API USubmixEffectSubmixEQPreset : public USoundEffectSubmixPreset
+UCLASS(ClassGroup = AudioSourceEffect, meta = (BlueprintSpawnableComponent), MinimalAPI)
+class USubmixEffectSubmixEQPreset : public USoundEffectSubmixPreset
 {
 	GENERATED_BODY()
 
@@ -115,7 +115,7 @@ public:
 	EFFECT_PRESET_METHODS(SubmixEffectSubmixEQ)
 
 	UFUNCTION(BlueprintCallable, Category = "Audio|Effects")
-	void SetSettings(const FSubmixEffectSubmixEQSettings& InSettings);
+	AUDIOMIXER_API void SetSettings(const FSubmixEffectSubmixEQSettings& InSettings);
 
 	UPROPERTY(EditAnywhere, BlueprintReadWrite, Category = SubmixEffectPreset)
 	FSubmixEffectSubmixEQSettings Settings;
diff --git a/Engine/Source/Runtime/AudioMixer/Classes/SubmixEffects/AudioMixerSubmixEffectReverb.h b/Engine/Source/Runtime/AudioMixer/Classes/SubmixEffects/AudioMixerSubmixEffectReverb.h
--- a/Engine/Source/Runtime/AudioMixer/Classes/SubmixEffects/AudioMixerSubmixEffectReverb.h
+++ b/Engine/Source/Runtime/AudioMixer/Classes/SubmixEffects/AudioMixerSubmixEffectReverb.h
@@ -20,7 +20,7 @@ struct FAudioEffectParameters;
 DECLARE_CYCLE_STAT_EXTERN(TEXT("Submix Reverb"), STAT_AudioMixerSubmixReverb, STATGROUP_AudioMixer, AUDIOMIXER_API);
 
 USTRUCT(BlueprintType)
-struct AUDIOMIXER_API FSubmixEffectReverbSettings
+struct FSubmixEffectReverbSettings
 {
 	GENERATED_USTRUCT_BODY()
 
@@ -111,25 +111,25 @@ struct AUDIOMIXER_API FSubmixEffectReverbSettings
 	}
 };
 
-class AUDIOMIXER_API FSubmixEffectReverb : public FSoundEffectSubmix
+class FSubmixEffectReverb : public FSoundEffectSubmix
 {
 public:
-	FSubmixEffectReverb();
+	AUDIOMIXER_API FSubmixEffectReverb();
 
 	// Called on an audio effect at initialization on main thread before audio processing begins.
-	virtual void Init(const FSoundEffectSubmixInitData& InSampleRate) override;
+	AUDIOMIXER_API virtual void Init(const FSoundEffectSubmixInitData& InSampleRate) override;
 	
 	// Called when an audio effect preset is changed
-	virtual void OnPresetChanged() override;
+	AUDIOMIXER_API virtual void OnPresetChanged() override;
 
 	// Forces receiving downmixed submix audio to stereo input for the reverb effect
 	virtual uint32 GetDesiredInputChannelCountOverride() const override { return 2; }
 
 	// Process the input block of audio. Called on audio thread.
-	virtual void OnProcessAudio(const FSoundEffectSubmixInputData& InData, FSoundEffectSubmixOutputData& OutData) override;
+	AUDIOMIXER_API virtual void OnProcessAudio(const FSoundEffectSubmixInputData& InData, FSoundEffectSubmixOutputData& OutData) override;
 
 	// Sets the reverb effect parameters based from audio thread code
-	virtual bool SetParameters(const FAudioEffectParameters& InParameters) override;
+	AUDIOMIXER_API virtual bool SetParameters(const FAudioEffectParameters& InParameters) override;
 
 	// Whether this effect supports the default reverb system
 	virtual bool SupportsDefaultReverb() const override
@@ -142,10 +142,10 @@ public:
 
 private:
 
-	static const float MinWetness;
-	static const float MaxWetness;
+	static AUDIOMIXER_API const float MinWetness;
+	static AUDIOMIXER_API const float MaxWetness;
 
-	void UpdateParameters();
+	AUDIOMIXER_API void UpdateParameters();
 
 	// The reverb effect
 	TUniquePtr<Audio::FPlateReverbFast> PlateReverb;
@@ -165,8 +165,8 @@ private:
 	FRichCurve DecayCurve;
 };
 
-UCLASS()
-class AUDIOMIXER_API USubmixEffectReverbPreset : public USoundEffectSubmixPreset
+UCLASS(MinimalAPI)
+class USubmixEffectReverbPreset : public USoundEffectSubmixPreset
 {
 	GENERATED_BODY()
 
@@ -174,10 +174,10 @@ public:
 	EFFECT_PRESET_METHODS(SubmixEffectReverb)
 
 	UFUNCTION(BlueprintCallable, Category = "Audio|Effects")
-	void SetSettings(const FSubmixEffectReverbSettings& InSettings);
+	AUDIOMIXER_API void SetSettings(const FSubmixEffectReverbSettings& InSettings);
 
 	UFUNCTION(BlueprintCallable, Category = "Audio|Effects")
-	void SetSettingsWithReverbEffect(const UReverbEffect* InReverbEffect, const float WetLevel, const float DryLevel = 0.0f);
+	AUDIOMIXER_API void SetSettingsWithReverbEffect(const UReverbEffect* InReverbEffect, const float WetLevel, const float DryLevel = 0.0f);
 
 	UPROPERTY(EditAnywhere, BlueprintReadWrite, Category = SubmixEffectPreset)
 	FSubmixEffectReverbSettings Settings;
diff --git a/Engine/Source/Runtime/AudioMixer/Public/AudioBusSubsystem.h b/Engine/Source/Runtime/AudioMixer/Public/AudioBusSubsystem.h
--- a/Engine/Source/Runtime/AudioMixer/Public/AudioBusSubsystem.h
+++ b/Engine/Source/Runtime/AudioMixer/Public/AudioBusSubsystem.h
@@ -19,7 +19,7 @@ namespace Audio
 	struct FPatchOutput;
 	typedef TSharedPtr<FPatchOutput, ESPMode::ThreadSafe> FPatchOutputStrongPtr;
 
-	struct AUDIOMIXER_API FAudioBusKey
+	struct FAudioBusKey
 	{
 		uint32 ObjectId = INDEX_NONE; // from a corresponding UObject (UAudioBus) if applicable
 		uint32 InstanceId = INDEX_NONE;
@@ -57,38 +57,38 @@ namespace Audio
 		
 
 	private:
-		static std::atomic<uint32> InstanceIdCounter;
+		static AUDIOMIXER_API std::atomic<uint32> InstanceIdCounter;
 	};
 }
 
 /**
 *  UAudioBusSubsystem
 */
-UCLASS()
-class AUDIOMIXER_API UAudioBusSubsystem : public UAudioEngineSubsystem
+UCLASS(MinimalAPI)
+class UAudioBusSubsystem : public UAudioEngineSubsystem
 {
 	GENERATED_BODY()
 
 public:
-	UAudioBusSubsystem();
+	AUDIOMIXER_API UAudioBusSubsystem();
 	virtual ~UAudioBusSubsystem() = default;
 
 	//~ Begin USubsystem interface
-	virtual bool ShouldCreateSubsystem(UObject* Outer) const override;
-	virtual void Initialize(FSubsystemCollectionBase& Collection) override;
-	virtual void Deinitialize() override;
+	AUDIOMIXER_API virtual bool ShouldCreateSubsystem(UObject* Outer) const override;
+	AUDIOMIXER_API virtual void Initialize(FSubsystemCollectionBase& Collection) override;
+	AUDIOMIXER_API virtual void Deinitialize() override;
 	//~ End USubsystem interface
 
 	// Audio bus API from FMixerDevice
-	void StartAudioBus(Audio::FAudioBusKey InAudioBusKey, int32 InNumChannels, bool bInIsAutomatic);
-	void StopAudioBus(Audio::FAudioBusKey InAudioBusKey);
-	bool IsAudioBusActive(Audio::FAudioBusKey InAudioBusKey) const;
+	AUDIOMIXER_API void StartAudioBus(Audio::FAudioBusKey InAudioBusKey, int32 InNumChannels, bool bInIsAutomatic);
+	AUDIOMIXER_API void StopAudioBus(Audio::FAudioBusKey InAudioBusKey);
+	AUDIOMIXER_API bool IsAudioBusActive(Audio::FAudioBusKey InAudioBusKey) const;
 	
-	Audio::FPatchInput AddPatchInputForAudioBus(Audio::FAudioBusKey InAudioBusKey, int32 InFrames, int32 InChannels, float InGain = 1.f);
-	Audio::FPatchOutputStrongPtr AddPatchOutputForAudioBus(Audio::FAudioBusKey InAudioBusKey, int32 InFrames, int32 InChannels, float InGain = 1.f);
+	AUDIOMIXER_API Audio::FPatchInput AddPatchInputForAudioBus(Audio::FAudioBusKey InAudioBusKey, int32 InFrames, int32 InChannels, float InGain = 1.f);
+	AUDIOMIXER_API Audio::FPatchOutputStrongPtr AddPatchOutputForAudioBus(Audio::FAudioBusKey InAudioBusKey, int32 InFrames, int32 InChannels, float InGain = 1.f);
 
-	void InitDefaultAudioBuses();
-	void ShutdownDefaultAudioBuses();
+	AUDIOMIXER_API void InitDefaultAudioBuses();
+	AUDIOMIXER_API void ShutdownDefaultAudioBuses();
 
 private:
 	struct FActiveBusData
diff --git a/Engine/Source/Runtime/AudioMixer/Public/AudioDeviceNotificationSubsystem.h b/Engine/Source/Runtime/AudioMixer/Public/AudioDeviceNotificationSubsystem.h
--- a/Engine/Source/Runtime/AudioMixer/Public/AudioDeviceNotificationSubsystem.h
+++ b/Engine/Source/Runtime/AudioMixer/Public/AudioDeviceNotificationSubsystem.h
@@ -46,8 +46,8 @@ DECLARE_MULTICAST_DELEGATE_OneParam(FOnAudioDeviceChangeNative, FString);
 /**
  *  UAudioDeviceNotificationSubsystem
  */
-UCLASS()
-class AUDIOMIXER_API UAudioDeviceNotificationSubsystem : public UEngineSubsystem
+UCLASS(MinimalAPI)
+class UAudioDeviceNotificationSubsystem : public UEngineSubsystem
 													   , public Audio::IAudioMixerDeviceChangedListener
 {
 	GENERATED_BODY()
@@ -59,19 +59,19 @@ public:
 	static UAudioDeviceNotificationSubsystem* Get() { return GEngine->GetEngineSubsystem<UAudioDeviceNotificationSubsystem>(); }
 
 	//~ Begin UEngineSubsystem Interface
-	virtual void Initialize(FSubsystemCollectionBase& Collection) override;
-	virtual void Deinitialize() override;
+	AUDIOMIXER_API virtual void Initialize(FSubsystemCollectionBase& Collection) override;
+	AUDIOMIXER_API virtual void Deinitialize() override;
 	//~ End UEngineSubsystem Interface
 
 	//~ Begin IAudioMixerDeviceChangedListener Interface
-	virtual void OnDefaultCaptureDeviceChanged(const Audio::EAudioDeviceRole InAudioDeviceRole, const FString& DeviceId) override;
-	virtual void OnDefaultRenderDeviceChanged(const Audio::EAudioDeviceRole InAudioDeviceRole, const FString& DeviceId) override;
-	virtual void OnDeviceAdded(const FString& DeviceId, bool bIsRenderDevice) override;
-	virtual void OnDeviceRemoved(const FString& DeviceId, bool bIsRenderDevice) override;
-	virtual void OnDeviceStateChanged(const FString& DeviceId, const Audio::EAudioDeviceState InState, bool bIsRenderDevice) override;
+	AUDIOMIXER_API virtual void OnDefaultCaptureDeviceChanged(const Audio::EAudioDeviceRole InAudioDeviceRole, const FString& DeviceId) override;
+	AUDIOMIXER_API virtual void OnDefaultRenderDeviceChanged(const Audio::EAudioDeviceRole InAudioDeviceRole, const FString& DeviceId) override;
+	AUDIOMIXER_API virtual void OnDeviceAdded(const FString& DeviceId, bool bIsRenderDevice) override;
+	AUDIOMIXER_API virtual void OnDeviceRemoved(const FString& DeviceId, bool bIsRenderDevice) override;
+	AUDIOMIXER_API virtual void OnDeviceStateChanged(const FString& DeviceId, const Audio::EAudioDeviceState InState, bool bIsRenderDevice) override;
 	//~ End IAudioMixerDeviceChangedListener Interface
 
-	virtual void OnDeviceSwitched(const FString& DeviceId);
+	AUDIOMIXER_API virtual void OnDeviceSwitched(const FString& DeviceId);
 
 	/** Multicast delegate triggered when default capture device changes */
 	UPROPERTY(BlueprintAssignable, Category = "Audio Delegates")
@@ -111,7 +111,7 @@ public:
 
 protected:
 
-	EAudioDeviceChangedRole GetDeviceChangedRole(Audio::EAudioDeviceRole InRole) const;
-	EAudioDeviceChangedState GetDeviceChangedState(Audio::EAudioDeviceState InState) const;
+	AUDIOMIXER_API EAudioDeviceChangedRole GetDeviceChangedRole(Audio::EAudioDeviceRole InRole) const;
+	AUDIOMIXER_API EAudioDeviceChangedState GetDeviceChangedState(Audio::EAudioDeviceState InState) const;
 };
 
diff --git a/Engine/Source/Runtime/AudioMixer/Public/AudioMixerBlueprintLibrary.h b/Engine/Source/Runtime/AudioMixer/Public/AudioMixerBlueprintLibrary.h
--- a/Engine/Source/Runtime/AudioMixer/Public/AudioMixerBlueprintLibrary.h
+++ b/Engine/Source/Runtime/AudioMixer/Public/AudioMixerBlueprintLibrary.h
@@ -122,7 +122,7 @@ enum class ESwapAudioOutputDeviceResultState : uint8
  * Out structure for use with AudioMixerBlueprintLibrary::SwapAudioOutputDevice
  */
 USTRUCT(BlueprintType)
-struct AUDIOMIXER_API FSwapAudioOutputResult
+struct FSwapAudioOutputResult
 {
 	GENERATED_USTRUCT_BODY()
 
@@ -145,7 +145,7 @@ struct AUDIOMIXER_API FSwapAudioOutputResult
  * Platform audio output device info, in a Blueprint-readable format
  */
 USTRUCT(BlueprintType)
-struct AUDIOMIXER_API FAudioOutputDeviceInfo
+struct FAudioOutputDeviceInfo
 {
 	GENERATED_USTRUCT_BODY()
 
@@ -159,7 +159,7 @@ struct AUDIOMIXER_API FAudioOutputDeviceInfo
 		, bIsCurrentDevice(false)
 	{};
 
-	FAudioOutputDeviceInfo(const Audio::FAudioPlatformDeviceInfo& InDeviceInfo);
+	AUDIOMIXER_API FAudioOutputDeviceInfo(const Audio::FAudioPlatformDeviceInfo& InDeviceInfo);
 
 	/** The name of the audio device */
 	UPROPERTY(BlueprintReadOnly, Category="Audio")
@@ -209,8 +209,8 @@ DECLARE_DYNAMIC_DELEGATE_OneParam(FOnMainAudioOutputDeviceObtained, const FStrin
  */
 DECLARE_DYNAMIC_DELEGATE_OneParam(FOnCompletedDeviceSwap, const FSwapAudioOutputResult&, SwapResult);
 
-UCLASS(meta=(ScriptName="AudioMixerLibrary"))
-class AUDIOMIXER_API UAudioMixerBlueprintLibrary : public UBlueprintFunctionLibrary
+UCLASS(meta=(ScriptName="AudioMixerLibrary"), MinimalAPI)
+class UAudioMixerBlueprintLibrary : public UBlueprintFunctionLibrary
 {
 	GENERATED_BODY()
 
@@ -221,83 +221,83 @@ public:
 	* @return The data in a string format
 	*/
 	UFUNCTION(BlueprintPure, meta = (DisplayName = "Audio Output Device Info To String", CompactNodeTitle = "->", BlueprintAutocast), Category = "Audio")
-	static FString Conv_AudioOutputDeviceInfoToString(const FAudioOutputDeviceInfo& Info);
+	static AUDIOMIXER_API FString Conv_AudioOutputDeviceInfoToString(const FAudioOutputDeviceInfo& Info);
 
 	/** Adds a submix effect preset to the master submix. */
 	UFUNCTION(BlueprintCallable, Category = "Audio|Effects", meta=(WorldContext="WorldContextObject"))
-	static void AddMasterSubmixEffect(const UObject* WorldContextObject, USoundEffectSubmixPreset* SubmixEffectPreset);
+	static AUDIOMIXER_API void AddMasterSubmixEffect(const UObject* WorldContextObject, USoundEffectSubmixPreset* SubmixEffectPreset);
 
 	/** Removes a submix effect preset from the master submix. */
 	UFUNCTION(BlueprintCallable, Category = "Audio|Effects", meta=(WorldContext="WorldContextObject"))
-	static void RemoveMasterSubmixEffect(const UObject* WorldContextObject, USoundEffectSubmixPreset* SubmixEffectPreset);
+	static AUDIOMIXER_API void RemoveMasterSubmixEffect(const UObject* WorldContextObject, USoundEffectSubmixPreset* SubmixEffectPreset);
 
 	/** Clears all master submix effects. */
 	UFUNCTION(BlueprintCallable, Category = "Audio|Effects", meta = (WorldContext = "WorldContextObject"))
-	static void ClearMasterSubmixEffects(const UObject* WorldContextObject);
+	static AUDIOMIXER_API void ClearMasterSubmixEffects(const UObject* WorldContextObject);
 
 	/** Adds a submix effect preset to the given submix at the end of its submix effect chain. Returns the number of submix effects. */
 	UFUNCTION(BlueprintCallable, Category = "Audio|Effects", meta = (WorldContext = "WorldContextObject"))
-	static int32 AddSubmixEffect(const UObject* WorldContextObject, USoundSubmix* SoundSubmix, USoundEffectSubmixPreset* SubmixEffectPreset);
+	static AUDIOMIXER_API int32 AddSubmixEffect(const UObject* WorldContextObject, USoundSubmix* SoundSubmix, USoundEffectSubmixPreset* SubmixEffectPreset);
 
 	UE_DEPRECATED(4.27, "RemoveSubmixEffectPreset is deprecated, use RemoveSubmixEffect.")
 	UFUNCTION(BlueprintCallable, Category = "Audio|Effects", meta = (WorldContext = "WorldContextObject", DeprecatedFunction))
-	static void RemoveSubmixEffectPreset(const UObject* WorldContextObject, USoundSubmix* SoundSubmix, USoundEffectSubmixPreset* SubmixEffectPreset);
+	static AUDIOMIXER_API void RemoveSubmixEffectPreset(const UObject* WorldContextObject, USoundSubmix* SoundSubmix, USoundEffectSubmixPreset* SubmixEffectPreset);
 
 	/** Removes all instances of a submix effect preset from the given submix. */
 	UFUNCTION(BlueprintCallable, Category = "Audio|Effects", meta = (WorldContext = "WorldContextObject"))
-	static void RemoveSubmixEffect(const UObject* WorldContextObject, USoundSubmix* SoundSubmix, USoundEffectSubmixPreset* SubmixEffectPreset);
+	static AUDIOMIXER_API void RemoveSubmixEffect(const UObject* WorldContextObject, USoundSubmix* SoundSubmix, USoundEffectSubmixPreset* SubmixEffectPreset);
 
 	UE_DEPRECATED(4.27, "RemoveSubmixEffectPresetAtIndex is deprecated, use RemoveSubmixEffectAtIndex.")
 	UFUNCTION(BlueprintCallable, Category = "Audio|Effects", meta = (WorldContext = "WorldContextObject", DeprecatedFunction))
-	static void RemoveSubmixEffectPresetAtIndex(const UObject* WorldContextObject, USoundSubmix* SoundSubmix, int32 SubmixChainIndex);
+	static AUDIOMIXER_API void RemoveSubmixEffectPresetAtIndex(const UObject* WorldContextObject, USoundSubmix* SoundSubmix, int32 SubmixChainIndex);
 
 	/** Removes the submix effect at the given submix chain index, if there is a submix effect at that index. */
 	UFUNCTION(BlueprintCallable, Category = "Audio|Effects", meta = (WorldContext = "WorldContextObject"))
-	static void RemoveSubmixEffectAtIndex(const UObject* WorldContextObject, USoundSubmix* SoundSubmix, int32 SubmixChainIndex);
+	static AUDIOMIXER_API void RemoveSubmixEffectAtIndex(const UObject* WorldContextObject, USoundSubmix* SoundSubmix, int32 SubmixChainIndex);
 
 	UE_DEPRECATED(4.27, "ReplaceSoundEffectSubmix is deprecated, use ReplaceSubmixEffect.")
 	UFUNCTION(BlueprintCallable, Category = "Audio|Effects", meta = (WorldContext = "WorldContextObject", DeprecatedFunction))
-	static void ReplaceSoundEffectSubmix(const UObject* WorldContextObject, USoundSubmix* InSoundSubmix, int32 SubmixChainIndex, USoundEffectSubmixPreset* SubmixEffectPreset);
+	static AUDIOMIXER_API void ReplaceSoundEffectSubmix(const UObject* WorldContextObject, USoundSubmix* InSoundSubmix, int32 SubmixChainIndex, USoundEffectSubmixPreset* SubmixEffectPreset);
 
 	/** Replaces the submix effect at the given submix chain index, adds the effect if there is none at that index. */
 	UFUNCTION(BlueprintCallable, Category = "Audio|Effects", meta = (WorldContext = "WorldContextObject"))
-	static void ReplaceSubmixEffect(const UObject* WorldContextObject, USoundSubmix* InSoundSubmix, int32 SubmixChainIndex, USoundEffectSubmixPreset* SubmixEffectPreset);
+	static AUDIOMIXER_API void ReplaceSubmixEffect(const UObject* WorldContextObject, USoundSubmix* InSoundSubmix, int32 SubmixChainIndex, USoundEffectSubmixPreset* SubmixEffectPreset);
 
 	/** Clears all submix effects on the given submix. */
 	UFUNCTION(BlueprintCallable, Category = "Audio|Effects", meta = (WorldContext = "WorldContextObject"))
-	static void ClearSubmixEffects(const UObject* WorldContextObject, USoundSubmix* SoundSubmix);
+	static AUDIOMIXER_API void ClearSubmixEffects(const UObject* WorldContextObject, USoundSubmix* SoundSubmix);
 
 	/** Sets a submix effect chain override on the given submix. The effect chain will cross fade from the base effect chain or current override to the new override. */
 	UFUNCTION(BlueprintCallable, Category = "Audio|Effects", meta = (WorldContext = "WorldContextObject"))
-	static void SetSubmixEffectChainOverride(const UObject* WorldContextObject, USoundSubmix* SoundSubmix, TArray<USoundEffectSubmixPreset*> SubmixEffectPresetChain, float FadeTimeSec);
+	static AUDIOMIXER_API void SetSubmixEffectChainOverride(const UObject* WorldContextObject, USoundSubmix* SoundSubmix, TArray<USoundEffectSubmixPreset*> SubmixEffectPresetChain, float FadeTimeSec);
 
 	/** Clears all submix effect overrides on the given submix and returns it to the default effect chain. */
 	UFUNCTION(BlueprintCallable, Category = "Audio|Effects", meta = (WorldContext = "WorldContextObject"))
-	static void ClearSubmixEffectChainOverride(const UObject* WorldContextObject, USoundSubmix* SoundSubmix, float FadeTimeSec);
+	static AUDIOMIXER_API void ClearSubmixEffectChainOverride(const UObject* WorldContextObject, USoundSubmix* SoundSubmix, float FadeTimeSec);
 
 	/** Start recording audio. By leaving the Submix To Record field blank, you can record the master output of the game. */
 	UFUNCTION(BlueprintCallable, Category = "Audio|Recording", meta = (WorldContext = "WorldContextObject", AdvancedDisplay = 1))
-	static void StartRecordingOutput(const UObject* WorldContextObject, float ExpectedDuration, USoundSubmix* SubmixToRecord = nullptr);
+	static AUDIOMIXER_API void StartRecordingOutput(const UObject* WorldContextObject, float ExpectedDuration, USoundSubmix* SubmixToRecord = nullptr);
 	
 	/** Stop recording audio. Path can be absolute, or relative (to the /Saved/BouncedWavFiles folder). By leaving the Submix To Record field blank, you can record the master output of the game.  */
 	UFUNCTION(BlueprintCallable, Category = "Audio|Recording", meta = (WorldContext = "WorldContextObject", DisplayName = "Finish Recording Output", AdvancedDisplay = 4))
-	static USoundWave* StopRecordingOutput(const UObject* WorldContextObject, EAudioRecordingExportType ExportType, const FString& Name, FString Path, USoundSubmix* SubmixToRecord = nullptr, USoundWave* ExistingSoundWaveToOverwrite= nullptr);
+	static AUDIOMIXER_API USoundWave* StopRecordingOutput(const UObject* WorldContextObject, EAudioRecordingExportType ExportType, const FString& Name, FString Path, USoundSubmix* SubmixToRecord = nullptr, USoundWave* ExistingSoundWaveToOverwrite= nullptr);
 
 	/** Pause recording audio, without finalizing the recording to disk. By leaving the Submix To Record field blank, you can record the master output of the game. */
 	UFUNCTION(BlueprintCallable, Category = "Audio|Recording", meta = (WorldContext = "WorldContextObject", AdvancedDisplay = 1))
-	static void PauseRecordingOutput(const UObject* WorldContextObject, USoundSubmix* SubmixToPause = nullptr);
+	static AUDIOMIXER_API void PauseRecordingOutput(const UObject* WorldContextObject, USoundSubmix* SubmixToPause = nullptr);
 
 	/** Resume recording audio after pausing. By leaving the Submix To Pause field blank, you can record the master output of the game. */
 	UFUNCTION(BlueprintCallable, Category = "Audio|Recording", meta = (WorldContext = "WorldContextObject", AdvancedDisplay = 1))
-	static void ResumeRecordingOutput(const UObject* WorldContextObject, USoundSubmix* SubmixToPause = nullptr);
+	static AUDIOMIXER_API void ResumeRecordingOutput(const UObject* WorldContextObject, USoundSubmix* SubmixToPause = nullptr);
 
 	/** Start spectrum analysis of the audio output. By leaving the Submix To Analyze blank, you can analyze the master output of the game. */
 	UFUNCTION(BlueprintCallable, Category = "Audio|Analysis", meta = (WorldContext = "WorldContextObject", AdvancedDisplay = 1))
-	static void StartAnalyzingOutput(const UObject* WorldContextObject, USoundSubmix* SubmixToAnalyze = nullptr, EFFTSize FFTSize = EFFTSize::DefaultSize, EFFTPeakInterpolationMethod InterpolationMethod = EFFTPeakInterpolationMethod::Linear, EFFTWindowType WindowType = EFFTWindowType::Hann, float HopSize = 0, EAudioSpectrumType SpectrumType = EAudioSpectrumType::MagnitudeSpectrum);
+	static AUDIOMIXER_API void StartAnalyzingOutput(const UObject* WorldContextObject, USoundSubmix* SubmixToAnalyze = nullptr, EFFTSize FFTSize = EFFTSize::DefaultSize, EFFTPeakInterpolationMethod InterpolationMethod = EFFTPeakInterpolationMethod::Linear, EFFTWindowType WindowType = EFFTWindowType::Hann, float HopSize = 0, EAudioSpectrumType SpectrumType = EAudioSpectrumType::MagnitudeSpectrum);
 
 	/** Stop spectrum analysis. */
 	UFUNCTION(BlueprintCallable, Category = "Audio|Analysis", meta = (WorldContext = "WorldContextObject", AdvancedDisplay = 1))
-	static void StopAnalyzingOutput(const UObject* WorldContextObject, USoundSubmix* SubmixToStopAnalyzing = nullptr);
+	static AUDIOMIXER_API void StopAnalyzingOutput(const UObject* WorldContextObject, USoundSubmix* SubmixToStopAnalyzing = nullptr);
 
 	/** Make an array of musically spaced bands with ascending frequency.
 	 *
@@ -308,7 +308,7 @@ public:
 	 *  @param InReleaseTimeMsec - The release time (in milliseconds) to apply to each band's envelope tracker.
 	 */
 	UFUNCTION(BlueprintPure, Category = "Audio|Analysis", meta = (AdvancedDisplay = 3))
-	static TArray<FSoundSubmixSpectralAnalysisBandSettings> MakeMusicalSpectralAnalysisBandSettings(int32 InNumSemitones=60, EMusicalNoteName InStartingMusicalNote = EMusicalNoteName::C, int32 InStartingOctave = 2, int32 InAttackTimeMsec = 10, int32 InReleaseTimeMsec = 10);
+	static AUDIOMIXER_API TArray<FSoundSubmixSpectralAnalysisBandSettings> MakeMusicalSpectralAnalysisBandSettings(int32 InNumSemitones=60, EMusicalNoteName InStartingMusicalNote = EMusicalNoteName::C, int32 InStartingOctave = 2, int32 InAttackTimeMsec = 10, int32 InReleaseTimeMsec = 10);
 
 	/** Make an array of logarithmically spaced bands. 
 	 *
@@ -319,7 +319,7 @@ public:
 	 *  @param InReleaseTimeMsec - The release time (in milliseconds) to apply to each band's envelope tracker.
 	 */
 	UFUNCTION(BlueprintPure, Category = "Audio|Analysis", meta = (AdvancedDisplay = 3))
-	static TArray<FSoundSubmixSpectralAnalysisBandSettings> MakeFullSpectrumSpectralAnalysisBandSettings(int32 InNumBands = 30, float InMinimumFrequency=40.f, float InMaximumFrequency=16000.f, int32 InAttackTimeMsec = 10, int32 InReleaseTimeMsec = 10);
+	static AUDIOMIXER_API TArray<FSoundSubmixSpectralAnalysisBandSettings> MakeFullSpectrumSpectralAnalysisBandSettings(int32 InNumBands = 30, float InMinimumFrequency=40.f, float InMaximumFrequency=16000.f, int32 InAttackTimeMsec = 10, int32 InReleaseTimeMsec = 10);
 
 	/** Make an array of bands which span the frequency range of a given EAudioSpectrumBandPresetType. 
 	 *
@@ -329,69 +329,69 @@ public:
 	 *  @param InReleaseTimeMsec - The release time (in milliseconds) to apply to each band's envelope tracker.
 	 */
 	UFUNCTION(BlueprintPure, Category = "Audio|Analysis", meta = (AdvancedDisplay = 2))
-	static TArray<FSoundSubmixSpectralAnalysisBandSettings> MakePresetSpectralAnalysisBandSettings(EAudioSpectrumBandPresetType InBandPresetType, int32 InNumBands = 10, int32 InAttackTimeMsec = 10, int32 InReleaseTimeMsec = 10);
+	static AUDIOMIXER_API TArray<FSoundSubmixSpectralAnalysisBandSettings> MakePresetSpectralAnalysisBandSettings(EAudioSpectrumBandPresetType InBandPresetType, int32 InNumBands = 10, int32 InAttackTimeMsec = 10, int32 InReleaseTimeMsec = 10);
 
 	/** Retrieve the magnitudes for the given frequencies. */
 	UFUNCTION(BlueprintCallable, Category = "Audio|Analysis", meta = (WorldContext = "WorldContextObject", AdvancedDisplay = 3))
-	static void GetMagnitudeForFrequencies(const UObject* WorldContextObject, const TArray<float>& Frequencies, TArray<float>& Magnitudes, USoundSubmix* SubmixToAnalyze = nullptr);
+	static AUDIOMIXER_API void GetMagnitudeForFrequencies(const UObject* WorldContextObject, const TArray<float>& Frequencies, TArray<float>& Magnitudes, USoundSubmix* SubmixToAnalyze = nullptr);
 
 	/** Retrieve the phases for the given frequencies. */
 	UFUNCTION(BlueprintCallable, Category = "Audio|Analysis", meta = (WorldContext = "WorldContextObject", AdvancedDisplay = 3))
-	static void GetPhaseForFrequencies(const UObject* WorldContextObject, const TArray<float>& Frequencies, TArray<float>& Phases, USoundSubmix* SubmixToAnalyze = nullptr);
+	static AUDIOMIXER_API void GetPhaseForFrequencies(const UObject* WorldContextObject, const TArray<float>& Frequencies, TArray<float>& Phases, USoundSubmix* SubmixToAnalyze = nullptr);
 
 	/** Adds source effect entry to preset chain. Only effects the instance of the preset chain */
 	UFUNCTION(BlueprintCallable, Category = "Audio|Effects", meta = (WorldContext = "WorldContextObject"))
-	static void AddSourceEffectToPresetChain(const UObject* WorldContextObject, USoundEffectSourcePresetChain* PresetChain, FSourceEffectChainEntry Entry);
+	static AUDIOMIXER_API void AddSourceEffectToPresetChain(const UObject* WorldContextObject, USoundEffectSourcePresetChain* PresetChain, FSourceEffectChainEntry Entry);
 
 	/** Removes source effect entry from preset chain. Only affects the instance of preset chain. */
 	UFUNCTION(BlueprintCallable, Category = "Audio|Effects", meta = (WorldContext = "WorldContextObject"))
-	static void RemoveSourceEffectFromPresetChain(const UObject* WorldContextObject, USoundEffectSourcePresetChain* PresetChain, int32 EntryIndex);
+	static AUDIOMIXER_API void RemoveSourceEffectFromPresetChain(const UObject* WorldContextObject, USoundEffectSourcePresetChain* PresetChain, int32 EntryIndex);
 
 	/** Set whether or not to bypass the effect at the source effect chain index. */
 	UFUNCTION(BlueprintCallable, Category = "Audio|Effects", meta = (WorldContext = "WorldContextObject"))
-	static void SetBypassSourceEffectChainEntry(const UObject* WorldContextObject, USoundEffectSourcePresetChain* PresetChain, int32 EntryIndex, bool bBypassed);
+	static AUDIOMIXER_API void SetBypassSourceEffectChainEntry(const UObject* WorldContextObject, USoundEffectSourcePresetChain* PresetChain, int32 EntryIndex, bool bBypassed);
 
 	/** Returns the number of effect chain entries in the given source effect chain. */
 	UFUNCTION(BlueprintCallable, Category = "Audio|Effects", meta = (WorldContext = "WorldContextObject"))
-	static int32 GetNumberOfEntriesInSourceEffectChain(const UObject* WorldContextObject, USoundEffectSourcePresetChain* PresetChain);
+	static AUDIOMIXER_API int32 GetNumberOfEntriesInSourceEffectChain(const UObject* WorldContextObject, USoundEffectSourcePresetChain* PresetChain);
 
 	/** Begin loading a sound into the cache so that it can be played immediately. */
 	UFUNCTION(BlueprintCallable, Category = "Audio|Cache")
-	static void PrimeSoundForPlayback(USoundWave* SoundWave, const FOnSoundLoadComplete OnLoadCompletion);
+	static AUDIOMIXER_API void PrimeSoundForPlayback(USoundWave* SoundWave, const FOnSoundLoadComplete OnLoadCompletion);
 
 	/** Begin loading any sounds referenced by a sound cue into the cache so that it can be played immediately. */
 	UFUNCTION(BlueprintCallable, Category = "Audio|Cache")
-	static void PrimeSoundCueForPlayback(USoundCue* SoundCue);
+	static AUDIOMIXER_API void PrimeSoundCueForPlayback(USoundCue* SoundCue);
 
 	/** Trim memory used by the audio cache. Returns the number of megabytes freed. */
 	UFUNCTION(BlueprintCallable, Category = "Audio|Cache")
-	static float TrimAudioCache(float InMegabytesToFree);
+	static AUDIOMIXER_API float TrimAudioCache(float InMegabytesToFree);
 
 	/** Starts the given audio bus. */
 	UFUNCTION(BlueprintCallable, Category = "Audio|Bus", meta = (WorldContext = "WorldContextObject"))
-	static void StartAudioBus(const UObject* WorldContextObject, UAudioBus* AudioBus);
+	static AUDIOMIXER_API void StartAudioBus(const UObject* WorldContextObject, UAudioBus* AudioBus);
 
 	/** Stops the given audio bus. */
 	UFUNCTION(BlueprintCallable, Category = "Audio|Bus", meta = (WorldContext = "WorldContextObject"))
-	static void StopAudioBus(const UObject* WorldContextObject, UAudioBus* AudioBus);
+	static AUDIOMIXER_API void StopAudioBus(const UObject* WorldContextObject, UAudioBus* AudioBus);
 
 	/** Queries if the given audio bus is active (and audio can be mixed to it). */
 	UFUNCTION(BlueprintCallable, Category = "Audio|Bus", meta = (WorldContext = "WorldContextObject"))
-	static bool IsAudioBusActive(const UObject* WorldContextObject, UAudioBus* AudioBus);
+	static AUDIOMIXER_API bool IsAudioBusActive(const UObject* WorldContextObject, UAudioBus* AudioBus);
 
 	/**
 	* Gets information about all audio output devices available in the system
 	* @param OnObtainDevicesEvent - the event to fire when the audio endpoint devices have been retrieved
 	*/
 	UFUNCTION(BlueprintCallable, Category = "Audio", meta = (WorldContext = "WorldContextObject"))
-	static void GetAvailableAudioOutputDevices(const UObject* WorldContextObject, const FOnAudioOutputDevicesObtained& OnObtainDevicesEvent);
+	static AUDIOMIXER_API void GetAvailableAudioOutputDevices(const UObject* WorldContextObject, const FOnAudioOutputDevicesObtained& OnObtainDevicesEvent);
 
 	/**
 	* Gets information about the currently used audio output device
 	* @param OnObtainCurrentDeviceEvent - the event to fire when the audio endpoint devices have been retrieved
 	*/
 	UFUNCTION(BlueprintCallable, Category = "Audio", meta = (WorldContext = "WorldContextObject"))
-	static void GetCurrentAudioOutputDeviceName(const UObject* WorldContextObject, const FOnMainAudioOutputDeviceObtained& OnObtainCurrentDeviceEvent);
+	static AUDIOMIXER_API void GetCurrentAudioOutputDeviceName(const UObject* WorldContextObject, const FOnMainAudioOutputDeviceObtained& OnObtainCurrentDeviceEvent);
 
 	/**
 	* Hotswaps to the requested audio output device
@@ -399,6 +399,6 @@ public:
 	* @param OnCompletedDeviceSwap - the event to fire when the audio endpoint devices have been retrieved
 	*/
 	UFUNCTION(BlueprintCallable, Category = "Audio", meta = (WorldContext = "WorldContextObject"))
-	static void SwapAudioOutputDevice(const UObject* WorldContextObject, const FString& NewDeviceId, const FOnCompletedDeviceSwap& OnCompletedDeviceSwap);
+	static AUDIOMIXER_API void SwapAudioOutputDevice(const UObject* WorldContextObject, const FString& NewDeviceId, const FOnCompletedDeviceSwap& OnCompletedDeviceSwap);
 };
 
diff --git a/Engine/Source/Runtime/AudioMixer/Public/AudioMixerSubmix.h b/Engine/Source/Runtime/AudioMixer/Public/AudioMixerSubmix.h
--- a/Engine/Source/Runtime/AudioMixer/Public/AudioMixerSubmix.h
+++ b/Engine/Source/Runtime/AudioMixer/Public/AudioMixerSubmix.h
@@ -111,255 +111,255 @@ namespace Audio
 		}
 	};
 
-	class AUDIOMIXER_API FMixerSubmix
+	class FMixerSubmix
 	{
 	public:
-		FMixerSubmix(FMixerDevice* InMixerDevice);
-		virtual ~FMixerSubmix();
+		AUDIOMIXER_API FMixerSubmix(FMixerDevice* InMixerDevice);
+		AUDIOMIXER_API virtual ~FMixerSubmix();
 
 		// Initialize the submix object with the USoundSubmix ptr. Sets up child and parent connects.
-		void Init(const USoundSubmixBase* InSoundSubmix, bool bAllowReInit = true);
+		AUDIOMIXER_API void Init(const USoundSubmixBase* InSoundSubmix, bool bAllowReInit = true);
 
 		// Returns the mixer submix Id
 		uint32 GetId() const { return Id; }
 
 		// Sets the parent submix to the given submix
-		void SetParentSubmix(TWeakPtr<FMixerSubmix, ESPMode::ThreadSafe> Submix);
+		AUDIOMIXER_API void SetParentSubmix(TWeakPtr<FMixerSubmix, ESPMode::ThreadSafe> Submix);
 
 		// Adds the given submix to this submix's children
-		void AddChildSubmix(TWeakPtr<FMixerSubmix, ESPMode::ThreadSafe> Submix);
+		AUDIOMIXER_API void AddChildSubmix(TWeakPtr<FMixerSubmix, ESPMode::ThreadSafe> Submix);
 
 		// Removes the given submix from this submix's children
-		void RemoveChildSubmix(TWeakPtr<FMixerSubmix, ESPMode::ThreadSafe> SubmixWeakPtr);
+		AUDIOMIXER_API void RemoveChildSubmix(TWeakPtr<FMixerSubmix, ESPMode::ThreadSafe> SubmixWeakPtr);
 
 		// Sets the output level of the submix in linear gain
-		void SetOutputVolume(float InOutputLevel);
+		AUDIOMIXER_API void SetOutputVolume(float InOutputLevel);
 
 		// Sets the static output volume of the submix in linear gain
-		void SetDryLevel(float InDryLevel);
+		AUDIOMIXER_API void SetDryLevel(float InDryLevel);
 
 		// Sets the wet level of the submix in linear gain
-		void SetWetLevel(float InWetLevel);
+		AUDIOMIXER_API void SetWetLevel(float InWetLevel);
 
 		// Update modulation settings of the submix
-		void UpdateModulationSettings(const TSet<TObjectPtr<USoundModulatorBase>>& InOutputModulators, const TSet<TObjectPtr<USoundModulatorBase>>& InWetLevelModulators, const TSet<TObjectPtr<USoundModulatorBase>>& InDryLevelModulators);
+		AUDIOMIXER_API void UpdateModulationSettings(const TSet<TObjectPtr<USoundModulatorBase>>& InOutputModulators, const TSet<TObjectPtr<USoundModulatorBase>>& InWetLevelModulators, const TSet<TObjectPtr<USoundModulatorBase>>& InDryLevelModulators);
 
 		// Update modulation settings of the submix with Decibel values
-		void SetModulationBaseLevels(float InVolumeModBaseDb, float InWetModeBaseDb, float InDryModBaseDb);
+		AUDIOMIXER_API void SetModulationBaseLevels(float InVolumeModBaseDb, float InWetModeBaseDb, float InDryModBaseDb);
 
 		// Gets the submix channels channels
-		int32 GetSubmixChannels() const;
+		AUDIOMIXER_API int32 GetSubmixChannels() const;
 
 		// Gets this submix's parent submix
-		TWeakPtr<FMixerSubmix, ESPMode::ThreadSafe> GetParentSubmix();
+		AUDIOMIXER_API TWeakPtr<FMixerSubmix, ESPMode::ThreadSafe> GetParentSubmix();
 
 		// Returns the number of source voices currently a part of this submix.
-		int32 GetNumSourceVoices() const;
+		AUDIOMIXER_API int32 GetNumSourceVoices() const;
 
 		// Returns the number of wet effects in this submix.
-		int32 GetNumEffects() const;
+		AUDIOMIXER_API int32 GetNumEffects() const;
 
 		// Returns the size of the submix chain. 
-		int32 GetSizeOfSubmixChain() const;
+		AUDIOMIXER_API int32 GetSizeOfSubmixChain() const;
 
 		// Add (if not already added) or sets the amount of the source voice's send amount
-		void AddOrSetSourceVoice(FMixerSourceVoice* InSourceVoice, const float SendLevel, EMixerSourceSubmixSendStage InSubmixSendStage);
+		AUDIOMIXER_API void AddOrSetSourceVoice(FMixerSourceVoice* InSourceVoice, const float SendLevel, EMixerSourceSubmixSendStage InSubmixSendStage);
 
-		FPatchOutputStrongPtr AddPatch(float InGain);
+		AUDIOMIXER_API FPatchOutputStrongPtr AddPatch(float InGain);
 
 		/** Removes the given source voice from the submix. */
-		void RemoveSourceVoice(FMixerSourceVoice* InSourceVoice);
+		AUDIOMIXER_API void RemoveSourceVoice(FMixerSourceVoice* InSourceVoice);
 
 		/** Appends the effect submix to the effect submix chain. */
-		void AddSoundEffectSubmix(FSoundEffectSubmixPtr InSoundEffectSubmix);
+		AUDIOMIXER_API void AddSoundEffectSubmix(FSoundEffectSubmixPtr InSoundEffectSubmix);
 
 		/** Removes the submix effect from the effect submix chain. */
-		void RemoveSoundEffectSubmix(uint32 SubmixPresetId);
+		AUDIOMIXER_API void RemoveSoundEffectSubmix(uint32 SubmixPresetId);
 
 		/** Removes the submix effect from the effect submix chain at the given submix index. */
-		void RemoveSoundEffectSubmixAtIndex(int32 InIndex);
+		AUDIOMIXER_API void RemoveSoundEffectSubmixAtIndex(int32 InIndex);
 
 		/** Clears all submix effects from the effect submix chain. */
-		void ClearSoundEffectSubmixes();
+		AUDIOMIXER_API void ClearSoundEffectSubmixes();
 
 		/** Sets a submix effect chain override with the given fade time in seconds. */
-		void SetSubmixEffectChainOverride(const TArray<FSoundEffectSubmixPtr>& InSubmixEffectPresetChain, float InFadeTimeSec);
+		AUDIOMIXER_API void SetSubmixEffectChainOverride(const TArray<FSoundEffectSubmixPtr>& InSubmixEffectPresetChain, float InFadeTimeSec);
 
 		/** Clears any submix effect chain overrides in the given fade time in seconds. */
-		void ClearSubmixEffectChainOverride(float InFadeTimeSec);
+		AUDIOMIXER_API void ClearSubmixEffectChainOverride(float InFadeTimeSec);
 
 		/** Swaps effect for provided submix at the given index.  Fails if effect at index doesn't exist */
-		void ReplaceSoundEffectSubmix(int32 InIndex, FSoundEffectSubmixPtr InEffectInstance);
+		AUDIOMIXER_API void ReplaceSoundEffectSubmix(int32 InIndex, FSoundEffectSubmixPtr InEffectInstance);
 
 		/** Whether or not this submix instance is muted. */
-		void SetBackgroundMuted(bool bInMuted);
+		AUDIOMIXER_API void SetBackgroundMuted(bool bInMuted);
 
 		/** Checks to see if submix is valid.  Submix can be considered invalid if the OwningSubmix
 		  * pointer is stale.
 		  */
-		bool IsValid() const;
+		AUDIOMIXER_API bool IsValid() const;
 
 		// Function which processes audio.
-		void ProcessAudio(FAlignedFloatBuffer& OutAudio);
-		void ProcessAudio(ISoundfieldAudioPacket& OutputAudio);
+		AUDIOMIXER_API void ProcessAudio(FAlignedFloatBuffer& OutAudio);
+		AUDIOMIXER_API void ProcessAudio(ISoundfieldAudioPacket& OutputAudio);
 
-		void SendAudioToSubmixBufferListeners(FAlignedFloatBuffer& OutAudioBuffer);
+		AUDIOMIXER_API void SendAudioToSubmixBufferListeners(FAlignedFloatBuffer& OutAudioBuffer);
 
 		// This should be called if this submix doesn't send it's audio to a parent submix,
 		// but rather an external endpoint.
-		void ProcessAudioAndSendToEndpoint();
+		AUDIOMIXER_API void ProcessAudioAndSendToEndpoint();
 
 		// Returns the device sample rate this submix is rendering to
-		int32 GetSampleRate() const;
+		AUDIOMIXER_API int32 GetSampleRate() const;
 
 		// Returns the output channels this submix is rendering to
-		int32 GetNumOutputChannels() const;
+		AUDIOMIXER_API int32 GetNumOutputChannels() const;
 
 		// Returns the number of effects in this submix's effect chain
-		int32 GetNumChainEffects();
+		AUDIOMIXER_API int32 GetNumChainEffects();
 
 		// Returns the submix effect at the given effect chain index
-		FSoundEffectSubmixPtr GetSubmixEffect(const int32 InIndex);
+		AUDIOMIXER_API FSoundEffectSubmixPtr GetSubmixEffect(const int32 InIndex);
 
 		// This must be called on the entire submix graph before calling SetupSoundfieldStreams.
-		void SetSoundfieldFactory(ISoundfieldFactory* InSoundfieldFactory);
+		AUDIOMIXER_API void SetSoundfieldFactory(ISoundfieldFactory* InSoundfieldFactory);
 
 		// updates settings, potentially creating or removing ambisonics streams based on what types of submixes this submix is connected to.
-		void SetupSoundfieldStreams(const USoundfieldEncodingSettingsBase* SoundfieldSettings, TArray<USoundfieldEffectBase*>& Processors, ISoundfieldFactory* InSoundfieldFactory);
-		void TeardownSoundfieldStreams();
+		AUDIOMIXER_API void SetupSoundfieldStreams(const USoundfieldEncodingSettingsBase* SoundfieldSettings, TArray<USoundfieldEffectBase*>& Processors, ISoundfieldFactory* InSoundfieldFactory);
+		AUDIOMIXER_API void TeardownSoundfieldStreams();
 
-		void SetupEndpoint(IAudioEndpointFactory* InFactory, const UAudioEndpointSettingsBase* InSettings);
-		void SetupEndpoint(ISoundfieldEndpointFactory* InFactory, const USoundfieldEndpointSettingsBase* InSettings);
+		AUDIOMIXER_API void SetupEndpoint(IAudioEndpointFactory* InFactory, const UAudioEndpointSettingsBase* InSettings);
+		AUDIOMIXER_API void SetupEndpoint(ISoundfieldEndpointFactory* InFactory, const USoundfieldEndpointSettingsBase* InSettings);
 
-		void UpdateEndpointSettings(TUniquePtr<IAudioEndpointSettingsProxy>&& InSettings);
-		void UpdateEndpointSettings(TUniquePtr<ISoundfieldEndpointSettingsProxy>&& InSettings);
+		AUDIOMIXER_API void UpdateEndpointSettings(TUniquePtr<IAudioEndpointSettingsProxy>&& InSettings);
+		AUDIOMIXER_API void UpdateEndpointSettings(TUniquePtr<ISoundfieldEndpointSettingsProxy>&& InSettings);
 
 		// This is called by the corresponding USoundSubmix when StartRecordingOutput is called.
-		void OnStartRecordingOutput(float ExpectedDuration);
+		AUDIOMIXER_API void OnStartRecordingOutput(float ExpectedDuration);
 
 		// This is called by the corresponding USoundSubmix when StopRecordingOutput is called.
-		FAlignedFloatBuffer& OnStopRecordingOutput(float& OutNumChannels, float& OutSampleRate);
+		AUDIOMIXER_API FAlignedFloatBuffer& OnStopRecordingOutput(float& OutNumChannels, float& OutSampleRate);
 
 		// This is called by the corresponding USoundSubmix when PauseRecording is called.
-		void PauseRecordingOutput();
+		AUDIOMIXER_API void PauseRecordingOutput();
 
 		// This is called by the corresponding USoundSubmix when ResumeRecording is called.
-		void ResumeRecordingOutput();
+		AUDIOMIXER_API void ResumeRecordingOutput();
 
 		// Register buffer listener with this submix
-		void RegisterBufferListener(ISubmixBufferListener* BufferListener);
+		AUDIOMIXER_API void RegisterBufferListener(ISubmixBufferListener* BufferListener);
 		
 		// Unregister buffer listener with this submix
-		void UnregisterBufferListener(ISubmixBufferListener* BufferListener);
+		AUDIOMIXER_API void UnregisterBufferListener(ISubmixBufferListener* BufferListener);
 
 		// Starts envelope following with the given attack time and release time
-		void StartEnvelopeFollowing(int32 AttackTime, int32 ReleaseTime);
+		AUDIOMIXER_API void StartEnvelopeFollowing(int32 AttackTime, int32 ReleaseTime);
 
 		// Stops envelope following the submix
-		void StopEnvelopeFollowing();
+		AUDIOMIXER_API void StopEnvelopeFollowing();
 
 		// Adds an envelope follower delegate
-		void AddEnvelopeFollowerDelegate(const FOnSubmixEnvelopeBP& OnSubmixEnvelopeBP);
+		AUDIOMIXER_API void AddEnvelopeFollowerDelegate(const FOnSubmixEnvelopeBP& OnSubmixEnvelopeBP);
 
 		// Initializes a new FFT analyzer for this submix and immediately begins feeding audio to it.
-		void StartSpectrumAnalysis(const FSoundSpectrumAnalyzerSettings& InSettings);
+		AUDIOMIXER_API void StartSpectrumAnalysis(const FSoundSpectrumAnalyzerSettings& InSettings);
 
 		// Terminates whatever FFT Analyzer is being used for this submix.
-		void StopSpectrumAnalysis();
+		AUDIOMIXER_API void StopSpectrumAnalysis();
 
 		// Adds an spectral analysis delegate
-		void AddSpectralAnalysisDelegate(const FSoundSpectrumAnalyzerDelegateSettings& InDelegateSettings, const FOnSubmixSpectralAnalysisBP& OnSubmixSpectralAnalysisBP);
+		AUDIOMIXER_API void AddSpectralAnalysisDelegate(const FSoundSpectrumAnalyzerDelegateSettings& InDelegateSettings, const FOnSubmixSpectralAnalysisBP& OnSubmixSpectralAnalysisBP);
 
 		// Removes an existing spectral analysis delegate
-		void RemoveSpectralAnalysisDelegate(const FOnSubmixSpectralAnalysisBP& OnSubmixSpectralAnalysisBP);
+		AUDIOMIXER_API void RemoveSpectralAnalysisDelegate(const FOnSubmixSpectralAnalysisBP& OnSubmixSpectralAnalysisBP);
 
 		// Gets the most recent magnitude values for each corresponding value in InFrequencies (in Hz).
 		// This requires StartSpectrumAnalysis to be called first.
-		void GetMagnitudeForFrequencies(const TArray<float>& InFrequencies, TArray<float>& OutMagnitudes);
+		AUDIOMIXER_API void GetMagnitudeForFrequencies(const TArray<float>& InFrequencies, TArray<float>& OutMagnitudes);
 
 		// Gets the most recent phase values for each corresponding value in InFrequencies (in Hz).
 		// This requires StartSpectrumAnalysis to be called first.
-		void GetPhaseForFrequencies(const TArray<float>& InFrequencies, TArray<float>& OutPhases);
+		AUDIOMIXER_API void GetPhaseForFrequencies(const TArray<float>& InFrequencies, TArray<float>& OutPhases);
 
 		// Broadcast the envelope and submix delegates on the game thread
-		void BroadcastDelegates();
+		AUDIOMIXER_API void BroadcastDelegates();
 
 		// returns true if this submix is encoded to a soundfield.
-		bool IsSoundfieldSubmix() const;
+		AUDIOMIXER_API bool IsSoundfieldSubmix() const;
 
 		// returns true if this submix sends it's audio to the default endpoint.
-		bool IsDefaultEndpointSubmix() const;
+		AUDIOMIXER_API bool IsDefaultEndpointSubmix() const;
 
 		// Returns true if this submix sends its audio to an IAudioEndpoint.
-		bool IsExternalEndpointSubmix() const;
+		AUDIOMIXER_API bool IsExternalEndpointSubmix() const;
 
 		// returns true if this submix sends its audio to an ISoundfieldEndpoint.
-		bool IsSoundfieldEndpointSubmix() const;
+		AUDIOMIXER_API bool IsSoundfieldEndpointSubmix() const;
 
 		//Returns true if this is an endpoint type that should no-op for this platform
-		bool IsDummyEndpointSubmix() const;
+		AUDIOMIXER_API bool IsDummyEndpointSubmix() const;
 
 		// Returns true if the submix is currently rendering audio. The current rendering time is passed in.
-		bool IsRenderingAudio() const;
+		AUDIOMIXER_API bool IsRenderingAudio() const;
 
 		// Set whether or not this submix is told to auto disable. 
-		void SetAutoDisable(bool bInAutoDisable);
+		AUDIOMIXER_API void SetAutoDisable(bool bInAutoDisable);
 
 		// Sets the auto-disable time
-		void SetAutoDisableTime(float InAutoDisableTime);
+		AUDIOMIXER_API void SetAutoDisableTime(float InAutoDisableTime);
 
 		// Get a unique key for this submix's format and settings.
 		// If another submix has an identical format and settings it will have an equivalent key.
-		FSoundfieldEncodingKey GetKeyForSubmixEncoding();
+		AUDIOMIXER_API FSoundfieldEncodingKey GetKeyForSubmixEncoding();
 
-		ISoundfieldFactory* GetSoundfieldFactory();
+		AUDIOMIXER_API ISoundfieldFactory* GetSoundfieldFactory();
 
-		ISoundfieldEncodingSettingsProxy& GetSoundfieldSettings();
+		AUDIOMIXER_API ISoundfieldEncodingSettingsProxy& GetSoundfieldSettings();
 
-		FAudioPluginInitializationParams GetInitializationParamsForSoundfieldStream();
+		AUDIOMIXER_API FAudioPluginInitializationParams GetInitializationParamsForSoundfieldStream();
 
-		FSoundfieldSpeakerPositionalData GetDefaultPositionalDataForAudioDevice();
+		AUDIOMIXER_API FSoundfieldSpeakerPositionalData GetDefaultPositionalDataForAudioDevice();
 
 	protected:
 		// Initialize the submix internal
-		void InitInternal();
+		AUDIOMIXER_API void InitInternal();
 
 		// Down mix the given buffer to the desired down mix channel count
-		static void DownmixBuffer(const int32 InChannels, const FAlignedFloatBuffer& InBuffer, const int32 OutChannels, FAlignedFloatBuffer& OutNewBuffer);
+		static AUDIOMIXER_API void DownmixBuffer(const int32 InChannels, const FAlignedFloatBuffer& InBuffer, const int32 OutChannels, FAlignedFloatBuffer& OutNewBuffer);
 
-		void MixBufferDownToMono(const FAlignedFloatBuffer& InBuffer, int32 NumInputChannels, FAlignedFloatBuffer& OutBuffer);
+		AUDIOMIXER_API void MixBufferDownToMono(const FAlignedFloatBuffer& InBuffer, int32 NumInputChannels, FAlignedFloatBuffer& OutBuffer);
 
-		void SetupSoundfieldEncodersForChildren();
-		void SetupSoundfieldEncodingForChild(FChildSubmixInfo& InChild);
+		AUDIOMIXER_API void SetupSoundfieldEncodersForChildren();
+		AUDIOMIXER_API void SetupSoundfieldEncodingForChild(FChildSubmixInfo& InChild);
 
 		// Check to see if we need to decode from ambisonics for parent
-		void SetupSoundfieldStreamForParent();
+		AUDIOMIXER_API void SetupSoundfieldStreamForParent();
 
 		// This sets up the ambisonics positional data for speakers, based on what new format we need to convert to.
-		void SetUpSoundfieldPositionalData(const TSharedPtr<Audio::FMixerSubmix, ESPMode::ThreadSafe>& InParentSubmix);
+		AUDIOMIXER_API void SetUpSoundfieldPositionalData(const TSharedPtr<Audio::FMixerSubmix, ESPMode::ThreadSafe>& InParentSubmix);
 
 		// Encode a source and sum it into the mixed soundfield.
-		void MixInSource(const ISoundfieldAudioPacket& InAudio, const ISoundfieldEncodingSettingsProxy& InSettings, ISoundfieldAudioPacket& PacketToSumTo);
+		AUDIOMIXER_API void MixInSource(const ISoundfieldAudioPacket& InAudio, const ISoundfieldEncodingSettingsProxy& InSettings, ISoundfieldAudioPacket& PacketToSumTo);
 
-		void UpdateListenerRotation(const FQuat& InRotation);
+		AUDIOMIXER_API void UpdateListenerRotation(const FQuat& InRotation);
 
 		// Calls ProcessAudio on the child submix, performs all necessary conversions and mixes in it's resulting audio.
-		void MixInChildSubmix(FChildSubmixInfo& Child, ISoundfieldAudioPacket& PacketToSumTo);
+		AUDIOMIXER_API void MixInChildSubmix(FChildSubmixInfo& Child, ISoundfieldAudioPacket& PacketToSumTo);
 
-		FName GetSoundfieldFormat() const;
+		AUDIOMIXER_API FName GetSoundfieldFormat() const;
 
-		TUniquePtr<ISoundfieldTranscodeStream> GetTranscoderForChildSubmix(const TSharedPtr<Audio::FMixerSubmix, ESPMode::ThreadSafe>& InChildSubmix);
+		AUDIOMIXER_API TUniquePtr<ISoundfieldTranscodeStream> GetTranscoderForChildSubmix(const TSharedPtr<Audio::FMixerSubmix, ESPMode::ThreadSafe>& InChildSubmix);
 
 	protected:
 
 		// Pump command queue
-		void PumpCommandQueue();
+		AUDIOMIXER_API void PumpCommandQueue();
 
 		// Add command to the command queue
-		void SubmixCommand(TFunction<void()> Command);
+		AUDIOMIXER_API void SubmixCommand(TFunction<void()> Command);
 
 		// Generates audio from the given effect chain into the given buffer
-		bool GenerateEffectChainAudio(FSoundEffectSubmixInputData& InputData, const FAlignedFloatBuffer& InAudioBuffer, TArray<FSoundEffectSubmixPtr>& InEffectChain, FAlignedFloatBuffer& OutBuffer);
+		AUDIOMIXER_API bool GenerateEffectChainAudio(FSoundEffectSubmixInputData& InputData, const FAlignedFloatBuffer& InAudioBuffer, TArray<FSoundEffectSubmixPtr>& InEffectChain, FAlignedFloatBuffer& OutBuffer);
 
 		// This mixer submix's Id
 		uint32 Id;
