[Engine]

diff --git a/Engine/Source/Runtime/AudioMixer/Public/Components/SynthComponent.h b/Engine/Source/Runtime/AudioMixer/Public/Components/SynthComponent.h
--- a/Engine/Source/Runtime/AudioMixer/Public/Components/SynthComponent.h
+++ b/Engine/Source/Runtime/AudioMixer/Public/Components/SynthComponent.h
@@ -41,20 +41,20 @@ DECLARE_DYNAMIC_MULTICAST_DELEGATE_OneParam(FOnSynthEnvelopeValue, const float,
 DECLARE_MULTICAST_DELEGATE_TwoParams(FOnSynthEnvelopeValueNative, const class UAudioComponent*, const float);
 
 
-UCLASS()
-class AUDIOMIXER_API USynthSound : public USoundWaveProcedural
+UCLASS(MinimalAPI)
+class USynthSound : public USoundWaveProcedural
 {
 	GENERATED_UCLASS_BODY()
 
-	void Init(USynthComponent* InSynthComponent, const int32 InNumChannels, const int32 SampleRate, const int32 InCallbackSize);
-	void StartOnAudioDevice(FAudioDevice* InAudioDevice);
+	AUDIOMIXER_API void Init(USynthComponent* InSynthComponent, const int32 InNumChannels, const int32 SampleRate, const int32 InCallbackSize);
+	AUDIOMIXER_API void StartOnAudioDevice(FAudioDevice* InAudioDevice);
 
 	/** Begin USoundWave */
-	virtual void OnBeginGenerate() override;
-	virtual int32 OnGeneratePCMAudio(TArray<uint8>& OutAudio, int32 NumSamples) override;
-	virtual void OnEndGenerate() override;
-	virtual Audio::EAudioMixerStreamDataFormat::Type GetGeneratedPCMDataFormat() const override;
-	virtual ISoundGeneratorPtr CreateSoundGenerator(const FSoundGeneratorInitParams& InParams) override;
+	AUDIOMIXER_API virtual void OnBeginGenerate() override;
+	AUDIOMIXER_API virtual int32 OnGeneratePCMAudio(TArray<uint8>& OutAudio, int32 NumSamples) override;
+	AUDIOMIXER_API virtual void OnEndGenerate() override;
+	AUDIOMIXER_API virtual Audio::EAudioMixerStreamDataFormat::Type GetGeneratedPCMDataFormat() const override;
+	AUDIOMIXER_API virtual ISoundGeneratorPtr CreateSoundGenerator(const FSoundGeneratorInitParams& InParams) override;
 	/** End USoundWave */
 
 protected:
@@ -75,85 +75,85 @@ public:
 	}
 };
 
-UCLASS(abstract, ClassGroup = Synth, hidecategories = (Object, ActorComponent, Physics, Rendering, Mobility, LOD))
-class AUDIOMIXER_API USynthComponent : public USceneComponent
+UCLASS(abstract, ClassGroup = Synth, hidecategories = (Object, ActorComponent, Physics, Rendering, Mobility, LOD), MinimalAPI)
+class USynthComponent : public USceneComponent
 {
 	GENERATED_BODY()
 
 public:
-	USynthComponent(const FObjectInitializer& ObjectInitializer);
+	AUDIOMIXER_API USynthComponent(const FObjectInitializer& ObjectInitializer);
 
 	//~ Begin USceneComponent Interface
-	virtual void Activate(bool bReset = false) override;
-	virtual void Deactivate() override;
+	AUDIOMIXER_API virtual void Activate(bool bReset = false) override;
+	AUDIOMIXER_API virtual void Deactivate() override;
 	//~ End USceneComponent Interface
 
 	//~ Begin ActorComponent Interface.
-	virtual void OnRegister() override;
-	virtual void OnUnregister() override;
-	virtual bool IsReadyForOwnerToAutoDestroy() const override;
-	virtual void EndPlay(const EEndPlayReason::Type Reason) override;
+	AUDIOMIXER_API virtual void OnRegister() override;
+	AUDIOMIXER_API virtual void OnUnregister() override;
+	AUDIOMIXER_API virtual bool IsReadyForOwnerToAutoDestroy() const override;
+	AUDIOMIXER_API virtual void EndPlay(const EEndPlayReason::Type Reason) override;
 	//~ End ActorComponent Interface.
 
 	//~ Begin UObject Interface.
 #if WITH_EDITOR
-	virtual void PostEditChangeProperty(FPropertyChangedEvent& PropertyChangedEvent) override;
+	AUDIOMIXER_API virtual void PostEditChangeProperty(FPropertyChangedEvent& PropertyChangedEvent) override;
 #endif // WITH_EDITOR
 
 #if WITH_EDITORONLY_DATA
-	virtual void PostLoad() override;
+	AUDIOMIXER_API virtual void PostLoad() override;
 #endif //WITH_EDITORONLY_DATA
 
-	virtual void Serialize(FArchive& Ar) override;
+	AUDIOMIXER_API virtual void Serialize(FArchive& Ar) override;
 	//~ End UObject Interface
 
 	// Starts the synth generating audio.
 	UFUNCTION(BlueprintCallable, Category = "Synth|Components|Audio")
-	void Start();
+	AUDIOMIXER_API void Start();
 
 	// Stops the synth generating audio.
 	UFUNCTION(BlueprintCallable, Category = "Synth|Components|Audio")
-	void Stop();
+	AUDIOMIXER_API void Stop();
 
 	/** Returns true if this component is currently playing. */
 	UFUNCTION(BlueprintCallable, Category = "Synth|Components|Audio")
-	bool IsPlaying() const;
+	AUDIOMIXER_API bool IsPlaying() const;
 
 	/** Set a new volume multiplier */
 	UFUNCTION(BlueprintCallable, Category = "Audio|Components|Audio")
-	void SetVolumeMultiplier(float VolumeMultiplier);
+	AUDIOMIXER_API void SetVolumeMultiplier(float VolumeMultiplier);
 
 	/** Sets how much audio the sound should send to the given submix. */
 	UFUNCTION(BlueprintCallable, Category = "Audio|Components|Audio")
-	void SetSubmixSend(USoundSubmixBase* Submix, float SendLevel);
+	AUDIOMIXER_API void SetSubmixSend(USoundSubmixBase* Submix, float SendLevel);
 
 	/** Sets how much audio the sound should send to the given SourceBus (pre effect). */
 	UFUNCTION(BlueprintCallable, Category = "Audio|Components|Audio")
-	void SetSourceBusSendPreEffect(USoundSourceBus* SoundSourceBus, float SourceBusSendLevel);
+	AUDIOMIXER_API void SetSourceBusSendPreEffect(USoundSourceBus* SoundSourceBus, float SourceBusSendLevel);
 
 	/** Sets how much audio the sound should send to the given SourceBus (post effect). */
 	UFUNCTION(BlueprintCallable, Category = "Audio|Components|Audio")
-	void SetSourceBusSendPostEffect(USoundSourceBus* SoundSourceBus, float SourceBusSendLevel);
+	AUDIOMIXER_API void SetSourceBusSendPostEffect(USoundSourceBus* SoundSourceBus, float SourceBusSendLevel);
 
 	/** Sets how much audio the sound should send to the given AudioBus (pre effect). */
 	UFUNCTION(BlueprintCallable, Category = "Audio|Components|Audio")
-	void SetAudioBusSendPreEffect(UAudioBus* AudioBus, float AudioBusSendLevel);
+	AUDIOMIXER_API void SetAudioBusSendPreEffect(UAudioBus* AudioBus, float AudioBusSendLevel);
 
 	/** Sets how much audio the sound should send to the given AudioBus (post effect). */
 	UFUNCTION(BlueprintCallable, Category = "Audio|Components|Audio")
-	void SetAudioBusSendPostEffect(UAudioBus* AudioBus, float AudioBusSendLevel);
+	AUDIOMIXER_API void SetAudioBusSendPostEffect(UAudioBus* AudioBus, float AudioBusSendLevel);
 
 	/** Sets whether or not the low pass filter is enabled on the audio component. */
 	UFUNCTION(BlueprintCallable, Category = "Audio|Components|Audio")
-	void SetLowPassFilterEnabled(bool InLowPassFilterEnabled);
+	AUDIOMIXER_API void SetLowPassFilterEnabled(bool InLowPassFilterEnabled);
 
 	/** Sets lowpass filter frequency of the audio component. */
 	UFUNCTION(BlueprintCallable, Category = "Audio|Components|Audio")
-	virtual void SetLowPassFilterFrequency(float InLowPassFilterFrequency);
+	AUDIOMIXER_API virtual void SetLowPassFilterFrequency(float InLowPassFilterFrequency);
 
 	/** Sets whether or not the synth component outputs its audio to any source or audio buses. */
 	UFUNCTION(BlueprintCallable, Category = "Audio|Components|Audio")
-	void SetOutputToBusOnly(bool bInOutputToBusOnly);
+	AUDIOMIXER_API void SetOutputToBusOnly(bool bInOutputToBusOnly);
 
 	/**
 	 * This function allows designers to call Play on an Audio Component instance while applying a volume curve over time. 
@@ -164,7 +164,7 @@ public:
 	 * @param FadeCurve The curve to use when interpolating between the old and new volume
 	 */
 	UFUNCTION(BlueprintCallable, Category = "Audio|Components|Audio")
-	void FadeIn(float FadeInDuration, float FadeVolumeLevel = 1.0f, float StartTime = 0.0f, const EAudioFaderCurve FadeCurve = EAudioFaderCurve::Linear) const;
+	AUDIOMIXER_API void FadeIn(float FadeInDuration, float FadeVolumeLevel = 1.0f, float StartTime = 0.0f, const EAudioFaderCurve FadeCurve = EAudioFaderCurve::Linear) const;
 
 	/**
 	 * This function allows designers to call a delayed Stop on an Audio Component instance while applying a
@@ -175,7 +175,7 @@ public:
 	 * @param FadeCurve The curve to use when interpolating between the old and new volume
 	 */
 	UFUNCTION(BlueprintCallable, Category = "Audio|Components|Audio")
-	void FadeOut(float FadeOutDuration, float FadeVolumeLevel, const EAudioFaderCurve FadeCurve = EAudioFaderCurve::Linear) const;
+	AUDIOMIXER_API void FadeOut(float FadeOutDuration, float FadeVolumeLevel, const EAudioFaderCurve FadeCurve = EAudioFaderCurve::Linear) const;
 
 	/** This function allows designers to trigger an adjustment to the sound instance’s playback Volume with options for smoothly applying a curve over time.
      * @param AdjustVolumeDuration The length of time in which to interpolate between the initial volume and the new volume.
@@ -183,7 +183,7 @@ public:
      * @param FadeCurve The curve used when interpolating between the old and new volume.
      */
 	UFUNCTION(BlueprintCallable, Category = "Audio|Components|Audio")
-	void AdjustVolume(float AdjustVolumeDuration, float AdjustVolumeLevel, const EAudioFaderCurve FadeCurve = EAudioFaderCurve::Linear) const;
+	AUDIOMIXER_API void AdjustVolume(float AdjustVolumeDuration, float AdjustVolumeLevel, const EAudioFaderCurve FadeCurve = EAudioFaderCurve::Linear) const;
 
 	/**
 	* Sets the routing for one of the given Synth component's Modulation Destinations.
@@ -192,7 +192,7 @@ public:
 	* @param RoutingMethod The routing method to use for the given modulator.
 	*/
 	UFUNCTION(BlueprintCallable, Category = "Audio|Components|Audio", DisplayName = "Set Modulation Routing")
-	void SetModulationRouting(const TSet<USoundModulatorBase*>& Modulators, const EModulationDestination Destination, const EModulationRouting RoutingMethod = EModulationRouting::Inherit);
+	AUDIOMIXER_API void SetModulationRouting(const TSet<USoundModulatorBase*>& Modulators, const EModulationDestination Destination, const EModulationRouting RoutingMethod = EModulationRouting::Inherit);
 	
 	/**
 	* Gets the set of currently active modulators for a given Modulation Destination.
@@ -200,7 +200,7 @@ public:
 	* @return The set of of Modulators applied to this component for the given Destination.
 	*/
 	UFUNCTION(BlueprintPure, Category = "Audio|Components|Audio", DisplayName = "Get Modulators")
-	UPARAM(DisplayName = "Modulators") TSet<USoundModulatorBase*> GetModulators(const EModulationDestination Destination);
+	UPARAM(DisplayName = "ModAUDIOMIXER_API ulators") TSet<USoundModulatorBase*> GetModulators(const EModulationDestination Destination);
 
 	/** Auto destroy this component on completion */
 	UPROPERTY()
@@ -293,15 +293,15 @@ public:
 	/** Call if creating this synth component not via an actor component in BP, but in code or some other location. 
 	 *  Optionally override the sample rate of the sound wave, otherwise it uses the audio device's sample rate. 
 	 */
-	void Initialize(int32 SampleRateOverride = INDEX_NONE);
+	AUDIOMIXER_API void Initialize(int32 SampleRateOverride = INDEX_NONE);
 
 	/** Creates the audio component if it hasn't already been created yet. This should only be used when trying to
 	 *  assign explicit settings to the AudioComponent before calling Start(). 
 	 */
-	void CreateAudioComponent();
+	AUDIOMIXER_API void CreateAudioComponent();
 
 	/** Retrieves this synth component's audio component. */
-	UAudioComponent* GetAudioComponent();
+	AUDIOMIXER_API UAudioComponent* GetAudioComponent();
 
 	/** The attack time in milliseconds for the envelope follower. Delegate callbacks can be registered to get the
 	 *  envelope value of sounds played with this audio component. Only used in audio mixer.
@@ -321,20 +321,20 @@ public:
 	/** Shadow delegate for non UObject subscribers */
 	FOnSynthEnvelopeValueNative OnAudioEnvelopeValueNative;
 
-	void OnAudioComponentEnvelopeValue(const UAudioComponent* AudioComponent, const USoundWave* SoundWave, const float EnvelopeValue);
+	AUDIOMIXER_API void OnAudioComponentEnvelopeValue(const UAudioComponent* AudioComponent, const USoundWave* SoundWave, const float EnvelopeValue);
 
 	// Adds and removes audio buffer listener
-	void AddAudioBufferListener(IAudioBufferListener* InAudioBufferListener);
-	void RemoveAudioBufferListener(IAudioBufferListener* InAudioBufferListener);
+	AUDIOMIXER_API void AddAudioBufferListener(IAudioBufferListener* InAudioBufferListener);
+	AUDIOMIXER_API void RemoveAudioBufferListener(IAudioBufferListener* InAudioBufferListener);
 
-	virtual USoundClass* GetSoundClass();
+	AUDIOMIXER_API virtual USoundClass* GetSoundClass();
 
-	virtual void BeginDestroy() override;
+	AUDIOMIXER_API virtual void BeginDestroy() override;
 
 protected:
 
 	// Method to execute parameter changes on game thread in audio render thread
-	void SynthCommand(TFunction<void()> Command);
+	AUDIOMIXER_API void SynthCommand(TFunction<void()> Command);
 
 	// Called when synth is created.
 	virtual bool Init(int32& SampleRate) { return true; }
@@ -361,10 +361,10 @@ protected:
 
 	// Called by procedural sound wave
 	// Returns the number of samples actually generated
-	int32 OnGeneratePCMAudio(float* GeneratedPCMData, int32 NumSamples);
+	AUDIOMIXER_API int32 OnGeneratePCMAudio(float* GeneratedPCMData, int32 NumSamples);
 
 	// Gets the audio device associated with this synth component
-	FAudioDevice* GetAudioDevice() const;
+	AUDIOMIXER_API FAudioDevice* GetAudioDevice() const;
 
 	// Can be set by the derived class, defaults to 2
 	int32 NumChannels;
@@ -374,7 +374,7 @@ protected:
 
 private:
 	// Creates the synth component's sound generator, calls into overridden client code to create the instance.
-	ISoundGeneratorPtr CreateSoundGeneratorInternal(const FSoundGeneratorInitParams& InParams);
+	AUDIOMIXER_API ISoundGeneratorPtr CreateSoundGeneratorInternal(const FSoundGeneratorInitParams& InParams);
 
 	UPROPERTY(Transient)
 	TObjectPtr<USynthSound> Synth;
@@ -382,7 +382,7 @@ private:
 	UPROPERTY(Transient)
 	TObjectPtr<UAudioComponent> AudioComponent;
 
-	void PumpPendingMessages();
+	AUDIOMIXER_API void PumpPendingMessages();
 
 #if SYNTH_GENERATOR_TEST_TONE
 	Audio::FSineOsc TestSineLeft;
@@ -400,4 +400,4 @@ private:
 	ISoundGeneratorPtr SoundGenerator;
 
 	friend class USynthSound;
-};
\ No newline at end of file
+};
diff --git a/Engine/Source/Runtime/AudioMixer/Public/FileDecoder.h b/Engine/Source/Runtime/AudioMixer/Public/FileDecoder.h
--- a/Engine/Source/Runtime/AudioMixer/Public/FileDecoder.h
+++ b/Engine/Source/Runtime/AudioMixer/Public/FileDecoder.h
@@ -9,19 +9,19 @@
 #include "GenericPlatform/GenericPlatformFile.h"
 #include "Interfaces/IAudioFormat.h"
 
-class AUDIOMIXER_API FAudioFileReader
+class FAudioFileReader
 {
 public:
 	// Constructor. Takes a file path and immediately loads info.
 	// Optionally, CallbackSize can be used to indicate the size of chunks
 	// that will be popped off of this instance.
 	// When set to 0, the entire file is decompressed into memory.
-	FAudioFileReader(const FString& InPath);
+	AUDIOMIXER_API FAudioFileReader(const FString& InPath);
 
 	// Returns file information.
-	void GetFileInfo(FSoundQualityInfo& OutInfo);
+	AUDIOMIXER_API void GetFileInfo(FSoundQualityInfo& OutInfo);
 
-	bool PopAudio(float* OutAudio, int32 NumSamples);
+	AUDIOMIXER_API bool PopAudio(float* OutAudio, int32 NumSamples);
 
 private:
 	FAudioFileReader();
diff --git a/Engine/Source/Runtime/AudioMixer/Public/Quartz/AudioMixerClock.h b/Engine/Source/Runtime/AudioMixer/Public/Quartz/AudioMixerClock.h
--- a/Engine/Source/Runtime/AudioMixer/Public/Quartz/AudioMixerClock.h
+++ b/Engine/Source/Runtime/AudioMixer/Public/Quartz/AudioMixerClock.h
@@ -44,37 +44,37 @@ namespace Audio
 	 *		FQuartzClock itself when it pumps the command queue.
 	 *
 	 */
-	class AUDIOMIXER_API FQuartzClockProxy
+	class FQuartzClockProxy
 	{
 	public:
 		// ctor
 		FQuartzClockProxy() {}
 		FQuartzClockProxy(const FName& Name) : ClockId(Name){ } // conv ctor from FName
-		FQuartzClockProxy(TSharedPtr<FQuartzClock, ESPMode::ThreadSafe> InClock);
+		AUDIOMIXER_API FQuartzClockProxy(TSharedPtr<FQuartzClock, ESPMode::ThreadSafe> InClock);
 
 		FName GetClockName() const { return ClockId; }
 
-		bool IsValid() const;
+		AUDIOMIXER_API bool IsValid() const;
 		operator bool() const { return IsValid(); }
 
 		bool operator==(const FName& Name) const { return ClockId == Name; }
 
-		bool DoesClockExist() const;
+		AUDIOMIXER_API bool DoesClockExist() const;
 
-		bool IsClockRunning() const;
+		AUDIOMIXER_API bool IsClockRunning() const;
 
-		Audio::FQuartzClockTickRate GetTickRate() const;
+		AUDIOMIXER_API Audio::FQuartzClockTickRate GetTickRate() const;
 
-		float GetEstimatedClockRunTimeSeconds() const;
+		AUDIOMIXER_API float GetEstimatedClockRunTimeSeconds() const;
 
-		FQuartzTransportTimeStamp GetCurrentClockTimestamp() const;
+		AUDIOMIXER_API FQuartzTransportTimeStamp GetCurrentClockTimestamp() const;
 
-		float GetDurationOfQuantizationTypeInSeconds(const EQuartzCommandQuantization& QuantizationType, float Multiplier) const;
+		AUDIOMIXER_API float GetDurationOfQuantizationTypeInSeconds(const EQuartzCommandQuantization& QuantizationType, float Multiplier) const;
 
-		float GetBeatProgressPercent(const EQuartzCommandQuantization& QuantizationType) const;
+		AUDIOMIXER_API float GetBeatProgressPercent(const EQuartzCommandQuantization& QuantizationType) const;
 
 		// returns false if the clock is not valid or has shut down
-		bool SendCommandToClock(TFunction<void(FQuartzClock*)> InCommand);
+		AUDIOMIXER_API bool SendCommandToClock(TFunction<void(FQuartzClock*)> InCommand);
 
 		// implicit cast to underlying ID (FName)
 		operator const FName&() const { return ClockId; }
@@ -104,7 +104,7 @@ namespace Audio
 	 *		UpdateCachedState() updates a game-thread copy of data accessed via FQuartzClockProxy
 	 *		(see FQuartzClockState)
 	 */
-	class AUDIOMIXER_API FQuartzClock
+	class FQuartzClock
 	{
 	public:
 
@@ -112,88 +112,88 @@ namespace Audio
 		FQuartzClock(const FName& InName, const FQuartzClockSettings& InClockSettings, FQuartzClockManager* InOwningClockManagerPtr = nullptr);
 
 		// dtor
-		~FQuartzClock();
+		AUDIOMIXER_API ~FQuartzClock();
 
 		// Transport Control:
 		// alter the tick rate (take by-value to make sample-rate adjustments in-place)
-		void ChangeTickRate(FQuartzClockTickRate InNewTickRate, int32 NumFramesLeft = 0);
+		AUDIOMIXER_API void ChangeTickRate(FQuartzClockTickRate InNewTickRate, int32 NumFramesLeft = 0);
 
-		void ChangeTimeSignature(const FQuartzTimeSignature& InNewTimeSignature);
+		AUDIOMIXER_API void ChangeTimeSignature(const FQuartzTimeSignature& InNewTimeSignature);
 
-		void Resume();
+		AUDIOMIXER_API void Resume();
 
-		void Pause();
+		AUDIOMIXER_API void Pause();
 
-		void Restart(bool bPause = true);
+		AUDIOMIXER_API void Restart(bool bPause = true);
 
-		void Stop(bool CancelPendingEvents); // Pause + Restart
+		AUDIOMIXER_API void Stop(bool CancelPendingEvents); // Pause + Restart
 
-		void SetSampleRate(float InNewSampleRate);
+		AUDIOMIXER_API void SetSampleRate(float InNewSampleRate);
 
-		void ResetTransport();
+		AUDIOMIXER_API void ResetTransport();
 
 		// (used for StartOtherClock command to handle the sub-tick as the target clock)
-		void AddToTickDelay(int32 NumFramesOfDelayToAdd);
+		AUDIOMIXER_API void AddToTickDelay(int32 NumFramesOfDelayToAdd);
 
 		// (used for StartOtherClock command to handle the sub-tick as the target clock)
-		void SetTickDelay(int32 NumFramesOfDelay);
+		AUDIOMIXER_API void SetTickDelay(int32 NumFramesOfDelay);
 
-		void Shutdown();
+		AUDIOMIXER_API void Shutdown();
 
 		// Getters:
-		FQuartzClockTickRate GetTickRate();
+		AUDIOMIXER_API FQuartzClockTickRate GetTickRate();
 
-		FName GetName() const;
+		AUDIOMIXER_API FName GetName() const;
 
-		bool IgnoresFlush() const;
+		AUDIOMIXER_API bool IgnoresFlush() const;
 
-		bool DoesMatchSettings(const FQuartzClockSettings& InClockSettings) const;
+		AUDIOMIXER_API bool DoesMatchSettings(const FQuartzClockSettings& InClockSettings) const;
 
-		bool HasPendingEvents() const;
+		AUDIOMIXER_API bool HasPendingEvents() const;
 
-		int32 NumPendingEvents() const;
+		AUDIOMIXER_API int32 NumPendingEvents() const;
 
-		bool IsRunning() const;
+		AUDIOMIXER_API bool IsRunning() const;
 
-		float GetDurationOfQuantizationTypeInSeconds(const EQuartzCommandQuantization& QuantizationType, float Multiplier);
+		AUDIOMIXER_API float GetDurationOfQuantizationTypeInSeconds(const EQuartzCommandQuantization& QuantizationType, float Multiplier);
 
-		float GetBeatProgressPercent(const EQuartzCommandQuantization& QuantizationType) const;
+		AUDIOMIXER_API float GetBeatProgressPercent(const EQuartzCommandQuantization& QuantizationType) const;
 
-		FQuartzTransportTimeStamp GetCurrentTimestamp();
+		AUDIOMIXER_API FQuartzTransportTimeStamp GetCurrentTimestamp();
 
-		float GetEstimatedRunTime();
+		AUDIOMIXER_API float GetEstimatedRunTime();
 
-		FMixerDevice* GetMixerDevice();
+		AUDIOMIXER_API FMixerDevice* GetMixerDevice();
 
-		FMixerSourceManager* GetSourceManager();
+		AUDIOMIXER_API FMixerSourceManager* GetSourceManager();
 
-		FQuartzClockManager* GetClockManager();
+		AUDIOMIXER_API FQuartzClockManager* GetClockManager();
 
-		FQuartzClockCommandQueueWeakPtr GetCommandQueue() const;
+		AUDIOMIXER_API FQuartzClockCommandQueueWeakPtr GetCommandQueue() const;
 
 		// Metronome Event Subscription:
-		void SubscribeToTimeDivision(FQuartzGameThreadSubscriber InSubscriber, EQuartzCommandQuantization InQuantizationBoundary);
+		AUDIOMIXER_API void SubscribeToTimeDivision(FQuartzGameThreadSubscriber InSubscriber, EQuartzCommandQuantization InQuantizationBoundary);
 
-		void SubscribeToAllTimeDivisions(FQuartzGameThreadSubscriber InSubscriber);
+		AUDIOMIXER_API void SubscribeToAllTimeDivisions(FQuartzGameThreadSubscriber InSubscriber);
 
-		void UnsubscribeFromTimeDivision(FQuartzGameThreadSubscriber InSubscriber, EQuartzCommandQuantization InQuantizationBoundary);
+		AUDIOMIXER_API void UnsubscribeFromTimeDivision(FQuartzGameThreadSubscriber InSubscriber, EQuartzCommandQuantization InQuantizationBoundary);
 
-		void UnsubscribeFromAllTimeDivisions(FQuartzGameThreadSubscriber InSubscriber);
+		AUDIOMIXER_API void UnsubscribeFromAllTimeDivisions(FQuartzGameThreadSubscriber InSubscriber);
 
 		// Quantized Command Management:
-		void AddQuantizedCommand(FQuartzQuantizedRequestData& InQuantizedRequestData);
-		void AddQuantizedCommand(FQuartzQuantizedCommandInitInfo& InQuantizationCommandInitInfo);
+		AUDIOMIXER_API void AddQuantizedCommand(FQuartzQuantizedRequestData& InQuantizedRequestData);
+		AUDIOMIXER_API void AddQuantizedCommand(FQuartzQuantizedCommandInitInfo& InQuantizationCommandInitInfo);
 
-		void AddQuantizedCommand(FQuartzQuantizationBoundary InQuantizationBondary, TSharedPtr<IQuartzQuantizedCommand> InNewEvent);
+		AUDIOMIXER_API void AddQuantizedCommand(FQuartzQuantizationBoundary InQuantizationBondary, TSharedPtr<IQuartzQuantizedCommand> InNewEvent);
 
-		bool CancelQuantizedCommand(TSharedPtr<IQuartzQuantizedCommand> InCommandPtr);
+		AUDIOMIXER_API bool CancelQuantizedCommand(TSharedPtr<IQuartzQuantizedCommand> InCommandPtr);
 		
 		// low-resolution clock update
 		// (not sample-accurate!, useful when running without an Audio Device)
-		void LowResolutionTick(float InDeltaTimeSeconds);
+		AUDIOMIXER_API void LowResolutionTick(float InDeltaTimeSeconds);
 
 		// sample accurate clock update
-		void Tick(int32 InNumFramesUntilNextTick);
+		AUDIOMIXER_API void Tick(int32 InNumFramesUntilNextTick);
 
 	private:
 		// Contains the pending command and the number of frames it has to wait to fire
@@ -215,7 +215,7 @@ namespace Audio
 
 		// mutex-protected update at the end of Tick()
 		FCriticalSection CachedClockStateCritSec;
-		void UpdateCachedState();
+		AUDIOMIXER_API void UpdateCachedState();
 
 		// data is cached when an FQuartzClock is ticked
 		struct FQuartzClockState
@@ -230,9 +230,9 @@ namespace Audio
 			
 		} CachedClockState;
 
-		void TickInternal(int32 InNumFramesUntilNextTick, TArray<PendingCommand>& CommandsToTick, int32 FramesOfLatency = 0, int32 FramesOfDelay = 0);
+		AUDIOMIXER_API void TickInternal(int32 InNumFramesUntilNextTick, TArray<PendingCommand>& CommandsToTick, int32 FramesOfLatency = 0, int32 FramesOfDelay = 0);
 
-		bool CancelQuantizedCommandInternal(TSharedPtr<IQuartzQuantizedCommand> InCommandPtr, TArray<PendingCommand>& CommandsToTick);
+		AUDIOMIXER_API bool CancelQuantizedCommandInternal(TSharedPtr<IQuartzQuantizedCommand> InCommandPtr, TArray<PendingCommand>& CommandsToTick);
 
 		// don't allow default ctor, a clock needs to be ready to be used
 		// by the clock manager / FMixerDevice once constructed
diff --git a/Engine/Source/Runtime/AudioMixer/Public/Quartz/AudioMixerClockManager.h b/Engine/Source/Runtime/AudioMixer/Public/Quartz/AudioMixerClockManager.h
--- a/Engine/Source/Runtime/AudioMixer/Public/Quartz/AudioMixerClockManager.h
+++ b/Engine/Source/Runtime/AudioMixer/Public/Quartz/AudioMixerClockManager.h
@@ -14,108 +14,108 @@ namespace Audio
 
 	// Class that owns, updates, and provides access to all active clocks
 	// All methods are thread-safe. The method locks if it returns a value, and stages a command if it returns void
-	class AUDIOMIXER_API FQuartzClockManager : public FQuartLatencyTracker
+	class FQuartzClockManager : public FQuartLatencyTracker
 	{
 	public:
 		// ctor
-		FQuartzClockManager(Audio::FMixerDevice* InOwner = nullptr);
+		AUDIOMIXER_API FQuartzClockManager(Audio::FMixerDevice* InOwner = nullptr);
 
 		// dtor
-		~FQuartzClockManager();
+		AUDIOMIXER_API ~FQuartzClockManager();
 
 		int32 GetNumClocks() const { return ActiveClocks.Num(); }
 
 		// Called on AudioRenderThread
-		void Update(int32 NumFramesUntilNextUpdate);
-		void UpdateClock(FName InClockToAdvance, int32 NumFramesToAdvance);
+		AUDIOMIXER_API void Update(int32 NumFramesUntilNextUpdate);
+		AUDIOMIXER_API void UpdateClock(FName InClockToAdvance, int32 NumFramesToAdvance);
 
 		// can be called from any thread for low-resolution clock updates
 		// (i.e. used when running without an audio device)
 		// not sample-accurate!
-		void LowResoultionUpdate(float DeltaTimeSeconds);
+		AUDIOMIXER_API void LowResoultionUpdate(float DeltaTimeSeconds);
 
 		// add (and take ownership of) a new clock
 		// safe to call from AudioThread (uses critical section)
-		FQuartzClockProxy GetOrCreateClock(const FName& InClockName, const FQuartzClockSettings& InClockSettings, bool bOverrideTickRateIfClockExists = false);
-		FQuartzClockProxy GetClock(const FName& InClockName);
+		AUDIOMIXER_API FQuartzClockProxy GetOrCreateClock(const FName& InClockName, const FQuartzClockSettings& InClockSettings, bool bOverrideTickRateIfClockExists = false);
+		AUDIOMIXER_API FQuartzClockProxy GetClock(const FName& InClockName);
 
 		// returns true if a clock with the given name already exists.
-		bool DoesClockExist(const FName& InClockName);
+		AUDIOMIXER_API bool DoesClockExist(const FName& InClockName);
 
 		// returns true if the name is running
-		bool IsClockRunning(const FName& InClockName);
+		AUDIOMIXER_API bool IsClockRunning(const FName& InClockName);
 
 		// Returns the duration in seconds of the given Quantization Type, or -1 if the Clock is invalid or nonexistent
-		float GetDurationOfQuantizationTypeInSeconds(const FName& InClockName, const EQuartzCommandQuantization& QuantizationType, float Multiplier);
+		AUDIOMIXER_API float GetDurationOfQuantizationTypeInSeconds(const FName& InClockName, const EQuartzCommandQuantization& QuantizationType, float Multiplier);
 
 		// Returns the current location of the clock in the transport
-		FQuartzTransportTimeStamp GetCurrentTimestamp(const FName& InClockName);
+		AUDIOMIXER_API FQuartzTransportTimeStamp GetCurrentTimestamp(const FName& InClockName);
 
 		// Returns the amount of time, in seconds, the clock has been running. Caution: due to latency, this will not be perfectly accurate
-		float GetEstimatedRunTime(const FName& InClockName);
+		AUDIOMIXER_API float GetEstimatedRunTime(const FName& InClockName);
 
 		// remove existing clock
 		// safe to call from AudioThread (uses Audio Render Thread Command)
-		void RemoveClock(const FName& InName);
+		AUDIOMIXER_API void RemoveClock(const FName& InName);
 
 		// get Tick rate for clock
 		// safe to call from AudioThread (uses critical section)
-		FQuartzClockTickRate GetTickRateForClock(const FName& InName);
+		AUDIOMIXER_API FQuartzClockTickRate GetTickRateForClock(const FName& InName);
 
-		void SetTickRateForClock(const FQuartzClockTickRate& InNewTickRate, const FName& InName);
+		AUDIOMIXER_API void SetTickRateForClock(const FQuartzClockTickRate& InNewTickRate, const FName& InName);
 
 		// start the given clock
 		// safe to call from AudioThread (uses Audio Render Thread command)
-		void ResumeClock(const FName& InName, int32 NumFramesToDelayStart = 0);
+		AUDIOMIXER_API void ResumeClock(const FName& InName, int32 NumFramesToDelayStart = 0);
 
 		// stop the given clock
 		// safe to call from AudioThread (uses Audio Render Thread command)
-		void StopClock(const FName& InName, bool CancelPendingEvents);
+		AUDIOMIXER_API void StopClock(const FName& InName, bool CancelPendingEvents);
 
 		// stop the given clock
 		// safe to call from AudioThread (uses Audio Render Thread command)
-		void PauseClock(const FName& InName);
+		AUDIOMIXER_API void PauseClock(const FName& InName);
 
 		// shutdown all clocks that don't ignore Flush() (i.e. level change)
-		void Flush();
+		AUDIOMIXER_API void Flush();
 
 		// stop all clocks and cancel all pending events
-		void Shutdown();
+		AUDIOMIXER_API void Shutdown();
 
 		// add a new command to a given clock
 		// safe to call from AudioThread (uses Audio Render Thread command)
-		FQuartzQuantizedCommandHandle AddCommandToClock(FQuartzQuantizedCommandInitInfo& InQuantizationCommandInitInfo);
+		AUDIOMIXER_API FQuartzQuantizedCommandHandle AddCommandToClock(FQuartzQuantizedCommandInitInfo& InQuantizationCommandInitInfo);
 
 		// subscribe to a specific time division on a clock
 		// TODO: update the metronome subscription functions to take an FQuartzGameThreadSubscriber instead of the Command queue ptr
 		// (to support metronome event offset)
-		void SubscribeToTimeDivision(FName InClockName, MetronomeCommandQueuePtr InListenerQueue, EQuartzCommandQuantization InQuantizationBoundary);
+		AUDIOMIXER_API void SubscribeToTimeDivision(FName InClockName, MetronomeCommandQueuePtr InListenerQueue, EQuartzCommandQuantization InQuantizationBoundary);
 
 		// subscribe to all time divisions on a clock
-		void SubscribeToAllTimeDivisions(FName InClockName, MetronomeCommandQueuePtr InListenerQueue);
+		AUDIOMIXER_API void SubscribeToAllTimeDivisions(FName InClockName, MetronomeCommandQueuePtr InListenerQueue);
 
 		// un-subscribe from a specific time division on a clock
-		void UnsubscribeFromTimeDivision(FName InClockName, MetronomeCommandQueuePtr InListenerQueue, EQuartzCommandQuantization InQuantizationBoundary);
+		AUDIOMIXER_API void UnsubscribeFromTimeDivision(FName InClockName, MetronomeCommandQueuePtr InListenerQueue, EQuartzCommandQuantization InQuantizationBoundary);
 
 		// un-subscribe from all time divisions on a specific clock
-		void UnsubscribeFromAllTimeDivisions(FName InClockName, MetronomeCommandQueuePtr InListenerQueue);
+		AUDIOMIXER_API void UnsubscribeFromAllTimeDivisions(FName InClockName, MetronomeCommandQueuePtr InListenerQueue);
 
 		// cancel a queued command on a clock (i.e. cancel a PlayQuantized command if the sound is stopped before it is played)
-		bool CancelCommandOnClock(FName InOwningClockName, TSharedPtr<IQuartzQuantizedCommand> InCommandPtr);
+		AUDIOMIXER_API bool CancelCommandOnClock(FName InOwningClockName, TSharedPtr<IQuartzQuantizedCommand> InCommandPtr);
 
-		bool HasClockBeenTickedThisUpdate(FName InClockName);
+		AUDIOMIXER_API bool HasClockBeenTickedThisUpdate(FName InClockName);
 
 		int32 GetLastUpdateSizeInFrames() const { return LastUpdateSizeInFrames; }
 
 		// get access to the owning FMixerDevice
-		FMixerDevice* GetMixerDevice() const;
+		AUDIOMIXER_API FMixerDevice* GetMixerDevice() const;
 
 	private:
 		// updates all active clocks
-		void TickClocks(int32 NumFramesToTick);
+		AUDIOMIXER_API void TickClocks(int32 NumFramesToTick);
 
 		// find clock with a given key
-		TSharedPtr<FQuartzClock> FindClock(const FName& InName);
+		AUDIOMIXER_API TSharedPtr<FQuartzClock> FindClock(const FName& InName);
 
 		// pointer to owning FMixerDevice
 		FMixerDevice* MixerDevice;
@@ -134,7 +134,7 @@ namespace Audio
 	};
 
 	// data that the UQuartzSubsystem needs to persist on the AudioDevice across UWorld shutdown/startup
-	struct AUDIOMIXER_API FPersistentQuartzSubsystemData
+	struct FPersistentQuartzSubsystemData
 	{
 		// internal clock manager for game-thread-ticked clocks
 		FQuartzClockManager SubsystemClockManager;
diff --git a/Engine/Source/Runtime/AudioMixer/Public/Quartz/AudioMixerQuantizedCommands.h b/Engine/Source/Runtime/AudioMixer/Public/Quartz/AudioMixerQuantizedCommands.h
--- a/Engine/Source/Runtime/AudioMixer/Public/Quartz/AudioMixerQuantizedCommands.h
+++ b/Engine/Source/Runtime/AudioMixer/Public/Quartz/AudioMixerQuantizedCommands.h
@@ -6,26 +6,26 @@
 namespace Audio
 {
 	// QuartzQuantizedCommand that plays a sound on a sample-accurate boundary
-	class AUDIOMIXER_API FQuantizedPlayCommand : public IQuartzQuantizedCommand
+	class FQuantizedPlayCommand : public IQuartzQuantizedCommand
 	{
 	public:
 		// ctor
-		FQuantizedPlayCommand();
+		AUDIOMIXER_API FQuantizedPlayCommand();
 
 		// dtor
 		~FQuantizedPlayCommand() {}
 
-		virtual TSharedPtr<IQuartzQuantizedCommand> GetDeepCopyOfDerivedObject() const override;
+		AUDIOMIXER_API virtual TSharedPtr<IQuartzQuantizedCommand> GetDeepCopyOfDerivedObject() const override;
 
-		virtual void OnQueuedCustom(const FQuartzQuantizedCommandInitInfo& InCommandInitInfo) override;
+		AUDIOMIXER_API virtual void OnQueuedCustom(const FQuartzQuantizedCommandInitInfo& InCommandInitInfo) override;
 
-		virtual void OnFinalCallbackCustom(int32 InNumFramesLeft) override;
+		AUDIOMIXER_API virtual void OnFinalCallbackCustom(int32 InNumFramesLeft) override;
 
-		virtual void CancelCustom() override;
+		AUDIOMIXER_API virtual void CancelCustom() override;
 
 		virtual bool RequiresAudioDevice() const override { return true; }
 
-		virtual FName GetCommandName() const override;
+		AUDIOMIXER_API virtual FName GetCommandName() const override;
 		
 		// for your implementation, a new EQuartzCommandType needs to be defined in QuartzQuantizationUtilities.h
 		virtual EQuartzCommandType GetCommandType() const { return EQuartzCommandType::PlaySound; };
@@ -40,23 +40,23 @@ namespace Audio
 	}; // class FQuantizedPlayCommand 
 
 
-	class AUDIOMIXER_API FQuantizedQueueCommand : public IQuartzQuantizedCommand
+	class FQuantizedQueueCommand : public IQuartzQuantizedCommand
 	{
 	public:
-		virtual TSharedPtr<IQuartzQuantizedCommand> GetDeepCopyOfDerivedObject() const override;
+		AUDIOMIXER_API virtual TSharedPtr<IQuartzQuantizedCommand> GetDeepCopyOfDerivedObject() const override;
 
-		virtual void OnQueuedCustom(const FQuartzQuantizedCommandInitInfo& InCommandInitInfo) override;
+		AUDIOMIXER_API virtual void OnQueuedCustom(const FQuartzQuantizedCommandInitInfo& InCommandInitInfo) override;
 
-		virtual int32 OverrideFramesUntilExec(int32 NumFramesUntilExec) override;
+		AUDIOMIXER_API virtual int32 OverrideFramesUntilExec(int32 NumFramesUntilExec) override;
 
-		virtual void OnFinalCallbackCustom(int32 InNumFramesLeft) override;
+		AUDIOMIXER_API virtual void OnFinalCallbackCustom(int32 InNumFramesLeft) override;
 
 		virtual bool RequiresAudioDevice() const override { return true; }
 
-		virtual FName GetCommandName() const override;
+		AUDIOMIXER_API virtual FName GetCommandName() const override;
 		virtual EQuartzCommandType GetCommandType() const { return EQuartzCommandType::QueueSoundToPlay; };
 
-		void SetQueueCommand(const FAudioComponentCommandInfo& InAudioCommandData);
+		AUDIOMIXER_API void SetQueueCommand(const FAudioComponentCommandInfo& InAudioCommandData);
 
 		FQuantizedQueueCommand() {}
 
@@ -68,7 +68,7 @@ namespace Audio
 	}; // class FQuantizedQueueCommand 
 	
 	// QuartzQuantizedCommand that changes the TickRate of a clock on a sample-accurate boundary (i.e. BPM changes)
-	class AUDIOMIXER_API FQuantizedTickRateChange : public IQuartzQuantizedCommand
+	class FQuantizedTickRateChange : public IQuartzQuantizedCommand
 	{
 	public:
 		void SetTickRate(const FQuartzClockTickRate& InTickRate)
@@ -76,15 +76,15 @@ namespace Audio
 			TickRate = InTickRate;
 		}
 
-		virtual TSharedPtr<IQuartzQuantizedCommand> GetDeepCopyOfDerivedObject() const override;
+		AUDIOMIXER_API virtual TSharedPtr<IQuartzQuantizedCommand> GetDeepCopyOfDerivedObject() const override;
 
-		virtual void OnQueuedCustom(const FQuartzQuantizedCommandInitInfo& InCommandInitInfo) override;
+		AUDIOMIXER_API virtual void OnQueuedCustom(const FQuartzQuantizedCommandInitInfo& InCommandInitInfo) override;
 
-		virtual void OnFinalCallbackCustom(int32 InNumFramesLeft) override;
+		AUDIOMIXER_API virtual void OnFinalCallbackCustom(int32 InNumFramesLeft) override;
 
 		virtual bool IsClockAltering() override { return true; }
 
-		virtual FName GetCommandName() const override;
+		AUDIOMIXER_API virtual FName GetCommandName() const override;
 		virtual EQuartzCommandType GetCommandType() const { return EQuartzCommandType::TickRateChange; };
 
 	private:
@@ -95,18 +95,18 @@ namespace Audio
 
 
 	// QuartzQuantizedCommand that resets the transport of a clock's metronome on a sample-accurate boundary
-	class AUDIOMIXER_API FQuantizedTransportReset : public IQuartzQuantizedCommand
+	class FQuantizedTransportReset : public IQuartzQuantizedCommand
 	{
 	public:
-		virtual TSharedPtr<IQuartzQuantizedCommand> GetDeepCopyOfDerivedObject() const override;
+		AUDIOMIXER_API virtual TSharedPtr<IQuartzQuantizedCommand> GetDeepCopyOfDerivedObject() const override;
 
-		virtual void OnQueuedCustom(const FQuartzQuantizedCommandInitInfo& InCommandInitInfo) override;
+		AUDIOMIXER_API virtual void OnQueuedCustom(const FQuartzQuantizedCommandInitInfo& InCommandInitInfo) override;
 
-		virtual void OnFinalCallbackCustom(int32 InNumFramesLeft) override;
+		AUDIOMIXER_API virtual void OnFinalCallbackCustom(int32 InNumFramesLeft) override;
 
 		virtual bool IsClockAltering() override { return false; }
 
-		virtual FName GetCommandName() const override;
+		AUDIOMIXER_API virtual FName GetCommandName() const override;
 		virtual EQuartzCommandType GetCommandType() const { return EQuartzCommandType::TransportReset; };
 
 	private:
@@ -116,18 +116,18 @@ namespace Audio
 
 
 	// QuartzQuantizedCommand that starts a second clock on a sample-accurate boundary
-	class AUDIOMIXER_API FQuantizedOtherClockStart : public IQuartzQuantizedCommand
+	class FQuantizedOtherClockStart : public IQuartzQuantizedCommand
 	{
 	public:
-		virtual TSharedPtr<IQuartzQuantizedCommand> GetDeepCopyOfDerivedObject() const override;
+		AUDIOMIXER_API virtual TSharedPtr<IQuartzQuantizedCommand> GetDeepCopyOfDerivedObject() const override;
 
-		virtual void OnQueuedCustom(const FQuartzQuantizedCommandInitInfo& InCommandInitInfo) override;
+		AUDIOMIXER_API virtual void OnQueuedCustom(const FQuartzQuantizedCommandInitInfo& InCommandInitInfo) override;
 
-		virtual void OnFinalCallbackCustom(int32 InNumFramesLeft) override;
+		AUDIOMIXER_API virtual void OnFinalCallbackCustom(int32 InNumFramesLeft) override;
 
 		virtual bool IsClockAltering() override { return true; }
 
-		virtual FName GetCommandName() const override;
+		AUDIOMIXER_API virtual FName GetCommandName() const override;
 		virtual EQuartzCommandType GetCommandType() const { return EQuartzCommandType::StartOtherClock; };
 
 	private:
@@ -138,26 +138,26 @@ namespace Audio
 
 
 	// QuartzQuantizedCommand that basically no-ops, so the game thread can get notified on a musical boundary
-	class AUDIOMIXER_API FQuantizedNotify : public IQuartzQuantizedCommand
+	class FQuantizedNotify : public IQuartzQuantizedCommand
 	{
 	public:
 		// ctor
-		FQuantizedNotify(float InMsOffset = 0.f);
+		AUDIOMIXER_API FQuantizedNotify(float InMsOffset = 0.f);
 
 		// dtor
 		virtual ~FQuantizedNotify() override = default;
 
-		virtual TSharedPtr<IQuartzQuantizedCommand> GetDeepCopyOfDerivedObject() const override;
+		AUDIOMIXER_API virtual TSharedPtr<IQuartzQuantizedCommand> GetDeepCopyOfDerivedObject() const override;
 
 		virtual bool RequiresAudioDevice() const override { return true; }
 
-		virtual FName GetCommandName() const override;
+		AUDIOMIXER_API virtual FName GetCommandName() const override;
 
 		virtual EQuartzCommandType GetCommandType() const override { return EQuartzCommandType::Notify; };
 
-		virtual void OnQueuedCustom(const FQuartzQuantizedCommandInitInfo& InCommandInitInfo) override;
+		AUDIOMIXER_API virtual void OnQueuedCustom(const FQuartzQuantizedCommandInitInfo& InCommandInitInfo) override;
 
-		virtual int32 OverrideFramesUntilExec(int32 NumFramesUntilExec) override;
+		AUDIOMIXER_API virtual int32 OverrideFramesUntilExec(int32 NumFramesUntilExec) override;
 
 	protected:
 		TSharedPtr<FQuartzClock> OwningClockPtr{ nullptr };
diff --git a/Engine/Source/Runtime/AudioMixer/Public/Quartz/QuartzSubsystem.h b/Engine/Source/Runtime/AudioMixer/Public/Quartz/QuartzSubsystem.h
--- a/Engine/Source/Runtime/AudioMixer/Public/Quartz/QuartzSubsystem.h
+++ b/Engine/Source/Runtime/AudioMixer/Public/Quartz/QuartzSubsystem.h
@@ -28,13 +28,13 @@ using MetronomeCommandQueuePtr = TSharedPtr<Audio::TQuartzShareableCommandQueue<
 
 
 
-struct AUDIOMIXER_API FQuartzTickableObjectsManager : public FQuartLatencyTracker
+struct FQuartzTickableObjectsManager : public FQuartLatencyTracker
 {
 public:
-	void Tick(float DeltaTime);
-	bool IsTickable() const;
-	void SubscribeToQuartzTick(FQuartzTickableObject* InObjectToTick);
-	void UnsubscribeFromQuartzTick(FQuartzTickableObject* InObjectToTick);
+	AUDIOMIXER_API void Tick(float DeltaTime);
+	AUDIOMIXER_API bool IsTickable() const;
+	AUDIOMIXER_API void SubscribeToQuartzTick(FQuartzTickableObject* InObjectToTick);
+	AUDIOMIXER_API void UnsubscribeFromQuartzTick(FQuartzTickableObject* InObjectToTick);
 
 private:
 	// list of objects needing to be ticked by Quartz
@@ -45,8 +45,8 @@ private:
 };
 
 
-UCLASS(DisplayName = "Quartz")
-class AUDIOMIXER_API UQuartzSubsystem : public UTickableWorldSubsystem
+UCLASS(DisplayName = "Quartz", MinimalAPI)
+class UQuartzSubsystem : public UTickableWorldSubsystem
 {
 	GENERATED_BODY()
 
@@ -56,131 +56,131 @@ public:
 	virtual ~UQuartzSubsystem() override = default;
 
 	//~ Begin UWorldSubsystem Interface
-	virtual void Initialize(FSubsystemCollectionBase& Collection) override;
-	virtual void Deinitialize() override;
-	virtual bool DoesSupportWorldType(EWorldType::Type WorldType) const override;
-	void virtual BeginDestroy() override;
+	AUDIOMIXER_API virtual void Initialize(FSubsystemCollectionBase& Collection) override;
+	AUDIOMIXER_API virtual void Deinitialize() override;
+	AUDIOMIXER_API virtual bool DoesSupportWorldType(EWorldType::Type WorldType) const override;
+	AUDIOMIXER_API void virtual BeginDestroy() override;
 	//~ End UWorldSubsystem Interface
 
 	//~ Begin FTickableGameObject Interface
-	virtual void Tick(float DeltaTime) override;
-	virtual bool IsTickable() const override;
-	virtual TStatId GetStatId() const override;
+	AUDIOMIXER_API virtual void Tick(float DeltaTime) override;
+	AUDIOMIXER_API virtual bool IsTickable() const override;
+	AUDIOMIXER_API virtual TStatId GetStatId() const override;
 	//~ End FTickableGameObject Interface
 
 	// these calls are forwarded to the internal FQuartzTickableObjectsManager
-	void SubscribeToQuartzTick(FQuartzTickableObject* InObjectToTick);
-	void UnsubscribeFromQuartzTick(FQuartzTickableObject* InObjectToTick);
+	AUDIOMIXER_API void SubscribeToQuartzTick(FQuartzTickableObject* InObjectToTick);
+	AUDIOMIXER_API void UnsubscribeFromQuartzTick(FQuartzTickableObject* InObjectToTick);
 
 	// get C++ handle (proxy) to a clock if it exists
-	Audio::FQuartzClockProxy GetProxyForClock(FName ClockName) const;
+	AUDIOMIXER_API Audio::FQuartzClockProxy GetProxyForClock(FName ClockName) const;
 
 	// allow an external clock (not ticked by the Audio Mixer or QuartzSubsystem) to be accessible via this subsystem
-	void AddProxyForExternalClock(const Audio::FQuartzClockProxy& InProxy);
+	AUDIOMIXER_API void AddProxyForExternalClock(const Audio::FQuartzClockProxy& InProxy);
 
 	// static methods
-	static UQuartzSubsystem* Get(const UWorld* const World);
+	static AUDIOMIXER_API UQuartzSubsystem* Get(const UWorld* const World);
 
 	// Helper functions for initializing quantized command initialization struct (to consolidate eyesore)
-	static Audio::FQuartzQuantizedRequestData CreateRequestDataForTickRateChange(UQuartzClockHandle* InClockHandle, const FOnQuartzCommandEventBP& InDelegate, const Audio::FQuartzClockTickRate& InNewTickRate, const FQuartzQuantizationBoundary& InQuantizationBoundary);
-	static Audio::FQuartzQuantizedRequestData CreateRequestDataForTransportReset(UQuartzClockHandle* InClockHandle, const FQuartzQuantizationBoundary& InQuantizationBoundary, const FOnQuartzCommandEventBP& InDelegate);
-	static Audio::FQuartzQuantizedRequestData CreateRequestDataForStartOtherClock(UQuartzClockHandle* InClockHandle, FName InClockToStart, const FQuartzQuantizationBoundary& InQuantizationBoundary, const FOnQuartzCommandEventBP& InDelegate);
-	static Audio::FQuartzQuantizedRequestData CreateRequestDataForSchedulePlaySound(UQuartzClockHandle* InClockHandle, const FOnQuartzCommandEventBP& InDelegate, const FQuartzQuantizationBoundary& InQuantizationBoundary);
-	static Audio::FQuartzQuantizedRequestData CreateRequestDataForQuantizedNotify(UQuartzClockHandle* InClockHandle, const FQuartzQuantizationBoundary& InQuantizationBoundary, const FOnQuartzCommandEventBP& InDelegate, float InMsOffset = 0.f);
+	static AUDIOMIXER_API Audio::FQuartzQuantizedRequestData CreateRequestDataForTickRateChange(UQuartzClockHandle* InClockHandle, const FOnQuartzCommandEventBP& InDelegate, const Audio::FQuartzClockTickRate& InNewTickRate, const FQuartzQuantizationBoundary& InQuantizationBoundary);
+	static AUDIOMIXER_API Audio::FQuartzQuantizedRequestData CreateRequestDataForTransportReset(UQuartzClockHandle* InClockHandle, const FQuartzQuantizationBoundary& InQuantizationBoundary, const FOnQuartzCommandEventBP& InDelegate);
+	static AUDIOMIXER_API Audio::FQuartzQuantizedRequestData CreateRequestDataForStartOtherClock(UQuartzClockHandle* InClockHandle, FName InClockToStart, const FQuartzQuantizationBoundary& InQuantizationBoundary, const FOnQuartzCommandEventBP& InDelegate);
+	static AUDIOMIXER_API Audio::FQuartzQuantizedRequestData CreateRequestDataForSchedulePlaySound(UQuartzClockHandle* InClockHandle, const FOnQuartzCommandEventBP& InDelegate, const FQuartzQuantizationBoundary& InQuantizationBoundary);
+	static AUDIOMIXER_API Audio::FQuartzQuantizedRequestData CreateRequestDataForQuantizedNotify(UQuartzClockHandle* InClockHandle, const FQuartzQuantizationBoundary& InQuantizationBoundary, const FOnQuartzCommandEventBP& InDelegate, float InMsOffset = 0.f);
 
 	// DEPRECATED HELPERS: non-static versions of the above CreateDataFor...() functions
 	UE_DEPRECATED(5.1, "Use the static (CreateRequestDataFor) version of this function instead")
-	Audio::FQuartzQuantizedRequestData CreateDataForTickRateChange(UQuartzClockHandle* InClockHandle, const FOnQuartzCommandEventBP& InDelegate, const Audio::FQuartzClockTickRate& InNewTickRate, const FQuartzQuantizationBoundary& InQuantizationBoundary);
+	AUDIOMIXER_API Audio::FQuartzQuantizedRequestData CreateDataForTickRateChange(UQuartzClockHandle* InClockHandle, const FOnQuartzCommandEventBP& InDelegate, const Audio::FQuartzClockTickRate& InNewTickRate, const FQuartzQuantizationBoundary& InQuantizationBoundary);
 
 	UE_DEPRECATED(5.1, "Use the static (CreateRequestDataFor) version of this function instead")
-	Audio::FQuartzQuantizedRequestData CreateDataForTransportReset(UQuartzClockHandle* InClockHandle, const FQuartzQuantizationBoundary& InQuantizationBoundary, const FOnQuartzCommandEventBP& InDelegate);
+	AUDIOMIXER_API Audio::FQuartzQuantizedRequestData CreateDataForTransportReset(UQuartzClockHandle* InClockHandle, const FQuartzQuantizationBoundary& InQuantizationBoundary, const FOnQuartzCommandEventBP& InDelegate);
 	
 	UE_DEPRECATED(5.1, "Use the static (CreateRequestDataFor) version of this function instead")
-	Audio::FQuartzQuantizedRequestData CreateDataForStartOtherClock(UQuartzClockHandle* InClockHandle, FName InClockToStart, const FQuartzQuantizationBoundary& InQuantizationBoundary, const FOnQuartzCommandEventBP& InDelegate);
+	AUDIOMIXER_API Audio::FQuartzQuantizedRequestData CreateDataForStartOtherClock(UQuartzClockHandle* InClockHandle, FName InClockToStart, const FQuartzQuantizationBoundary& InQuantizationBoundary, const FOnQuartzCommandEventBP& InDelegate);
 
 	UE_DEPRECATED(5.1, "Use the static (CreateRequestDataFor) version of this function instead")
-	Audio::FQuartzQuantizedRequestData CreateDataDataForSchedulePlaySound(UQuartzClockHandle* InClockHandle, const FOnQuartzCommandEventBP& InDelegate, const FQuartzQuantizationBoundary& InQuantizationBoundary);
+	AUDIOMIXER_API Audio::FQuartzQuantizedRequestData CreateDataDataForSchedulePlaySound(UQuartzClockHandle* InClockHandle, const FOnQuartzCommandEventBP& InDelegate, const FQuartzQuantizationBoundary& InQuantizationBoundary);
 
 	UFUNCTION(BlueprintCallable, Category = "Quartz Clock Handle")
-	bool IsQuartzEnabled();
+	AUDIOMIXER_API bool IsQuartzEnabled();
 
 	// Clock Creation
 	// create a new clock (or return handle if clock already exists)
 	UFUNCTION(BlueprintCallable, Category = "Quartz Clock Handle", meta = (WorldContext = "WorldContextObject", AdvancedDisplay = "bUseAudioEngineClockManager"))
-	UQuartzClockHandle* CreateNewClock(const UObject* WorldContextObject, FName ClockName, FQuartzClockSettings InSettings, bool bOverrideSettingsIfClockExists = false, bool bUseAudioEngineClockManager = true);
+	AUDIOMIXER_API UQuartzClockHandle* CreateNewClock(const UObject* WorldContextObject, FName ClockName, FQuartzClockSettings InSettings, bool bOverrideSettingsIfClockExists = false, bool bUseAudioEngineClockManager = true);
 
 	// delete an existing clock given its name
 	UFUNCTION(BlueprintCallable, Category = "Quartz Clock Handle", meta = (WorldContext = "WorldContextObject"))
-	void DeleteClockByName(const UObject* WorldContextObject, FName ClockName);
+	AUDIOMIXER_API void DeleteClockByName(const UObject* WorldContextObject, FName ClockName);
 
 	// delete an existing clock given its clock handle
 	UFUNCTION(BlueprintCallable, Category = "Quartz Clock Handle", meta = (WorldContext = "WorldContextObject"))
-	void DeleteClockByHandle(const UObject* WorldContextObject, UPARAM(ref) UQuartzClockHandle*& InClockHandle);
+	AUDIOMIXER_API void DeleteClockByHandle(const UObject* WorldContextObject, UPARAM(ref) UQuartzClockHandle*& InClockHandle);
 
 	// get handle for existing clock
 	UFUNCTION(BlueprintCallable, Category = "Quartz Clock Handle", meta = (WorldContext = "WorldContextObject"))
-	UQuartzClockHandle* GetHandleForClock(const UObject* WorldContextObject, FName ClockName);
+	AUDIOMIXER_API UQuartzClockHandle* GetHandleForClock(const UObject* WorldContextObject, FName ClockName);
 
 	// returns true if the clock exists
 	UFUNCTION(BlueprintCallable, Category = "Quartz Clock Handle", meta = (WorldContext = "WorldContextObject"))
-	bool DoesClockExist(const UObject* WorldContextObject, FName ClockName);
+	AUDIOMIXER_API bool DoesClockExist(const UObject* WorldContextObject, FName ClockName);
 
 	// returns true if the clock is running
 	UFUNCTION(BlueprintCallable, Category = "Quartz Clock Handle", meta = (WorldContext = "WorldContextObject"))
-	bool IsClockRunning(const UObject* WorldContextObject, FName ClockName);
+	AUDIOMIXER_API bool IsClockRunning(const UObject* WorldContextObject, FName ClockName);
 
 	// Returns the duration in seconds of the given Quantization Type
 	UFUNCTION(BlueprintCallable, Category = "Quartz Clock Handle", meta = (WorldContext = "WorldContextObject"))
-	float GetDurationOfQuantizationTypeInSeconds(const UObject* WorldContextObject, FName ClockName, const EQuartzCommandQuantization& QuantizationType, float Multiplier = 1.0f);
+	AUDIOMIXER_API float GetDurationOfQuantizationTypeInSeconds(const UObject* WorldContextObject, FName ClockName, const EQuartzCommandQuantization& QuantizationType, float Multiplier = 1.0f);
 
 	// Retrieves a timestamp for the clock
 	UFUNCTION(BlueprintCallable, Category = "Quartz Clock Handle", meta = (WorldContext = "WorldContextObject"))
-	FQuartzTransportTimeStamp GetCurrentClockTimestamp(const UObject* WorldContextObject, const FName& InClockName);
+	AUDIOMIXER_API FQuartzTransportTimeStamp GetCurrentClockTimestamp(const UObject* WorldContextObject, const FName& InClockName);
 
 	// Returns the amount of time, in seconds, the clock has been running. Caution: due to latency, this will not be perfectly accurate
 	UFUNCTION(BlueprintCallable, Category = "Quartz Clock Handle", meta = (WorldContext = "WorldContextObject"))
-	float GetEstimatedClockRunTime(const UObject* WorldContextObject, const FName& InClockName);
+	AUDIOMIXER_API float GetEstimatedClockRunTime(const UObject* WorldContextObject, const FName& InClockName);
 
 	// latency data (Game thread -> Audio Render Thread)
 	UFUNCTION(BlueprintCallable, Category = "Quartz Subsystem", meta = (WorldContext = "WorldContextObject"))
-	float GetGameThreadToAudioRenderThreadAverageLatency(const UObject* WorldContextObject);
+	AUDIOMIXER_API float GetGameThreadToAudioRenderThreadAverageLatency(const UObject* WorldContextObject);
 
 	UFUNCTION(BlueprintCallable, Category = "Quartz Subsystem", meta = (WorldContext = "WorldContextObject"))
-	float GetGameThreadToAudioRenderThreadMinLatency(const UObject* WorldContextObject);
+	AUDIOMIXER_API float GetGameThreadToAudioRenderThreadMinLatency(const UObject* WorldContextObject);
 
 	UFUNCTION(BlueprintCallable, Category = "Quartz Subsystem", meta = (WorldContext = "WorldContextObject"))
-	float GetGameThreadToAudioRenderThreadMaxLatency(const UObject* WorldContextObject);
+	AUDIOMIXER_API float GetGameThreadToAudioRenderThreadMaxLatency(const UObject* WorldContextObject);
 
 	// latency data (Audio Render Thread -> Game thread)
 	UFUNCTION(BlueprintCallable, Category = "Quartz Subsystem")
-	float GetAudioRenderThreadToGameThreadAverageLatency();
+	AUDIOMIXER_API float GetAudioRenderThreadToGameThreadAverageLatency();
 
 	UFUNCTION(BlueprintCallable, Category = "Quartz Subsystem")
-	float GetAudioRenderThreadToGameThreadMinLatency();
+	AUDIOMIXER_API float GetAudioRenderThreadToGameThreadMinLatency();
 
 	UFUNCTION(BlueprintCallable, Category = "Quartz Subsystem")
-	float GetAudioRenderThreadToGameThreadMaxLatency();
+	AUDIOMIXER_API float GetAudioRenderThreadToGameThreadMaxLatency();
 
 	// latency data (Round trip)
 	UFUNCTION(BlueprintCallable, Category = "Quartz Subsystem", meta = (WorldContext = "WorldContextObject"))
-	float GetRoundTripAverageLatency(const UObject* WorldContextObject);
+	AUDIOMIXER_API float GetRoundTripAverageLatency(const UObject* WorldContextObject);
 
 	UFUNCTION(BlueprintCallable, Category = "Quartz Subsystem", meta = (WorldContext = "WorldContextObject"))
-	float GetRoundTripMinLatency(const UObject* WorldContextObject);
+	AUDIOMIXER_API float GetRoundTripMinLatency(const UObject* WorldContextObject);
 
 	UFUNCTION(BlueprintCallable, Category = "Quartz Subsystem", meta = (WorldContext = "WorldContextObject"))
-	float GetRoundTripMaxLatency(const UObject* WorldContextObject);
+	AUDIOMIXER_API float GetRoundTripMaxLatency(const UObject* WorldContextObject);
 
 	UE_DEPRECATED(5.1, "Obtain and use a UQuartzClockHandle / FQuartzClockProxy instead")
-	void AddCommandToClock(const UObject* WorldContextObject, Audio::FQuartzQuantizedCommandInitInfo& InQuantizationCommandInitInfo, FName ClockName);
+	AUDIOMIXER_API void AddCommandToClock(const UObject* WorldContextObject, Audio::FQuartzQuantizedCommandInitInfo& InQuantizationCommandInitInfo, FName ClockName);
 
 	// sharable to allow non-UObjects to un-subscribe if the Subsystem is going to outlive them
-	TWeakPtr<FQuartzTickableObjectsManager> GetTickableObjectManager() const;
+	AUDIOMIXER_API TWeakPtr<FQuartzTickableObjectsManager> GetTickableObjectManager() const;
 
 private:
 	// deletes proxies to clocks that no longer exists
-	void PruneStaleProxies();
-	static void PruneStaleProxiesInternal(TArray<Audio::FQuartzClockProxy>& ContainerToPrune);
+	AUDIOMIXER_API void PruneStaleProxies();
+	static AUDIOMIXER_API void PruneStaleProxiesInternal(TArray<Audio::FQuartzClockProxy>& ContainerToPrune);
 
 
 	// sharable tickable object manager to allow for non-UObject subscription / un-subscription
@@ -190,8 +190,8 @@ private:
 	TSharedPtr<Audio::FPersistentQuartzSubsystemData> ClockManagerDataPtr { nullptr };
 
 	// helpers
-	Audio::FQuartzClockProxy* FindProxyByName(const FName& ClockName);
-	Audio::FQuartzClockProxy const* FindProxyByName(const FName& ClockName) const;
-	Audio::FQuartzClockManager* GetClockManager(const UObject* WorldContextObject, bool bUseAudioEngineClockManager = true);
+	AUDIOMIXER_API Audio::FQuartzClockProxy* FindProxyByName(const FName& ClockName);
+	AUDIOMIXER_API Audio::FQuartzClockProxy const* FindProxyByName(const FName& ClockName) const;
+	AUDIOMIXER_API Audio::FQuartzClockManager* GetClockManager(const UObject* WorldContextObject, bool bUseAudioEngineClockManager = true);
 
-}; // class UQuartzGameSubsystem
\ No newline at end of file
+}; // class UQuartzGameSubsystem
diff --git a/Engine/Source/Runtime/AudioMixer/Public/SoundWaveDecoder.h b/Engine/Source/Runtime/AudioMixer/Public/SoundWaveDecoder.h
--- a/Engine/Source/Runtime/AudioMixer/Public/SoundWaveDecoder.h
+++ b/Engine/Source/Runtime/AudioMixer/Public/SoundWaveDecoder.h
@@ -211,14 +211,14 @@ namespace Audio
 
 	typedef TSharedPtr<FDecodingSoundSource> FDecodingSoundSourcePtr;
 	
-	class AUDIOMIXER_API FSoundSourceDecoder : public FGCObject
+	class FSoundSourceDecoder : public FGCObject
 	{
 	public:
-		FSoundSourceDecoder();
-		virtual ~FSoundSourceDecoder();
+		AUDIOMIXER_API FSoundSourceDecoder();
+		AUDIOMIXER_API virtual ~FSoundSourceDecoder();
 
 		//~ Begin FGCObject
-		virtual void AddReferencedObjects(FReferenceCollector & Collector) override;
+		AUDIOMIXER_API virtual void AddReferencedObjects(FReferenceCollector & Collector) override;
 		virtual FString GetReferencerName() const override
 		{
 			return TEXT("Audio::FSoundSourceDecoder");
@@ -227,45 +227,45 @@ namespace Audio
 
 		// Initialize the source decoder at the given output sample rate
 		// Sources will automatically sample rate convert to match this output
-		void Init(FAudioDevice* InAudioDevice, int32 SampleRate);
+		AUDIOMIXER_API void Init(FAudioDevice* InAudioDevice, int32 SampleRate);
 
 		// Creates a new decoding sound source handle
-		FDecodingSoundSourceHandle CreateSourceHandle(USoundWave* InSoundWave);
+		AUDIOMIXER_API FDecodingSoundSourceHandle CreateSourceHandle(USoundWave* InSoundWave);
 
 		// Called from the audio thread
-		void Update();
+		AUDIOMIXER_API void Update();
 
 		// Called from the audio render thread
-		void UpdateRenderThread();
+		AUDIOMIXER_API void UpdateRenderThread();
 
 		// Initialize a decoding instance of this sound wave object. Call only from game thread.
-		bool InitDecodingSource(const FSourceDecodeInit& InitData);
+		AUDIOMIXER_API bool InitDecodingSource(const FSourceDecodeInit& InitData);
 
 		// Removes the decoding source from the decoder
-		void RemoveDecodingSource(const FDecodingSoundSourceHandle& Handle);
+		AUDIOMIXER_API void RemoveDecodingSource(const FDecodingSoundSourceHandle& Handle);
 
 		// Resets internal state of decoder
-		void Reset();
+		AUDIOMIXER_API void Reset();
 
 		// Sets the source pitch scale
-		void SetSourcePitchScale(const FDecodingSoundSourceHandle& Handle, float InPitchScale);
+		AUDIOMIXER_API void SetSourcePitchScale(const FDecodingSoundSourceHandle& Handle, float InPitchScale);
 
 		// Sets the source volume scale
-		void SetSourceVolumeScale(const FDecodingSoundSourceHandle& Handle, float InVolumeScale);
+		AUDIOMIXER_API void SetSourceVolumeScale(const FDecodingSoundSourceHandle& Handle, float InVolumeScale);
 
 		// Get a decoded buffer for the given decoding sound wave handle. Call only from audio render thread or audio render thread task.
-		bool GetSourceBuffer(const FDecodingSoundSourceHandle& InHandle, const int32 NumOutFrames, const int32 NumOutChannels, FAlignedFloatBuffer& OutAudioBuffer);
+		AUDIOMIXER_API bool GetSourceBuffer(const FDecodingSoundSourceHandle& InHandle, const int32 NumOutFrames, const int32 NumOutChannels, FAlignedFloatBuffer& OutAudioBuffer);
 
 		// Queries if the decoding source is finished
-		bool IsFinished(const FDecodingSoundSourceHandle& InHandle) const;
+		AUDIOMIXER_API bool IsFinished(const FDecodingSoundSourceHandle& InHandle) const;
 
-		bool IsInitialized(const FDecodingSoundSourceHandle& InHandle) const;
+		AUDIOMIXER_API bool IsInitialized(const FDecodingSoundSourceHandle& InHandle) const;
 
 	private:
 		// Sends a command to the audio render thread from audio thread
-		void EnqueueDecoderCommand(TFunction<void()> Command);
-		void PumpDecoderCommandQueue();
-		bool InitDecodingSourceInternal(const FSourceDecodeInit& InitData);
+		AUDIOMIXER_API void EnqueueDecoderCommand(TFunction<void()> Command);
+		AUDIOMIXER_API void PumpDecoderCommandQueue();
+		AUDIOMIXER_API bool InitDecodingSourceInternal(const FSourceDecodeInit& InitData);
 
 		int32 AudioThreadId;
 		FAudioDevice* AudioDevice;
diff --git a/Engine/Source/Runtime/AudioMixerCore/Public/AudioMixer.h b/Engine/Source/Runtime/AudioMixerCore/Public/AudioMixer.h
--- a/Engine/Source/Runtime/AudioMixerCore/Public/AudioMixer.h
+++ b/Engine/Source/Runtime/AudioMixerCore/Public/AudioMixer.h
@@ -225,7 +225,7 @@ namespace Audio
 	};
 
 	// Interface for Caching Device Info.
-	class AUDIOMIXERCORE_API IAudioPlatformDeviceInfoCache
+	class IAudioPlatformDeviceInfoCache
 	{
 	public:
 		// Pure Interface. 
@@ -331,7 +331,7 @@ namespace Audio
 	};
 
 	/** Struct used to store render time analysis data. */
-	struct AUDIOMIXERCORE_API FAudioRenderTimeAnalysis
+	struct FAudioRenderTimeAnalysis
 	{
 		double AvgRenderTime;
 		double MaxRenderTime;
@@ -342,13 +342,13 @@ namespace Audio
 		uint64 RenderTimeCount;
 		int32 RenderInstanceId;
 
-		FAudioRenderTimeAnalysis();
-		void Start();
-		void End();
+		AUDIOMIXERCORE_API FAudioRenderTimeAnalysis();
+		AUDIOMIXERCORE_API void Start();
+		AUDIOMIXERCORE_API void End();
 	};
 
 	/** Class which wraps an output float buffer and handles conversion to device stream formats. */
-	class AUDIOMIXERCORE_API FOutputBuffer
+	class FOutputBuffer
 	{
 	public:
 		FOutputBuffer()
@@ -359,16 +359,16 @@ namespace Audio
 		~FOutputBuffer() = default;
  
 		/** Initialize the buffer with the given samples and output format. */
-		void Init(IAudioMixer* InAudioMixer, const int32 InNumSamples, const int32 InNumBuffers, const EAudioMixerStreamDataFormat::Type InDataFormat);
+		AUDIOMIXERCORE_API void Init(IAudioMixer* InAudioMixer, const int32 InNumSamples, const int32 InNumBuffers, const EAudioMixerStreamDataFormat::Type InDataFormat);
 
 		/** Gets the next mixed buffer from the audio mixer. Returns false if our buffer is already full. */
-		bool MixNextBuffer();
+		AUDIOMIXERCORE_API bool MixNextBuffer();
 
 		/** Gets the buffer data ptrs. Returns a TArrayView for the full buffer size requested, but in the case of an underrun, OutBytesPopped will be less that the size of the returned TArrayView. */
-		TArrayView<const uint8> PopBufferData(int32& OutBytesPopped) const;
+		AUDIOMIXERCORE_API TArrayView<const uint8> PopBufferData(int32& OutBytesPopped) const;
 
 		/** Gets the number of frames of the buffer. */
-		int32 GetNumSamples() const;
+		AUDIOMIXERCORE_API int32 GetNumSamples() const;
 
 		/** Returns the format of the buffer. */
 		EAudioMixerStreamDataFormat::Type GetFormat() const { return DataFormat; }
@@ -390,12 +390,12 @@ namespace Audio
 		FAlignedByteBuffer FormattedBuffer;
  		EAudioMixerStreamDataFormat::Type DataFormat;
 
-		static size_t GetSizeForDataFormat(EAudioMixerStreamDataFormat::Type InDataFormat);
+		static AUDIOMIXERCORE_API size_t GetSizeForDataFormat(EAudioMixerStreamDataFormat::Type InDataFormat);
 		int32 CallCounterMixNextBuffer{ 0 };
 	};
 
 	/** Abstract interface for receiving audio device changed notifications */
-	class AUDIOMIXERCORE_API IAudioMixerDeviceChangedListener
+	class IAudioMixerDeviceChangedListener
 	{
 	public:
 		virtual ~IAudioMixerDeviceChangedListener() = default;
@@ -666,7 +666,7 @@ namespace Audio
 		AUDIOMIXERCORE_API void ApplyPrimaryAttenuation(TArrayView<const uint8>& InOutPoppedAudio);
 
 		template<typename BufferType>
-		AUDIOMIXERCORE_API void ApplyAttenuationInternal(TArrayView<BufferType>& InOutBuffer);
+		void ApplyAttenuationInternal(TArrayView<BufferType>& InOutBuffer);
 
 		/** When called, spins up a thread to start consuming output when no audio device is available. */
 		AUDIOMIXERCORE_API void StartRunningNullDevice();
diff --git a/Engine/Source/Runtime/AudioMixerCore/Public/AudioMixerNullDevice.h b/Engine/Source/Runtime/AudioMixerCore/Public/AudioMixerNullDevice.h
--- a/Engine/Source/Runtime/AudioMixerCore/Public/AudioMixerNullDevice.h
+++ b/Engine/Source/Runtime/AudioMixerCore/Public/AudioMixerNullDevice.h
@@ -21,7 +21,7 @@ namespace Audio
 	 * This class, when started, spawns a new high priority thread that exists to query an FAudioMixerPlatformInterface
 	 * and immediately throw out whatever buffers it receives.
 	 */
-	class AUDIOMIXERCORE_API FMixerNullCallback : protected FRunnable
+	class FMixerNullCallback : protected FRunnable
 	{
 	public:
 
@@ -29,7 +29,7 @@ namespace Audio
 		 * Constructing the FMixerNullCallback immediately begins calling
 		 * InCallback every BufferDuration seconds.
 		 */
-		FMixerNullCallback(float BufferDuration, TFunction<void()> InCallback, EThreadPriority ThreadPriority = TPri_TimeCritical, bool bStartedPaused = false);
+		AUDIOMIXERCORE_API FMixerNullCallback(float BufferDuration, TFunction<void()> InCallback, EThreadPriority ThreadPriority = TPri_TimeCritical, bool bStartedPaused = false);
 
 		/**
 		 * The destructor waits on Callback to be completed before stopping the thread.
@@ -37,14 +37,14 @@ namespace Audio
 		virtual ~FMixerNullCallback() = default;
 
 		// FRunnable override:
-		virtual uint32 Run() override;
-		virtual void Stop() override;
+		AUDIOMIXERCORE_API virtual uint32 Run() override;
+		AUDIOMIXERCORE_API virtual void Stop() override;
 
 		// Resume a paused null renderer. 
-		void Resume(const TFunction<void()>& InCallback, float InBufferDuration);
+		AUDIOMIXERCORE_API void Resume(const TFunction<void()>& InCallback, float InBufferDuration);
 
 		// Pause the thread, making it sleep until woken, not consuming cycles or buffers.
-		void Pause();
+		AUDIOMIXERCORE_API void Pause();
 
 	private:
 
diff --git a/Engine/Source/Runtime/AudioMixerCore/Public/AudioMixerTypes.h b/Engine/Source/Runtime/AudioMixerCore/Public/AudioMixerTypes.h
--- a/Engine/Source/Runtime/AudioMixerCore/Public/AudioMixerTypes.h
+++ b/Engine/Source/Runtime/AudioMixerCore/Public/AudioMixerTypes.h
@@ -58,7 +58,7 @@ namespace Audio {
 	inline FName NAME_OPUS(TEXT("OPUS"));
 }
 
-struct AUDIOMIXERCORE_API FAudioPlatformSettings
+struct FAudioPlatformSettings
 {
 	/** Sample rate to use on the platform for the mixing engine. Higher sample rates will incur more CPU cost. */
 	int32 SampleRate;
@@ -75,7 +75,7 @@ struct AUDIOMIXERCORE_API FAudioPlatformSettings
 	/** The number of workers to use to compute source audio. Will only use up to the max number of sources. Will evenly divide sources to each source worker. */
 	int32 NumSourceWorkers;
 
-	static FAudioPlatformSettings GetPlatformSettings(const TCHAR* PlatformSettingsConfigFile);
+	static AUDIOMIXERCORE_API FAudioPlatformSettings GetPlatformSettings(const TCHAR* PlatformSettingsConfigFile);
 
 	FAudioPlatformSettings()
 		: SampleRate(48000)
diff --git a/Engine/Source/Runtime/AudioPlatformConfiguration/Public/AudioCompressionSettings.h b/Engine/Source/Runtime/AudioPlatformConfiguration/Public/AudioCompressionSettings.h
--- a/Engine/Source/Runtime/AudioPlatformConfiguration/Public/AudioCompressionSettings.h
+++ b/Engine/Source/Runtime/AudioPlatformConfiguration/Public/AudioCompressionSettings.h
@@ -58,12 +58,12 @@ struct FAudioStreamCachingSettings
 /* This struct is used for settings used during the cook to a target    */
 /* platform (platform-specific compression quality and resampling, etc.)*/
 /************************************************************************/
-struct AUDIOPLATFORMCONFIGURATION_API FPlatformAudioCookOverrides
+struct FPlatformAudioCookOverrides
 {
 	// Increment this return value to force a recook on all Stream Caching assets.
 	// For testing, it's useful to set this to either a negative number or
 	// absurdly large number, to ensure you do not pollute the DDC.
-	static int32 GetStreamCachingVersion();
+	static AUDIOPLATFORMCONFIGURATION_API int32 GetStreamCachingVersion();
 
 	bool bResampleForDevice;
 
@@ -100,11 +100,11 @@ struct AUDIOPLATFORMCONFIGURATION_API FPlatformAudioCookOverrides
 	}
 
 	// This is used to invalidate compressed audio for a specific platform.
-	static void GetHashSuffix(const FPlatformAudioCookOverrides* InOverrides, FString& OutSuffix);
+	static AUDIOPLATFORMCONFIGURATION_API void GetHashSuffix(const FPlatformAudioCookOverrides* InOverrides, FString& OutSuffix);
 };
 
 USTRUCT()
-struct AUDIOPLATFORMCONFIGURATION_API FPlatformRuntimeAudioCompressionOverrides
+struct FPlatformRuntimeAudioCompressionOverrides
 {
 	GENERATED_USTRUCT_BODY()
 
@@ -127,7 +127,7 @@ struct AUDIOPLATFORMCONFIGURATION_API FPlatformRuntimeAudioCompressionOverrides
 	UPROPERTY(EditAnywhere, Category = "SoundCueLoading", meta = (DisplayName = "Quality Index for Sound Cues", ClampMin = "0", ClampMax = "50"))
 	int32 SoundCueQualityIndex = 1;
 
-	FPlatformRuntimeAudioCompressionOverrides();
+	AUDIOPLATFORMCONFIGURATION_API FPlatformRuntimeAudioCompressionOverrides();
 
 	// Get singleton containing default settings for compression.
 	static FPlatformRuntimeAudioCompressionOverrides* GetDefaultCompressionOverrides()
@@ -141,5 +141,5 @@ struct AUDIOPLATFORMCONFIGURATION_API FPlatformRuntimeAudioCompressionOverrides
 	}
 
 private:
-	static FPlatformRuntimeAudioCompressionOverrides* DefaultCompressionOverrides;
-};
\ No newline at end of file
+	static AUDIOPLATFORMCONFIGURATION_API FPlatformRuntimeAudioCompressionOverrides* DefaultCompressionOverrides;
+};
diff --git a/Engine/Source/Runtime/AudioPlatformConfiguration/Public/AudioResampler.h b/Engine/Source/Runtime/AudioPlatformConfiguration/Public/AudioResampler.h
--- a/Engine/Source/Runtime/AudioPlatformConfiguration/Public/AudioResampler.h
+++ b/Engine/Source/Runtime/AudioPlatformConfiguration/Public/AudioResampler.h
@@ -68,15 +68,15 @@ namespace Audio
 
 	class FResamplerImpl;
 
-	class AUDIOPLATFORMCONFIGURATION_API FResampler
+	class FResampler
 	{
 	public:
-		FResampler();
-		~FResampler();
+		AUDIOPLATFORMCONFIGURATION_API FResampler();
+		AUDIOPLATFORMCONFIGURATION_API ~FResampler();
 
-		void Init(EResamplingMethod ResamplingMethod, float StartingSampleRateRatio, int32 InNumChannels);
-		void SetSampleRateRatio(float InRatio);
-		int32 ProcessAudio(float* InAudioBuffer, int32 InSamples, bool bEndOfInput, float* OutAudioBuffer, int32 MaxOutputFrames, int32& OutNumFrames);
+		AUDIOPLATFORMCONFIGURATION_API void Init(EResamplingMethod ResamplingMethod, float StartingSampleRateRatio, int32 InNumChannels);
+		AUDIOPLATFORMCONFIGURATION_API void SetSampleRateRatio(float InRatio);
+		AUDIOPLATFORMCONFIGURATION_API int32 ProcessAudio(float* InAudioBuffer, int32 InSamples, bool bEndOfInput, float* OutAudioBuffer, int32 MaxOutputFrames, int32& OutNumFrames);
 
 	private:
 		TUniquePtr<FResamplerImpl> CreateImpl();
diff --git a/Engine/Source/Runtime/AutomationTest/Public/AutomationTestExcludelist.h b/Engine/Source/Runtime/AutomationTest/Public/AutomationTestExcludelist.h
--- a/Engine/Source/Runtime/AutomationTest/Public/AutomationTestExcludelist.h
+++ b/Engine/Source/Runtime/AutomationTest/Public/AutomationTestExcludelist.h
@@ -205,33 +205,33 @@ struct FAutomationTestExcludelistEntry
 };
 
 
-UCLASS(config = Engine, defaultconfig)
-class AUTOMATIONTEST_API UAutomationTestExcludelist : public UObject
+UCLASS(config = Engine, defaultconfig, MinimalAPI)
+class UAutomationTestExcludelist : public UObject
 {
 	GENERATED_BODY()
 
 public:
 	UAutomationTestExcludelist() {}
 
-	static UAutomationTestExcludelist* Get();
+	static AUTOMATIONTEST_API UAutomationTestExcludelist* Get();
 
-	void AddToExcludeTest(const FString& TestName, const FAutomationTestExcludelistEntry& ExcludelistEntry);
-	void RemoveFromExcludeTest(const FString& TestName);
-	bool IsTestExcluded(const FString& TestName, const TSet<FName>& = TSet<FName>(), FName* OutReason = nullptr, bool* OutWarn = nullptr);
-	FAutomationTestExcludelistEntry* GetExcludeTestEntry(const FString& TestName, const TSet<FName>& = TSet<FName>());
+	AUTOMATIONTEST_API void AddToExcludeTest(const FString& TestName, const FAutomationTestExcludelistEntry& ExcludelistEntry);
+	AUTOMATIONTEST_API void RemoveFromExcludeTest(const FString& TestName);
+	AUTOMATIONTEST_API bool IsTestExcluded(const FString& TestName, const TSet<FName>& = TSet<FName>(), FName* OutReason = nullptr, bool* OutWarn = nullptr);
+	AUTOMATIONTEST_API FAutomationTestExcludelistEntry* GetExcludeTestEntry(const FString& TestName, const TSet<FName>& = TSet<FName>());
 
-	void SaveConfig();
+	AUTOMATIONTEST_API void SaveConfig();
 	FString GetConfigFilename() { return UObject::GetDefaultConfigFilename(); } const
 	// It is called automatically when CDO is created, usually you don't need to call LoadConfig manually
 	void LoadConfig() { UObject::LoadConfig(GetClass()); }
 
-	virtual void OverrideConfigSection(FString& SectionName) override;
+	AUTOMATIONTEST_API virtual void OverrideConfigSection(FString& SectionName) override;
 
 protected:
-	virtual void PostInitProperties() override;
+	AUTOMATIONTEST_API virtual void PostInitProperties() override;
 
 private:
-	FString GetFullTestName(const FAutomationTestExcludelistEntry& ExcludelistEntry);
+	AUTOMATIONTEST_API FString GetFullTestName(const FAutomationTestExcludelistEntry& ExcludelistEntry);
 
 	UPROPERTY(EditDefaultsOnly, globalconfig, Category = AutomationTestExcludelist)
 	TArray<FAutomationTestExcludelistEntry> ExcludeTest;
diff --git a/Engine/Source/Runtime/CUDA/Source/Public/CudaModule.h b/Engine/Source/Runtime/CUDA/Source/Public/CudaModule.h
--- a/Engine/Source/Runtime/CUDA/Source/Public/CudaModule.h
+++ b/Engine/Source/Runtime/CUDA/Source/Public/CudaModule.h
@@ -13,7 +13,7 @@
 DECLARE_MULTICAST_DELEGATE(FOnPostCUDAInit)
 
 DECLARE_LOG_CATEGORY_EXTERN(LogCUDA, Log, All);
-class CUDA_API FCUDAModule : public IModuleInterface
+class FCUDAModule : public IModuleInterface
 {
 public:
 	static CUDA_DRIVER_API_FUNCTION_LIST CUDA()
@@ -21,35 +21,35 @@ public:
 		return FModuleManager::GetModuleChecked<FCUDAModule>("CUDA").DriverApiPtrs;
 	}
 
-	virtual void StartupModule() override;
-	virtual void ShutdownModule() override;
+	CUDA_API virtual void StartupModule() override;
+	CUDA_API virtual void ShutdownModule() override;
 
 	// Determines if the CUDA Driver API is available for use
-	bool IsAvailable();
+	CUDA_API bool IsAvailable();
 
 	// Retrieves the function pointer list for the CUDA Driver API
-	const CUDA_DRIVER_API_FUNCTION_LIST* DriverAPI();
+	CUDA_API const CUDA_DRIVER_API_FUNCTION_LIST* DriverAPI();
 
 	// Retrieves the CUDA context for the GPU device currently in use by the Vulkan RHI
-	CUcontext GetCudaContext();
+	CUDA_API CUcontext GetCudaContext();
 	
 	// Retrieves the device index for the current context
-	uint32 GetCudaDeviceIndex() const;
+	CUDA_API uint32 GetCudaDeviceIndex() const;
 	
 	// Retrieves or creates the CUDA context for the specified GPU device
-	CUcontext GetCudaContextForDevice(int DeviceIndex);
+	CUDA_API CUcontext GetCudaContextForDevice(int DeviceIndex);
 
 	// Called after CUDA is loaded successfully
 	FOnPostCUDAInit OnPostCUDAInit;
 
-	bool IsDeviceIndexRHISelected(int DeviceIndex);
+	CUDA_API bool IsDeviceIndexRHISelected(int DeviceIndex);
 
 private:
-	bool LoadCuda();
-	void UnloadCuda();
+	CUDA_API bool LoadCuda();
+	CUDA_API void UnloadCuda();
 
-	void InitCuda();
-	bool IsRHISelectedDevice(CUdevice cuDevice);
+	CUDA_API void InitCuda();
+	CUDA_API bool IsRHISelectedDevice(CUdevice cuDevice);
 
 	void* DriverLibrary;
 	CUDA_DRIVER_API_FUNCTION_LIST DriverApiPtrs;
diff --git a/Engine/Source/Runtime/Cbor/Public/CborReader.h b/Engine/Source/Runtime/Cbor/Public/CborReader.h
--- a/Engine/Source/Runtime/Cbor/Public/CborReader.h
+++ b/Engine/Source/Runtime/Cbor/Public/CborReader.h
@@ -13,7 +13,7 @@ class FArchive;
  * Reader for a the cbor protocol encoded stream
  * @see http://cbor.io
  */
-class CBOR_API FCborReader
+class FCborReader
 {
 public:
 	/**
@@ -22,38 +22,38 @@ public:
 	 * @param InReaderEndianness Specify which endianness should be use to read the archive.
 	 * @note CBOR standard endianness is big endian. For interoperability with external tools, the standard endianness should be used. For internal usage, the platform endianness is faster.
 	 */
-	FCborReader(FArchive* InStream, ECborEndianness InReaderEndianness = ECborEndianness::Platform);
-	~FCborReader();
+	CBOR_API FCborReader(FArchive* InStream, ECborEndianness InReaderEndianness = ECborEndianness::Platform);
+	CBOR_API ~FCborReader();
 
 	/** @return the archive we are reading from. */
-	const FArchive* GetArchive() const;
+	CBOR_API const FArchive* GetArchive() const;
 
 	/** @return true if the reader is in error. */
-	bool IsError() const;
+	CBOR_API bool IsError() const;
 
 	/** @return A cbor Header containing an error code as its raw code. */
-	FCborHeader GetError() const;
+	CBOR_API FCborHeader GetError() const;
 
 	/**
 	 * The cbor context of the reader can either be
 	 * a container context or a dummy.
 	 * A reference to the context shouldn't be held while calling ReadNext.
 	 * @return The current cbor context. */
-	const FCborContext& GetContext() const;
+	CBOR_API const FCborContext& GetContext() const;
 
 	/**
 	 * Read the next value from the cbor stream.
 	 * @param OutContext the context to read the value into.
 	 * @return true if successful, false if an error was returned or the end of the stream was reached.
 	 */
-	bool ReadNext(FCborContext& OutContext);
+	CBOR_API bool ReadNext(FCborContext& OutContext);
 
 	/**
 	 * Skip a container of ContainerType type
 	 * @param ContainerType the container we expect to skip.
 	 * @return true if successful, false if the current container wasn't a ContainerType or an error occurred.
 	 */
-	bool SkipContainer(ECborCode ContainerType);
+	CBOR_API bool SkipContainer(ECborCode ContainerType);
 	
 private:
 	/** Read a uint value from Ar into OutContext and also return it. */
diff --git a/Engine/Source/Runtime/Cbor/Public/CborWriter.h b/Engine/Source/Runtime/Cbor/Public/CborWriter.h
--- a/Engine/Source/Runtime/Cbor/Public/CborWriter.h
+++ b/Engine/Source/Runtime/Cbor/Public/CborWriter.h
@@ -14,7 +14,7 @@ class FArchive;
 * Writer for encoding a stream with the cbor protocol
 * @see http://cbor.io
 */
-class CBOR_API FCborWriter
+class FCborWriter
 {
 public:
 	/*
@@ -23,34 +23,34 @@ public:
 	 * @param InWriterEndianness Specify which endianness should be use to write the archive.
 	 * @note CBOR standard endianness is big endian. For interoperability with external tools, the standard endianness should be used. For internal usage, the platform endianness is faster.
 	 */
-	FCborWriter(FArchive* InStream, ECborEndianness InWriterEndianness = ECborEndianness::Platform);
-	~FCborWriter();
+	CBOR_API FCborWriter(FArchive* InStream, ECborEndianness InWriterEndianness = ECborEndianness::Platform);
+	CBOR_API ~FCborWriter();
 
 public:
 	/** @return the archive we are writing to. */
-	const FArchive* GetArchive() const;
+	CBOR_API const FArchive* GetArchive() const;
 
 	/**
 	 * Write a container start code.
 	 * @param ContainerType container major type, either array or map.
 	 * @param NbItem the number of item in the container or negative to indicate indefinite containers.
 	 */
-	void WriteContainerStart(ECborCode ContainerType, int64 NbItem);
+	CBOR_API void WriteContainerStart(ECborCode ContainerType, int64 NbItem);
 
 	/** Write a container break code, need a indefinite container context. */
-	void WriteContainerEnd();
+	CBOR_API void WriteContainerEnd();
 
 	/** Write a value. */
 
-	void WriteNull();
-	void WriteValue(uint64 Value);
-	void WriteValue(int64 Value);
-	void WriteValue(bool Value);
-	void WriteValue(float Value);
-	void WriteValue(double Value);
-	void WriteValue(const FString& Value);
-	void WriteValue(const char* CString, uint64 Length);
-	void WriteValue(const uint8* Bytes, uint64 Length);
+	CBOR_API void WriteNull();
+	CBOR_API void WriteValue(uint64 Value);
+	CBOR_API void WriteValue(int64 Value);
+	CBOR_API void WriteValue(bool Value);
+	CBOR_API void WriteValue(float Value);
+	CBOR_API void WriteValue(double Value);
+	CBOR_API void WriteValue(const FString& Value);
+	CBOR_API void WriteValue(const char* CString, uint64 Length);
+	CBOR_API void WriteValue(const uint8* Bytes, uint64 Length);
 
 private:
 	/** Write a uint Value for Header in Ar and return the final generated cbor Header. */
diff --git a/Engine/Source/Runtime/CinematicCamera/Public/CameraRig_Crane.h b/Engine/Source/Runtime/CinematicCamera/Public/CameraRig_Crane.h
--- a/Engine/Source/Runtime/CinematicCamera/Public/CameraRig_Crane.h
+++ b/Engine/Source/Runtime/CinematicCamera/Public/CameraRig_Crane.h
@@ -14,18 +14,18 @@ class UStaticMeshComponent;
 /** 
  * A simple rig for simulating crane-like camera movements.
  */
-UCLASS(Blueprintable)
-class CINEMATICCAMERA_API ACameraRig_Crane : public AActor
+UCLASS(Blueprintable, MinimalAPI)
+class ACameraRig_Crane : public AActor
 {
 	GENERATED_BODY()
 	
 public:
 
 	// ctor
-	ACameraRig_Crane(const FObjectInitializer& ObjectInitialier);
+	CINEMATICCAMERA_API ACameraRig_Crane(const FObjectInitializer& ObjectInitialier);
 
-	virtual void Tick(float DeltaTime) override;
-	virtual bool ShouldTickIfViewportsOnly() const override;
+	CINEMATICCAMERA_API virtual void Tick(float DeltaTime) override;
+	CINEMATICCAMERA_API virtual bool ShouldTickIfViewportsOnly() const override;
 
 	/** Controls the pitch of the crane arm. */
 	UPROPERTY(Interp, EditAnywhere, BlueprintReadWrite, Category = "Crane Controls", meta = (UIMin = "-360", UIMax = "360", Units = deg))
@@ -48,10 +48,10 @@ public:
 	bool bLockMountYaw;
 
 #if WITH_EDITOR
-	virtual void PostEditChangeProperty(FPropertyChangedEvent& PropertyChangedEvent) override;
-	virtual void PostEditUndo() override;
+	CINEMATICCAMERA_API virtual void PostEditChangeProperty(FPropertyChangedEvent& PropertyChangedEvent) override;
+	CINEMATICCAMERA_API virtual void PostEditUndo() override;
 #endif
-	virtual class USceneComponent* GetDefaultAttachComponent() const override;
+	CINEMATICCAMERA_API virtual class USceneComponent* GetDefaultAttachComponent() const override;
 	
 private:
 #if WITH_EDITORONLY_DATA
diff --git a/Engine/Source/Runtime/CinematicCamera/Public/CameraRig_Rail.h b/Engine/Source/Runtime/CinematicCamera/Public/CameraRig_Rail.h
--- a/Engine/Source/Runtime/CinematicCamera/Public/CameraRig_Rail.h
+++ b/Engine/Source/Runtime/CinematicCamera/Public/CameraRig_Rail.h
@@ -15,17 +15,17 @@ class UStaticMeshComponent;
 /** 
  * 
  */
-UCLASS(Blueprintable)
-class CINEMATICCAMERA_API ACameraRig_Rail : public AActor
+UCLASS(Blueprintable, MinimalAPI)
+class ACameraRig_Rail : public AActor
 {
 	GENERATED_BODY()
 	
 public:
 	// ctor
-	ACameraRig_Rail(const FObjectInitializer& ObjectInitialier);
+	CINEMATICCAMERA_API ACameraRig_Rail(const FObjectInitializer& ObjectInitialier);
 	
-	virtual void Tick(float DeltaTime) override;
-	virtual bool ShouldTickIfViewportsOnly() const override;
+	CINEMATICCAMERA_API virtual void Tick(float DeltaTime) override;
+	CINEMATICCAMERA_API virtual bool ShouldTickIfViewportsOnly() const override;
 
 	/** Defines current position of the mount point along the rail, in terms of normalized distance from the beginning of the rail. */
 	UPROPERTY(Interp, EditAnywhere, BlueprintReadWrite, Category = "Rail Controls", meta=(ClampMin="0.0", ClampMax = "1.0"))
@@ -45,11 +45,11 @@ public:
 	float PreviewMeshScale;
 #endif
 
-	virtual class USceneComponent* GetDefaultAttachComponent() const override;
+	CINEMATICCAMERA_API virtual class USceneComponent* GetDefaultAttachComponent() const override;
 #if WITH_EDITOR
-	virtual void PostEditChangeProperty(FPropertyChangedEvent& PropertyChangedEvent) override;
-	virtual void PostEditUndo() override;
-	virtual void PostEditMove(bool bFinished) override;
+	CINEMATICCAMERA_API virtual void PostEditChangeProperty(FPropertyChangedEvent& PropertyChangedEvent) override;
+	CINEMATICCAMERA_API virtual void PostEditUndo() override;
+	CINEMATICCAMERA_API virtual void PostEditMove(bool bFinished) override;
 #endif
 
 	/** Returns the spline component that defines the rail path */
@@ -58,11 +58,11 @@ public:
 	
 protected:
 	/** Makes sure all components are arranged properly. Call when something changes that might affect components. */
-	virtual void UpdateRailComponents();
+	CINEMATICCAMERA_API virtual void UpdateRailComponents();
 #if WITH_EDITORONLY_DATA
-	void UpdatePreviewMeshes();
+	CINEMATICCAMERA_API void UpdatePreviewMeshes();
 #endif
-	USplineMeshComponent* CreateSplinePreviewSegment();
+	CINEMATICCAMERA_API USplineMeshComponent* CreateSplinePreviewSegment();
 
  	/** Root component to give the whole actor a transform. */
 	UPROPERTY(EditDefaultsOnly, Category = "Rail Components")
diff --git a/Engine/Source/Runtime/CinematicCamera/Public/CineCameraActor.h b/Engine/Source/Runtime/CinematicCamera/Public/CineCameraActor.h
--- a/Engine/Source/Runtime/CinematicCamera/Public/CineCameraActor.h
+++ b/Engine/Source/Runtime/CinematicCamera/Public/CineCameraActor.h
@@ -65,18 +65,18 @@ struct FCameraLookatTrackingSettings
 /** 
  * A CineCameraActor is a CameraActor specialized to work like a cinematic camera.
  */
-UCLASS(ClassGroup = Common, hideCategories = (Input, Rendering, AutoPlayerActivation), showcategories = ("Input|MouseInput", "Input|TouchInput"), Blueprintable)
-class CINEMATICCAMERA_API ACineCameraActor : public ACameraActor
+UCLASS(ClassGroup = Common, hideCategories = (Input, Rendering, AutoPlayerActivation), showcategories = ("Input|MouseInput", "Input|TouchInput"), Blueprintable, MinimalAPI)
+class ACineCameraActor : public ACameraActor
 {
 	GENERATED_BODY()
 
 public:
 	// Ctor
-	ACineCameraActor(const FObjectInitializer& ObjectInitializer);
+	CINEMATICCAMERA_API ACineCameraActor(const FObjectInitializer& ObjectInitializer);
 
-	virtual void Tick(float DeltaTime) override;
-	virtual bool ShouldTickIfViewportsOnly() const override;
-	virtual void PostInitializeComponents() override;
+	CINEMATICCAMERA_API virtual void Tick(float DeltaTime) override;
+	CINEMATICCAMERA_API virtual bool ShouldTickIfViewportsOnly() const override;
+	CINEMATICCAMERA_API virtual void PostInitializeComponents() override;
 
 	UPROPERTY(EditAnywhere, BlueprintReadWrite, Category = "Current Camera Settings")
 	FCameraLookatTrackingSettings LookatTrackingSettings;
@@ -89,11 +89,11 @@ protected:
 	/** Set to true to skip any interpolations on the next update. Resets to false automatically. */
 	uint8 bResetInterplation : 1;
 
-	FVector GetLookatLocation() const;
+	CINEMATICCAMERA_API FVector GetLookatLocation() const;
 
-	virtual void NotifyCameraCut() override;
+	CINEMATICCAMERA_API virtual void NotifyCameraCut() override;
 
-	bool ShouldTickForTracking() const;
+	CINEMATICCAMERA_API bool ShouldTickForTracking() const;
 
 private:
 	/** Returns CineCameraComponent subobject **/
diff --git a/Engine/Source/Runtime/CinematicCamera/Public/CineCameraComponent.h b/Engine/Source/Runtime/CinematicCamera/Public/CineCameraComponent.h
--- a/Engine/Source/Runtime/CinematicCamera/Public/CineCameraComponent.h
+++ b/Engine/Source/Runtime/CinematicCamera/Public/CineCameraComponent.h
@@ -17,18 +17,18 @@ class UStaticMeshComponent;
 /**
  * A specialized version of a camera component, geared toward cinematic usage.
  */
-UCLASS(HideCategories = (CameraSettings), HideFunctions = (SetFieldOfView, SetAspectRatio), Blueprintable, ClassGroup = Camera, meta = (BlueprintSpawnableComponent), Config = Engine)
-class CINEMATICCAMERA_API UCineCameraComponent : public UCameraComponent
+UCLASS(HideCategories = (CameraSettings), HideFunctions = (SetFieldOfView, SetAspectRatio), Blueprintable, ClassGroup = Camera, meta = (BlueprintSpawnableComponent), Config = Engine, MinimalAPI)
+class UCineCameraComponent : public UCameraComponent
 {
 	GENERATED_BODY()
 
 public:
 	/** Default constuctor. */
-	UCineCameraComponent();
+	CINEMATICCAMERA_API UCineCameraComponent();
 
-	virtual void GetCameraView(float DeltaTime, FMinimalViewInfo& DesiredView) override;
+	CINEMATICCAMERA_API virtual void GetCameraView(float DeltaTime, FMinimalViewInfo& DesiredView) override;
 #if WITH_EDITOR
-	virtual FText GetFilmbackText() const override;
+	CINEMATICCAMERA_API virtual FText GetFilmbackText() const override;
 #endif
 	UPROPERTY()
 	FCameraFilmbackSettings FilmbackSettings_DEPRECATED;
@@ -38,28 +38,28 @@ public:
 	FCameraFilmbackSettings Filmback;
 
 	UFUNCTION(BlueprintSetter)
-	void SetFilmback(const FCameraFilmbackSettings& NewFilmback);
+	CINEMATICCAMERA_API void SetFilmback(const FCameraFilmbackSettings& NewFilmback);
 
 	/** Controls the camera's lens. */
 	UPROPERTY(EditAnywhere, BlueprintSetter = SetLensSettings, BlueprintReadWrite, Category = "Current Camera Settings")
 	FCameraLensSettings LensSettings;
 
 	UFUNCTION(BlueprintSetter)
-	void SetLensSettings(const FCameraLensSettings& NewLensSettings);
+	CINEMATICCAMERA_API void SetLensSettings(const FCameraLensSettings& NewLensSettings);
 
 	/** Controls the camera's focus. */
 	UPROPERTY(EditAnywhere, BlueprintSetter = SetFocusSettings, BlueprintReadWrite, Category = "Current Camera Settings")
 	FCameraFocusSettings FocusSettings;
 
 	UFUNCTION(BlueprintSetter)
-	void SetFocusSettings(const FCameraFocusSettings& NewFocusSettings);
+	CINEMATICCAMERA_API void SetFocusSettings(const FCameraFocusSettings& NewFocusSettings);
 	
 	/** Controls the crop settings. */
 	UPROPERTY(EditAnywhere, BlueprintSetter = SetCropSettings, BlueprintReadWrite, Category = "Current Camera Settings")
 	FPlateCropSettings CropSettings;
 
 	UFUNCTION(BlueprintSetter)
-	void SetCropSettings(const FPlateCropSettings& NewCropSettings);
+	CINEMATICCAMERA_API void SetCropSettings(const FPlateCropSettings& NewCropSettings);
 
 	/** Current focal length of the camera (i.e. controls FoV, zoom) */
 	UPROPERTY(Interp, BlueprintSetter = SetCurrentFocalLength, EditAnywhere, BlueprintReadWrite, Category = "Current Camera Settings")
@@ -70,7 +70,7 @@ public:
 	float CurrentAperture;
 
 	UFUNCTION(BlueprintSetter)
-	void SetCurrentAperture(const float NewCurrentAperture);
+	CINEMATICCAMERA_API void SetCurrentAperture(const float NewCurrentAperture);
 	
 	/** Read-only. Control this value via FocusSettings. */
 	UPROPERTY(VisibleAnywhere, BlueprintReadOnly, Category = "Current Camera Settings")
@@ -90,77 +90,77 @@ public:
 #endif
 
 	/** Override setting FOV to manipulate Focal Length. */
-	virtual void SetFieldOfView(float InFieldOfView) override;
+	CINEMATICCAMERA_API virtual void SetFieldOfView(float InFieldOfView) override;
 	
 	UFUNCTION(BlueprintCallable, BlueprintSetter, Category = "Cine Camera")
-	void SetCurrentFocalLength(float InFocalLength);
+	CINEMATICCAMERA_API void SetCurrentFocalLength(float InFocalLength);
 
 	/** Sets near clipping plane of the cine camera. */
 	UFUNCTION(BlueprintCallable, BlueprintSetter, Category = "Cine Camera")
-	void SetCustomNearClippingPlane(const float NewCustomNearClippingPlane);
+	CINEMATICCAMERA_API void SetCustomNearClippingPlane(const float NewCustomNearClippingPlane);
 
 	/** Returns the horizonal FOV of the camera with current settings. */
 	UFUNCTION(BlueprintCallable, Category = "Cine Camera")
-	float GetHorizontalFieldOfView() const;
+	CINEMATICCAMERA_API float GetHorizontalFieldOfView() const;
 	
 	/** Returns the vertical FOV of the camera with current settings. */
 	UFUNCTION(BlueprintCallable, Category = "Cine Camera")
-	float GetVerticalFieldOfView() const;
+	CINEMATICCAMERA_API float GetVerticalFieldOfView() const;
 
 	/** Returns the filmback name of the camera with the current settings. */
 	UFUNCTION(BlueprintCallable, Category = "Cine Camera")
-	FString GetFilmbackPresetName() const;
+	CINEMATICCAMERA_API FString GetFilmbackPresetName() const;
 
 	/** Returns the name of the default filmback preset. */
 	UE_DEPRECATED(5.1, "This function has been deprecated, please access the presets via the CineCameraSettings object.")
 	UFUNCTION(BlueprintPure, Category = "Cine Camera")
-	FString GetDefaultFilmbackPresetName() const;
+	CINEMATICCAMERA_API FString GetDefaultFilmbackPresetName() const;
 
 	/** Set the current preset settings by preset name. */
 	UFUNCTION(BlueprintCallable, Category = "Cine Camera")
-	void SetFilmbackPresetByName(const FString& InPresetName);
+	CINEMATICCAMERA_API void SetFilmbackPresetByName(const FString& InPresetName);
 
 	/** Returns the lens name of the camera with the current settings. */
 	UFUNCTION(BlueprintCallable, Category = "Cine Camera")
-	FString GetLensPresetName() const;
+	CINEMATICCAMERA_API FString GetLensPresetName() const;
 
 	/** Set the current lens settings by preset name. */
 	UFUNCTION(BlueprintCallable, Category = "Cine Camera")
-	void SetLensPresetByName(const FString& InPresetName);
+	CINEMATICCAMERA_API void SetLensPresetByName(const FString& InPresetName);
 	
 	/** Returns the lens name of the camera with the current settings. */
 	UFUNCTION(BlueprintCallable, Category = "Cine Camera")
-	FString GetCropPresetName() const;
+	CINEMATICCAMERA_API FString GetCropPresetName() const;
 
 	/** Set the current lens settings by preset name. */
 	UFUNCTION(BlueprintCallable, Category = "Cine Camera")
-	void SetCropPresetByName(const FString& InPresetName);
+	CINEMATICCAMERA_API void SetCropPresetByName(const FString& InPresetName);
 
 	/** Returns a copy of the list of available filmback presets. */
 	UE_DEPRECATED(5.1, "This function has been deprecated, please access the presets via the CineCameraSettings object.")
 	UFUNCTION(BlueprintCallable, Category = "Cine Camera")
-	static TArray<FNamedFilmbackPreset> GetFilmbackPresetsCopy();
+	static CINEMATICCAMERA_API TArray<FNamedFilmbackPreset> GetFilmbackPresetsCopy();
 
 	/** Returns a copy of the list of available lens presets. */
 	UE_DEPRECATED(5.1, "This function has been deprecated, please access the presets via the CineCameraSettings object.")
 	UFUNCTION(BlueprintCallable, Category = "Cine Camera")
-	static TArray<FNamedLensPreset> GetLensPresetsCopy();
+	static CINEMATICCAMERA_API TArray<FNamedLensPreset> GetLensPresetsCopy();
 
 	/** Returns a list of available filmback presets. */
 	UE_DEPRECATED(5.1, "This function has been deprecated, please access the presets via the CineCameraSettings object.")
-	static TArray<FNamedFilmbackPreset> const& GetFilmbackPresets();
+	static CINEMATICCAMERA_API TArray<FNamedFilmbackPreset> const& GetFilmbackPresets();
 	
 	/** Returns a list of available lens presets. */
 	UE_DEPRECATED(5.1, "This function has been deprecated, please access the presets via the CineCameraSettings object.")
-	static TArray<FNamedLensPreset> const& GetLensPresets();
+	static CINEMATICCAMERA_API TArray<FNamedLensPreset> const& GetLensPresets();
 
 #if WITH_EDITOR
 	/** Update the debug focus plane position and orientation. */
-	void UpdateDebugFocusPlane();
+	CINEMATICCAMERA_API void UpdateDebugFocusPlane();
 #endif
 
 	/** Returns the world to meters scale for the current UWorld */
-	float GetWorldToMetersScale() const;
+	CINEMATICCAMERA_API float GetWorldToMetersScale() const;
 
 protected:
 
@@ -172,17 +172,17 @@ protected:
 
 	/// @cond DOXYGEN_WARNINGS
 	
-	virtual void PostLoad() override;
+	CINEMATICCAMERA_API virtual void PostLoad() override;
 	
 	/// @endcond
 	
-	virtual void Serialize(FArchive& Ar) override;
-	virtual void PostInitProperties() override;
-	virtual void OnRegister() override;
-	virtual void TickComponent(float DeltaTime, ELevelTick TickType, FActorComponentTickFunction* ThisTickFunction) override;
+	CINEMATICCAMERA_API virtual void Serialize(FArchive& Ar) override;
+	CINEMATICCAMERA_API virtual void PostInitProperties() override;
+	CINEMATICCAMERA_API virtual void OnRegister() override;
+	CINEMATICCAMERA_API virtual void TickComponent(float DeltaTime, ELevelTick TickType, FActorComponentTickFunction* ThisTickFunction) override;
 
 #if WITH_EDITOR
-	virtual void OnComponentDestroyed(bool bDestroyingHierarchy) override;
+	CINEMATICCAMERA_API virtual void OnComponentDestroyed(bool bDestroyingHierarchy) override;
 #endif
 
 #if WITH_EDITORONLY_DATA
@@ -202,8 +202,8 @@ protected:
 	UPROPERTY(transient)
 	TObjectPtr<UMaterialInstanceDynamic> DebugFocusPlaneMID;
 
-	virtual void PostEditChangeProperty(FPropertyChangedEvent& PropertyChangedEvent) override;
-	virtual void ResetProxyMeshTransform() override;
+	CINEMATICCAMERA_API virtual void PostEditChangeProperty(FPropertyChangedEvent& PropertyChangedEvent) override;
+	CINEMATICCAMERA_API virtual void ResetProxyMeshTransform() override;
 #endif
 
 	/** List of available filmback presets */
@@ -241,11 +241,11 @@ protected:
 	UPROPERTY(config)
 	float DefaultLensFStop;
 
-	virtual void UpdateCameraLens(float DeltaTime, FMinimalViewInfo& DesiredView);
+	CINEMATICCAMERA_API virtual void UpdateCameraLens(float DeltaTime, FMinimalViewInfo& DesiredView);
 
-	virtual void NotifyCameraCut() override;
+	CINEMATICCAMERA_API virtual void NotifyCameraCut() override;
 	
-	void RecalcDerivedData();
+	CINEMATICCAMERA_API void RecalcDerivedData();
 
 private:
 	float GetDesiredFocusDistance(const FVector& InLocation) const;
