Fix crash on shutdown in audio mixer (github pr)
#10372

diff --git a/Engine/Source/Runtime/AudioMixer/Private/AudioMixerSubmix.cpp b/Engine/Source/Runtime/AudioMixer/Private/AudioMixerSubmix.cpp
--- a/Engine/Source/Runtime/AudioMixer/Private/AudioMixerSubmix.cpp
+++ b/Engine/Source/Runtime/AudioMixer/Private/AudioMixerSubmix.cpp
@@ -1090,9 +1090,9 @@ namespace Audio
 		{
 			// query the SubmixBufferListeners to see if they plan to render audio into this buffer
 			FScopeLock Lock(&BufferListenerCriticalSection);
-			for(const ISubmixBufferListener* BufferListener : BufferListeners)
+			for (const ISubmixBufferListener* BufferListener : BufferListeners)
 			{
-				if(BufferListener->IsRenderingAudio())
+				if (BufferListener && BufferListener->IsRenderingAudio())
 				{
 					return true;
 				}
