Unshelved from pending changelist '26281904':

Quartz Subscriber Optimization

During the creation of this class, it would subscribe itself to the Quartz Subsystem, so that it can receive regular tick updates. The unsubscribe was happening exclusively in the destructor. Therefore for managed objects, they weren't being unsubscribed until the object was fully deleted by the garbage collector.

I wrapped up the subscription code inside a new QuartzSubscriptionToken. This class adds a few more checks for sanity, and prevents things from being subscribed multiple times, or unsubscribed when its not even subscribed. The Actor/Component objects that derive from this class now unsubscribe in the EndPlay function call - they subscribe in BeginPlay. This prevents the constant build up in subscribed objects until the garbage collector runs.

diff --git a/Engine/Source/Runtime/AudioMixer/Public/Quartz/AudioMixerClockHandle.h b/Engine/Source/Runtime/AudioMixer/Public/Quartz/AudioMixerClockHandle.h
--- a/Engine/Source/Runtime/AudioMixer/Public/Quartz/AudioMixerClockHandle.h
+++ b/Engine/Source/Runtime/AudioMixer/Public/Quartz/AudioMixerClockHandle.h
@@ -3,9 +3,7 @@
 #pragma once
 
 #include "CoreMinimal.h"
-#include "UObject/ObjectMacros.h"
-#include "Curves/RichCurve.h"
-#include "Curves/CurveBase.h"
+#include "Containers/Queue.h"
 #include "Sound/QuartzQuantizationUtilities.h"
 #include "Sound/QuartzSubscription.h"
 #include "Quartz/QuartzSubsystem.h"
@@ -13,6 +11,9 @@
 
 #include "AudioMixerClockHandle.generated.h"
 
+
+
+
 /**
  *  This class is a BP / Game thread wrapper around FQuartzClockProxy
  *	(to talk to the underlying clock)
@@ -23,48 +24,45 @@
  *  It can subscribe to Quantized Event & Metronome delegates to synchronize
  *  gameplay & VFX to Quartz events fired from the Audio Engine
  */
-UCLASS(BlueprintType, Blueprintable, Transient, ClassGroup = Quartz, meta = (BlueprintSpawnableComponent), MinimalAPI)
-class UQuartzClockHandle : public UObject, public FQuartzTickableObject
+UCLASS(BlueprintType, Blueprintable, Transient, ClassGroup = Quartz, meta = (BlueprintSpawnableComponent))
+class AUDIOMIXER_API UQuartzClockHandle : public UObject, public FQuartzTickableObject
 {
 	GENERATED_BODY()
 
 public:
 	// ctor
-	AUDIOMIXER_API UQuartzClockHandle();
+	UQuartzClockHandle();
 
 	// dtor
-	AUDIOMIXER_API ~UQuartzClockHandle();
+	AUDIOMIXER_API virtual ~UQuartzClockHandle();
 
 	// begin UObject interface
-	AUDIOMIXER_API void BeginDestroy() override;
+	void BeginDestroy() override;
 	// end UObject interface
 
 // Begin Blueprint Interface
 
 	// Clock manipulation
 	UFUNCTION(BlueprintCallable, Category = "Quartz Clock", meta = (WorldContext = "WorldContextObject"))
-	AUDIOMIXER_API void StartClock(const UObject* WorldContextObject, UQuartzClockHandle*& ClockHandle);
+	void StartClock(const UObject* WorldContextObject, UQuartzClockHandle*& ClockHandle);
 
 	UFUNCTION(BlueprintCallable, Category = "Quartz Subsystem", meta = (WorldContext = "WorldContextObject"))
-	AUDIOMIXER_API void StopClock(const UObject* WorldContextObject, bool CancelPendingEvents, UQuartzClockHandle*& ClockHandle);
+	void StopClock(const UObject* WorldContextObject, bool CancelPendingEvents, UQuartzClockHandle*& ClockHandle);
 
 	UFUNCTION(BlueprintCallable, Category = "Quartz Subsystem", meta = (WorldContext = "WorldContextObject"))
-	AUDIOMIXER_API void PauseClock(const UObject* WorldContextObject, UQuartzClockHandle*& ClockHandle);
+	void PauseClock(const UObject* WorldContextObject, UQuartzClockHandle*& ClockHandle);
 
 	UFUNCTION(BlueprintCallable, Category = "Quartz Clock", meta = (WorldContext = "WorldContextObject"))
-	AUDIOMIXER_API void ResumeClock(const UObject* WorldContextObject, UQuartzClockHandle*& ClockHandle);
+	void ResumeClock(const UObject* WorldContextObject, UQuartzClockHandle*& ClockHandle);
 
 	UFUNCTION(BlueprintCallable, Category = "Quartz Clock", meta = (WorldContext = "WorldContextObject", AutoCreateRefTerm = "InDelegate", Keywords = "Transport, Counter"))
-	AUDIOMIXER_API void ResetTransport(const UObject* WorldContextObject, const FOnQuartzCommandEventBP& InDelegate);
+	void ResetTransport(const UObject* WorldContextObject, const FOnQuartzCommandEventBP& InDelegate);
 
 	UFUNCTION(BlueprintCallable, Category = "Quartz Clock", meta = (WorldContext = "WorldContextObject", AutoCreateRefTerm = "InDelegate", Keywords = "Transport, Counter"))
-	AUDIOMIXER_API void ResetTransportQuantized(const UObject* WorldContextObject, FQuartzQuantizationBoundary InQuantizationBoundary, const FOnQuartzCommandEventBP& InDelegate, UQuartzClockHandle*& ClockHandle);
+	void ResetTransportQuantized(const UObject* WorldContextObject, FQuartzQuantizationBoundary InQuantizationBoundary, const FOnQuartzCommandEventBP& InDelegate, UQuartzClockHandle*& ClockHandle);
 
 	UFUNCTION(BlueprintCallable, Category = "Quartz Clock", meta = (WorldContext = "WorldContextObject", Keywords = "Transport, Counter"))
-	AUDIOMIXER_API bool IsClockRunning(const UObject* WorldContextObject);
-
-	UFUNCTION(BlueprintCallable, Category = "Quartz Clock", meta = (WorldContext = "WorldContextObject", AutoCreateRefTerm = "InDelegate", Keywords = "Transport, Counter"))
-	AUDIOMIXER_API void NotifyOnQuantizationBoundary(const UObject* WorldContextObject, FQuartzQuantizationBoundary InQuantizationBoundary, const FOnQuartzCommandEventBP& InDelegate, float InMsOffset = 0.f);
+	bool IsClockRunning(const UObject* WorldContextObject);
 
 	/** Returns the duration in seconds of the given Quantization Type
 	 *
@@ -73,72 +71,64 @@ public:
 	 * @return The duration, in seconds, of a multiplier amount of the Quantization Type, or -1 in the case the clock is invalid
 	 */
 	UFUNCTION(BlueprintCallable, Category = "Quartz Clock", meta = (WorldContext = "WorldContextObject"))
-	AUDIOMIXER_API float GetDurationOfQuantizationTypeInSeconds(const UObject* WorldContextObject, const EQuartzCommandQuantization& QuantizationType, float Multiplier = 1.0f);
+	float GetDurationOfQuantizationTypeInSeconds(const UObject* WorldContextObject, const EQuartzCommandQuantization& QuantizationType, float Multiplier = 1.0f);
 
 	//Retrieves a timestamp for the clock
 	UFUNCTION(BlueprintCallable, Category = "Quartz Clock Handle", meta = (WorldContext = "WorldContextObject"))
-	AUDIOMIXER_API FQuartzTransportTimeStamp GetCurrentTimestamp(const UObject* WorldContextObject);
+	FQuartzTransportTimeStamp GetCurrentTimestamp(const UObject* WorldContextObject);
 
 	// Returns the amount of time, in seconds, the clock has been running. Caution: due to latency, this will not be perfectly accurate
 	UFUNCTION(BlueprintCallable, Category = "Quartz Clock Handle", meta = (WorldContext = "WorldContextObject"))
-	AUDIOMIXER_API float GetEstimatedRunTime(const UObject* WorldContextObject);
+	float GetEstimatedRunTime(const UObject* WorldContextObject);
 
 	// "other" clock manipulation
 	UFUNCTION(BlueprintCallable, Category = "Quartz Clock", meta = (WorldContext = "WorldContextObject", AutoCreateRefTerm = "InDelegate", Keywords = "Transport, Counter"))
-	AUDIOMIXER_API void StartOtherClock(const UObject* WorldContextObject, FName OtherClockName, FQuartzQuantizationBoundary InQuantizationBoundary, const FOnQuartzCommandEventBP& InDelegate);
+	void StartOtherClock(const UObject* WorldContextObject, FName OtherClockName, FQuartzQuantizationBoundary InQuantizationBoundary, const FOnQuartzCommandEventBP& InDelegate);
 
 	// Metronome subscription
 	UFUNCTION(BlueprintCallable, Category = "Quartz Clock", meta = (WorldContext = "WorldContextObject"))
-	AUDIOMIXER_API void SubscribeToQuantizationEvent(const UObject* WorldContextObject, EQuartzCommandQuantization InQuantizationBoundary, const FOnQuartzMetronomeEventBP& OnQuantizationEvent, UQuartzClockHandle*& ClockHandle);
+	void SubscribeToQuantizationEvent(const UObject* WorldContextObject, EQuartzCommandQuantization InQuantizationBoundary, const FOnQuartzMetronomeEventBP& OnQuantizationEvent, UQuartzClockHandle*& ClockHandle);
 
 	UFUNCTION(BlueprintCallable, Category = "Quartz Clock", meta = (WorldContext = "WorldContextObject"))
-	AUDIOMIXER_API void SubscribeToAllQuantizationEvents(const UObject* WorldContextObject, const FOnQuartzMetronomeEventBP& OnQuantizationEvent, UQuartzClockHandle*& ClockHandle);
+	void SubscribeToAllQuantizationEvents(const UObject* WorldContextObject, const FOnQuartzMetronomeEventBP& OnQuantizationEvent, UQuartzClockHandle*& ClockHandle);
 
 	UFUNCTION(BlueprintCallable, Category = "Quartz Clock", meta = (WorldContext = "WorldContextObject"))
-	AUDIOMIXER_API void UnsubscribeFromTimeDivision(const UObject* WorldContextObject, EQuartzCommandQuantization InQuantizationBoundary, UQuartzClockHandle*& ClockHandle);
+	void UnsubscribeFromTimeDivision(const UObject* WorldContextObject, EQuartzCommandQuantization InQuantizationBoundary, UQuartzClockHandle*& ClockHandle);
 
 	UFUNCTION(BlueprintCallable, Category = "Quartz Clock", meta = (WorldContext = "WorldContextObject"))
-	AUDIOMIXER_API void UnsubscribeFromAllTimeDivisions(const UObject* WorldContextObject, UQuartzClockHandle*& ClockHandle);
+	void UnsubscribeFromAllTimeDivisions(const UObject* WorldContextObject, UQuartzClockHandle*& ClockHandle);
 
 	// Metronome Alteration (setters)
 	UFUNCTION(BlueprintCallable, Category = "Quantization", meta = (WorldContext = "WorldContextObject", AdvancedDisplay = "QuantizationBoundary, Delegate", AutoCreateRefTerm = "QuantizationBoundary, Delegate", Keywords = "BPM, Tempo"))
-	AUDIOMIXER_API void SetMillisecondsPerTick(const UObject* WorldContextObject, UPARAM(ref) const FQuartzQuantizationBoundary& QuantizationBoundary, const FOnQuartzCommandEventBP& Delegate, UQuartzClockHandle*& ClockHandle, float MillisecondsPerTick = 100.f);
+	void SetMillisecondsPerTick(const UObject* WorldContextObject, UPARAM(ref) const FQuartzQuantizationBoundary& QuantizationBoundary, const FOnQuartzCommandEventBP& Delegate, UQuartzClockHandle*& ClockHandle, float MillisecondsPerTick = 100.f);
 
 	UFUNCTION(BlueprintCallable, Category = "Quantization", meta = (WorldContext = "WorldContextObject", AdvancedDisplay = "QuantizationBoundary, Delegate", AutoCreateRefTerm = "QuantizationBoundary, Delegate", Keywords = "BPM, Tempo"))
-	AUDIOMIXER_API void SetTicksPerSecond(const UObject* WorldContextObject, UPARAM(ref) const FQuartzQuantizationBoundary& QuantizationBoundary, const FOnQuartzCommandEventBP& Delegate, UQuartzClockHandle*& ClockHandle, float TicksPerSecond = 10.f);
+	void SetTicksPerSecond(const UObject* WorldContextObject, UPARAM(ref) const FQuartzQuantizationBoundary& QuantizationBoundary, const FOnQuartzCommandEventBP& Delegate, UQuartzClockHandle*& ClockHandle, float TicksPerSecond = 10.f);
 
 	UFUNCTION(BlueprintCallable, Category = "Quantization", meta = (WorldContext = "WorldContextObject", AdvancedDisplay = "QuantizationBoundary, Delegate", AutoCreateRefTerm = "QuantizationBoundary, Delegate", Keywords = "BPM, Tempo"))
-	AUDIOMIXER_API void SetSecondsPerTick(const UObject* WorldContextObject, UPARAM(ref) const FQuartzQuantizationBoundary& QuantizationBoundary, const FOnQuartzCommandEventBP& Delegate, UQuartzClockHandle*& ClockHandle, float SecondsPerTick = 0.25f);
+	void SetSecondsPerTick(const UObject* WorldContextObject, UPARAM(ref) const FQuartzQuantizationBoundary& QuantizationBoundary, const FOnQuartzCommandEventBP& Delegate, UQuartzClockHandle*& ClockHandle, float SecondsPerTick = 0.25f);
 
 	UFUNCTION(BlueprintCallable, Category = "Quantization", meta = (WorldContext = "WorldContextObject", AdvancedDisplay = "QuantizationBoundary, Delegate", AutoCreateRefTerm = "QuantizationBoundary, Delegate", Keywords = "BPM, Tempo"))
-	AUDIOMIXER_API void SetThirtySecondNotesPerMinute(const UObject* WorldContextObject, UPARAM(ref) const FQuartzQuantizationBoundary& QuantizationBoundary, const FOnQuartzCommandEventBP& Delegate, UQuartzClockHandle*& ClockHandle, float ThirtySecondsNotesPerMinute = 960.f);
+	void SetThirtySecondNotesPerMinute(const UObject* WorldContextObject, UPARAM(ref) const FQuartzQuantizationBoundary& QuantizationBoundary, const FOnQuartzCommandEventBP& Delegate, UQuartzClockHandle*& ClockHandle, float ThirtySecondsNotesPerMinute = 960.f);
 
 	UFUNCTION(BlueprintCallable, Category = "Quantization", meta = (WorldContext = "WorldContextObject", AdvancedDisplay = "QuantizationBoundary, Delegate", AutoCreateRefTerm = "QuantizationBoundary, Delegate", Keywords = "BPM, Tempo"))
-	AUDIOMIXER_API void SetBeatsPerMinute(const UObject* WorldContextObject, UPARAM(ref) const FQuartzQuantizationBoundary& QuantizationBoundary, const FOnQuartzCommandEventBP& Delegate, UQuartzClockHandle*& ClockHandle, float BeatsPerMinute = 60.f);
+	void SetBeatsPerMinute(const UObject* WorldContextObject, UPARAM(ref) const FQuartzQuantizationBoundary& QuantizationBoundary, const FOnQuartzCommandEventBP& Delegate, UQuartzClockHandle*& ClockHandle, float BeatsPerMinute = 60.f);
 
 	// Metronome getters
 	UFUNCTION(BlueprintCallable, Category = "Quantization", meta = (WorldContext = "WorldContextObject", AutoCreateRefTerm = "InDelegate", Keywords = "BPM, Tempo"))
-	AUDIOMIXER_API float GetMillisecondsPerTick(const UObject* WorldContextObject) const;
+	float GetMillisecondsPerTick(const UObject* WorldContextObject) const;
 
 	UFUNCTION(BlueprintCallable, Category = "Quantization", meta = (WorldContext = "WorldContextObject", AutoCreateRefTerm = "InDelegate", Keywords = "BPM, Tempo"))
-	AUDIOMIXER_API float GetTicksPerSecond(const UObject* WorldContextObject) const;
+	float GetTicksPerSecond(const UObject* WorldContextObject) const;
 
 	UFUNCTION(BlueprintCallable, Category = "Quantization", meta = (WorldContext = "WorldContextObject", AutoCreateRefTerm = "InDelegate", Keywords = "BPM, Tempo"))
-	AUDIOMIXER_API float GetSecondsPerTick(const UObject* WorldContextObject) const;
+	float GetSecondsPerTick(const UObject* WorldContextObject) const;
 
 	UFUNCTION(BlueprintCallable, Category = "Quantization", meta = (WorldContext = "WorldContextObject", AutoCreateRefTerm = "InDelegate", Keywords = "BPM, Tempo"))
-	AUDIOMIXER_API float GetThirtySecondNotesPerMinute(const UObject* WorldContextObject) const;
+	float GetThirtySecondNotesPerMinute(const UObject* WorldContextObject) const;
 
 	UFUNCTION(BlueprintCallable, Category = "Quantization", meta = (WorldContext = "WorldContextObject", AutoCreateRefTerm = "InDelegate", Keywords = "BPM, Tempo"))
-	AUDIOMIXER_API float GetBeatsPerMinute(const UObject* WorldContextObject) const;
-
-	/**
-	 * Returns the current progress until the next occurrence of the provided musical duration as a float value from 0 (previous beat) to 1 (next beat).
-	 * This is useful for indexing into curves to animate parameters to musical time.
-	 * Ms and Phase offsets are combined internally.
-	 */
-	UFUNCTION(BlueprintCallable, Category = "Quantization", meta = ( AutoCreateRefTerm = "PhaseOffset", Keywords = "BPM, Tempo"))
-	AUDIOMIXER_API float GetBeatProgressPercent(EQuartzCommandQuantization QuantizationBoundary = EQuartzCommandQuantization::Beat, float PhaseOffset = 0.f, float MsOffset = 0.f);
+	float GetBeatsPerMinute(const UObject* WorldContextObject) const;
 
 	// todo: un-comment when metronome events support the offset
 	// Set how early we would like to receive Metronome (not yet supported) and "About To Start" Delegates. (all other command delegates will execute as normal)
@@ -150,9 +140,9 @@ public:
 	// void SetNotificationAnticipationAmountAsMusicalDuration(const UObject* WorldContextObject, UQuartzClockHandle*& ClockHandle, const EQuartzCommandQuantization MusicalDuration = EQuartzCommandQuantization::QuarterNote, const double Multiplier = 1.0);
 
 // End Blueprint Interface
-	AUDIOMIXER_API void QueueQuantizedSound(const UObject* WorldContextObject, UQuartzClockHandle*& ClockHandle, const FAudioComponentCommandInfo& AudioComponentData, const FOnQuartzCommandEventBP& InDelegate, const FQuartzQuantizationBoundary& InTargetBoundary);
+	void QueueQuantizedSound(const UObject* WorldContextObject, UQuartzClockHandle*& ClockHandle, const FAudioComponentCommandInfo& AudioComponentData, const FOnQuartzCommandEventBP& InDelegate, const FQuartzQuantizationBoundary& InTargetBoundary);
 
-	AUDIOMIXER_API UQuartzClockHandle* SubscribeToClock(const UObject* WorldContextObject, FName ClockName, Audio::FQuartzClockProxy const* InHandlePtr = nullptr);
+	UQuartzClockHandle* SubscribeToClock(const UObject* WorldContextObject, FName ClockName, Audio::FQuartzClockProxy const* InHandlePtr = nullptr);
 
 	FName GetClockName() const { return CurrentClockId; }
 
@@ -167,10 +157,10 @@ public:
 	UE_DEPRECATED(5.1, "This function should not be called directly, and the original functionality has been moved into FQuartzTickable")
 	virtual void ProcessCommand(const Audio::FQuartzMetronomeDelegateData& Data) override {};
 
-	AUDIOMIXER_API bool GetCurrentTickRate(const UObject* WorldContextObject, Audio::FQuartzClockTickRate& OutTickRate) const;
+	bool GetCurrentTickRate(const UObject* WorldContextObject, Audio::FQuartzClockTickRate& OutTickRate) const;
 
 private:
-	AUDIOMIXER_API void SetTickRateInternal(const FQuartzQuantizationBoundary& InQuantizationBoundary, const FOnQuartzCommandEventBP& InDelegate, const Audio::FQuartzClockTickRate& NewTickRate);
+	void SetTickRateInternal(const FQuartzQuantizationBoundary& InQuantizationBoundary, const FOnQuartzCommandEventBP& InDelegate, const Audio::FQuartzClockTickRate& NewTickRate);
 
 	Audio::FQuartzClockProxy RawHandle;
 
diff --git a/Engine/Source/Runtime/Engine/Classes/Sound/QuartzSubscription.cpp b/Engine/Source/Runtime/Engine/Classes/Sound/QuartzSubscription.cpp
--- a/Engine/Source/Runtime/Engine/Classes/Sound/QuartzSubscription.cpp
+++ b/Engine/Source/Runtime/Engine/Classes/Sound/QuartzSubscription.cpp
@@ -22,13 +22,13 @@ namespace Audio
 	}
 } // namespace Audio
 
+FQuartzTickableObject::FQuartzTickableObject()
+{}
+
 PRAGMA_DISABLE_DEPRECATION_WARNINGS
 FQuartzTickableObject::~FQuartzTickableObject()
 {
-	if(const TSharedPtr<FQuartzTickableObjectsManager> ObjManagerPtr = TickableObjectManagerPtr.Pin())
-	{
-		ObjManagerPtr->UnsubscribeFromQuartzTick(this);
-	}
+	QuartzUnsubscribe();
 }
 PRAGMA_ENABLE_DEPRECATION_WARNINGS
 
@@ -48,19 +48,16 @@ FQuartzTickableObject* FQuartzTickableObject::Init(UWorld* InWorldPtr)
 	PRAGMA_ENABLE_DEPRECATION_WARNINGS
 
 	UQuartzSubsystem* QuartzSubsystemPtr = UQuartzSubsystem::Get(InWorldPtr);
-	if(ensure(QuartzSubsystemPtr))
-	{
-		TickableObjectManagerPtr = QuartzSubsystemPtr->GetTickableObjectManager();
-
-		if(TSharedPtr<FQuartzTickableObjectsManager> ObjManagerPtr = TickableObjectManagerPtr.Pin())
-		{
-			ObjManagerPtr->SubscribeToQuartzTick(this);
-		}
-	}
+	QuartzSubscriptionToken.Subscribe(this, QuartzSubsystemPtr);
 
 	return this;
 }
 
+void FQuartzTickableObject::QuartzUnsubscribe()
+{
+	QuartzSubscriptionToken.Unsubscribe();
+}
+
 int32 FQuartzTickableObject::AddCommandDelegate(const FOnQuartzCommandEventBP& InDelegate)
 {
 	const int32 Num = QuantizedCommandDelegates.Num();
@@ -90,8 +87,8 @@ UQuartzSubsystem* FQuartzTickableObject::GetQuartzSubsystem() const
 void FQuartzTickableObject::ExecCommand(const Audio::FQuartzQuantizedCommandDelegateData& Data)
 {
 	checkSlow(Data.DelegateSubType < EQuartzCommandDelegateSubType::Count);
-
-	if(const TSharedPtr<FQuartzTickableObjectsManager> ObjManagerPtr = TickableObjectManagerPtr.Pin())
+	
+	if(const TSharedPtr<FQuartzTickableObjectsManager> ObjManagerPtr = QuartzSubscriptionToken.GetTickableObjectManager())
 	{
 		ObjManagerPtr->PushLatencyTrackerResult(Data.RequestRecieved());
 	}
@@ -132,7 +129,7 @@ void FQuartzTickableObject::ExecCommand(const Audio::FQuartzQuantizedCommandDele
 
 void FQuartzTickableObject::ExecCommand(const Audio::FQuartzMetronomeDelegateData& Data)
 {
-	if(const TSharedPtr<FQuartzTickableObjectsManager> ObjManagerPtr = TickableObjectManagerPtr.Pin())
+	if(const TSharedPtr<FQuartzTickableObjectsManager> ObjManagerPtr = QuartzSubscriptionToken.GetTickableObjectManager())
 	{
 		ObjManagerPtr->PushLatencyTrackerResult(Data.RequestRecieved());
 	}
diff --git a/Engine/Source/Runtime/Engine/Classes/Sound/QuartzSubscription.h b/Engine/Source/Runtime/Engine/Classes/Sound/QuartzSubscription.h
--- a/Engine/Source/Runtime/Engine/Classes/Sound/QuartzSubscription.h
+++ b/Engine/Source/Runtime/Engine/Classes/Sound/QuartzSubscription.h
@@ -2,6 +2,7 @@
 
 #pragma once
 
+#include "QuartzSubscriptionToken.h"
 #include "Sound/QuartzQuantizationUtilities.h"
 #include "Containers/ConsumeAllMpmcQueue.h"
 
@@ -126,7 +127,7 @@ class FQuartzTickableObject
 {
 public:
 	// ctor
-	FQuartzTickableObject() {}
+	FQuartzTickableObject();
 
 	// explicitly defaulted ctors (to disable deprecation warnings in compiler-generated functions)
     PRAGMA_DISABLE_DEPRECATION_WARNINGS
@@ -149,8 +150,7 @@ public:
 
 	ENGINE_API void AddMetronomeBpDelegate(EQuartzCommandQuantization InQuantizationBoundary, const FOnQuartzMetronomeEventBP& OnQuantizationEvent);
 
-
-	bool IsInitialized() const { return TickableObjectManagerPtr.IsValid(); }
+	bool IsInitialized() const { return QuartzSubscriptionToken.IsSubscribed(); }
 
 	ENGINE_API Audio::FQuartzGameThreadSubscriber GetQuartzSubscriber();
 
@@ -178,6 +178,8 @@ protected:
 	ENGINE_API void SetNotificationAnticipationAmountMilliseconds(const double Milliseconds);
 	ENGINE_API void SetNotificationAnticipationAmountMusicalDuration(const EQuartzCommandQuantization Duration,  const double Multiplier);
 
+	void QuartzUnsubscribe();
+
 private:
 	struct FMetronomeDelegateGameThreadData
 	{
@@ -203,7 +205,7 @@ public:
 
 private:
 	Audio::FQuartzOffset NotificationOffset;
-	TWeakPtr<FQuartzTickableObjectsManager> TickableObjectManagerPtr;
 
+	FQuartzSubscriptionToken QuartzSubscriptionToken;
 }; // class FQuartzTickableObject
 
diff --git a/Engine/Source/Runtime/Engine/Classes/Sound/QuartzSubscriptionToken.cpp b/Engine/Source/Runtime/Engine/Classes/Sound/QuartzSubscriptionToken.cpp
new file mode 100644
--- /dev/null
+++ b/Engine/Source/Runtime/Engine/Classes/Sound/QuartzSubscriptionToken.cpp
@@ -0,0 +1,74 @@
+﻿#include "QuartzSubscriptionToken.h"
+
+#include "Quartz/QuartzSubsystem.h"
+#include "Sound/QuartzQuantizationUtilities.h"
+
+FQuartzSubscriptionToken::FQuartzSubscriptionToken()
+	: SubscribingObject(nullptr)
+	{}
+
+FQuartzSubscriptionToken::~FQuartzSubscriptionToken()
+{
+	Unsubscribe();
+}
+
+void FQuartzSubscriptionToken::Subscribe(FQuartzTickableObject* Subscriber, UQuartzSubsystem* QuartzSubsystem)
+{
+	if(IsSubscribed())
+	{
+		return;
+	}
+
+	if(ensure(Subscriber))
+	{
+		SubscribingObject = Subscriber;
+	}
+	else
+	{
+		UE_LOG(LogAudioQuartz, Error, TEXT("Invalid subscribing object."));
+		return;
+	}
+	
+	if(ensure(QuartzSubsystem))
+	{
+		TickableObjectManagerPtr = QuartzSubsystem->GetTickableObjectManager();
+	}
+	else
+	{
+		UE_LOG(LogAudioQuartz, Error, TEXT("Unable to obtain Tickable Object Manager."));
+		return;
+	}
+	
+	TSharedPtr<FQuartzTickableObjectsManager> ObjManagerPtr = GetTickableObjectManager();
+	if(ObjManagerPtr.IsValid())
+	{
+		ObjManagerPtr->SubscribeToQuartzTick(SubscribingObject);
+	}
+	else
+	{
+		UE_LOG(LogAudioQuartz, Warning, TEXT("Unable to subscribe to Quartz tick. Tickable Object Manager pointer is not valid."));
+	}
+}
+
+void FQuartzSubscriptionToken::Unsubscribe()
+{
+	const TSharedPtr<FQuartzTickableObjectsManager> ObjManagerPtr = GetTickableObjectManager();
+	if(ObjManagerPtr.IsValid())
+	{
+		ObjManagerPtr->UnsubscribeFromQuartzTick(SubscribingObject);
+	}
+
+	SubscribingObject = nullptr;
+	TickableObjectManagerPtr = nullptr;
+}
+
+TSharedPtr<FQuartzTickableObjectsManager> FQuartzSubscriptionToken::GetTickableObjectManager()
+{
+	TSharedPtr<FQuartzTickableObjectsManager> ObjManagerPtr = TickableObjectManagerPtr.Pin();
+	return ObjManagerPtr;
+}
+
+bool FQuartzSubscriptionToken::IsSubscribed() const
+{
+	return TickableObjectManagerPtr.IsValid();
+}
\ No newline at end of file
diff --git a/Engine/Source/Runtime/Engine/Classes/Sound/QuartzSubscriptionToken.h b/Engine/Source/Runtime/Engine/Classes/Sound/QuartzSubscriptionToken.h
new file mode 100644
--- /dev/null
+++ b/Engine/Source/Runtime/Engine/Classes/Sound/QuartzSubscriptionToken.h
@@ -0,0 +1,24 @@
+﻿// Copyright Epic Games, Inc. All Rights Reserved.
+
+#pragma once
+#include "Templates/SharedPointer.h"
+
+class UQuartzSubsystem;
+struct FQuartzTickableObjectsManager;
+
+class FQuartzSubscriptionToken
+{
+public:
+	FQuartzSubscriptionToken();
+	~FQuartzSubscriptionToken();
+	
+	void Subscribe(FQuartzTickableObject* Subscriber, UQuartzSubsystem* QuartzSubsystem);
+	void Unsubscribe();
+
+	TSharedPtr<FQuartzTickableObjectsManager> GetTickableObjectManager();
+	bool IsSubscribed() const;
+	
+private:
+	FQuartzTickableObject* SubscribingObject;
+	TWeakPtr<FQuartzTickableObjectsManager> TickableObjectManagerPtr;
+};
diff --git a/Engine/Source/Runtime/Engine/Private/Components/AudioComponent.cpp b/Engine/Source/Runtime/Engine/Private/Components/AudioComponent.cpp
--- a/Engine/Source/Runtime/Engine/Private/Components/AudioComponent.cpp
+++ b/Engine/Source/Runtime/Engine/Private/Components/AudioComponent.cpp
@@ -308,6 +308,7 @@ void UAudioComponent::OnUnregister()
 void UAudioComponent::EndPlay(const EEndPlayReason::Type EndPlayReason)
 {
 	Super::EndPlay(EndPlayReason);
+	QuartzUnsubscribe();
 }
 
 const UObject* UAudioComponent::AdditionalStatObject() const
