Fixing compile error after bad merge conflict resolve

diff --git a/Engine/Source/Runtime/AudioMixer/Public/Quartz/AudioMixerClockHandle.h b/Engine/Source/Runtime/AudioMixer/Public/Quartz/AudioMixerClockHandle.h
--- a/Engine/Source/Runtime/AudioMixer/Public/Quartz/AudioMixerClockHandle.h
+++ b/Engine/Source/Runtime/AudioMixer/Public/Quartz/AudioMixerClockHandle.h
@@ -3,7 +3,9 @@
 #pragma once
 
 #include "CoreMinimal.h"
-#include "Containers/Queue.h"
+#include "UObject/ObjectMacros.h"
+#include "Curves/RichCurve.h"
+#include "Curves/CurveBase.h"
 #include "Sound/QuartzQuantizationUtilities.h"
 #include "Sound/QuartzSubscription.h"
 #include "Quartz/QuartzSubsystem.h"
@@ -11,9 +13,6 @@
 
 #include "AudioMixerClockHandle.generated.h"
 
-
-
-
 /**
  *  This class is a BP / Game thread wrapper around FQuartzClockProxy
  *	(to talk to the underlying clock)
@@ -24,45 +23,48 @@
  *  It can subscribe to Quantized Event & Metronome delegates to synchronize
  *  gameplay & VFX to Quartz events fired from the Audio Engine
  */
-UCLASS(BlueprintType, Blueprintable, Transient, ClassGroup = Quartz, meta = (BlueprintSpawnableComponent))
-class AUDIOMIXER_API UQuartzClockHandle : public UObject, public FQuartzTickableObject
+UCLASS(BlueprintType, Blueprintable, Transient, ClassGroup = Quartz, meta = (BlueprintSpawnableComponent), MinimalAPI)
+class UQuartzClockHandle : public UObject, public FQuartzTickableObject
 {
 	GENERATED_BODY()
 
 public:
 	// ctor
-	UQuartzClockHandle();
+	AUDIOMIXER_API UQuartzClockHandle();
 
 	// dtor
-	AUDIOMIXER_API virtual ~UQuartzClockHandle();
+	AUDIOMIXER_API ~UQuartzClockHandle();
 
 	// begin UObject interface
-	void BeginDestroy() override;
+	AUDIOMIXER_API void BeginDestroy() override;
 	// end UObject interface
 
 // Begin Blueprint Interface
 
 	// Clock manipulation
 	UFUNCTION(BlueprintCallable, Category = "Quartz Clock", meta = (WorldContext = "WorldContextObject"))
-	void StartClock(const UObject* WorldContextObject, UQuartzClockHandle*& ClockHandle);
+	AUDIOMIXER_API void StartClock(const UObject* WorldContextObject, UQuartzClockHandle*& ClockHandle);
 
 	UFUNCTION(BlueprintCallable, Category = "Quartz Subsystem", meta = (WorldContext = "WorldContextObject"))
-	void StopClock(const UObject* WorldContextObject, bool CancelPendingEvents, UQuartzClockHandle*& ClockHandle);
+	AUDIOMIXER_API void StopClock(const UObject* WorldContextObject, bool CancelPendingEvents, UQuartzClockHandle*& ClockHandle);
 
 	UFUNCTION(BlueprintCallable, Category = "Quartz Subsystem", meta = (WorldContext = "WorldContextObject"))
-	void PauseClock(const UObject* WorldContextObject, UQuartzClockHandle*& ClockHandle);
+	AUDIOMIXER_API void PauseClock(const UObject* WorldContextObject, UQuartzClockHandle*& ClockHandle);
 
 	UFUNCTION(BlueprintCallable, Category = "Quartz Clock", meta = (WorldContext = "WorldContextObject"))
-	void ResumeClock(const UObject* WorldContextObject, UQuartzClockHandle*& ClockHandle);
+	AUDIOMIXER_API void ResumeClock(const UObject* WorldContextObject, UQuartzClockHandle*& ClockHandle);
 
 	UFUNCTION(BlueprintCallable, Category = "Quartz Clock", meta = (WorldContext = "WorldContextObject", AutoCreateRefTerm = "InDelegate", Keywords = "Transport, Counter"))
-	void ResetTransport(const UObject* WorldContextObject, const FOnQuartzCommandEventBP& InDelegate);
+	AUDIOMIXER_API void ResetTransport(const UObject* WorldContextObject, const FOnQuartzCommandEventBP& InDelegate);
 
 	UFUNCTION(BlueprintCallable, Category = "Quartz Clock", meta = (WorldContext = "WorldContextObject", AutoCreateRefTerm = "InDelegate", Keywords = "Transport, Counter"))
-	void ResetTransportQuantized(const UObject* WorldContextObject, FQuartzQuantizationBoundary InQuantizationBoundary, const FOnQuartzCommandEventBP& InDelegate, UQuartzClockHandle*& ClockHandle);
+	AUDIOMIXER_API void ResetTransportQuantized(const UObject* WorldContextObject, FQuartzQuantizationBoundary InQuantizationBoundary, const FOnQuartzCommandEventBP& InDelegate, UQuartzClockHandle*& ClockHandle);
 
 	UFUNCTION(BlueprintCallable, Category = "Quartz Clock", meta = (WorldContext = "WorldContextObject", Keywords = "Transport, Counter"))
-	bool IsClockRunning(const UObject* WorldContextObject);
+	AUDIOMIXER_API bool IsClockRunning(const UObject* WorldContextObject);
+
+	UFUNCTION(BlueprintCallable, Category = "Quartz Clock", meta = (WorldContext = "WorldContextObject", AutoCreateRefTerm = "InDelegate", Keywords = "Transport, Counter"))
+	AUDIOMIXER_API void NotifyOnQuantizationBoundary(const UObject* WorldContextObject, FQuartzQuantizationBoundary InQuantizationBoundary, const FOnQuartzCommandEventBP& InDelegate, float InMsOffset = 0.f);
 
 	/** Returns the duration in seconds of the given Quantization Type
 	 *
@@ -71,64 +73,72 @@ public:
 	 * @return The duration, in seconds, of a multiplier amount of the Quantization Type, or -1 in the case the clock is invalid
 	 */
 	UFUNCTION(BlueprintCallable, Category = "Quartz Clock", meta = (WorldContext = "WorldContextObject"))
-	float GetDurationOfQuantizationTypeInSeconds(const UObject* WorldContextObject, const EQuartzCommandQuantization& QuantizationType, float Multiplier = 1.0f);
+	AUDIOMIXER_API float GetDurationOfQuantizationTypeInSeconds(const UObject* WorldContextObject, const EQuartzCommandQuantization& QuantizationType, float Multiplier = 1.0f);
 
 	//Retrieves a timestamp for the clock
 	UFUNCTION(BlueprintCallable, Category = "Quartz Clock Handle", meta = (WorldContext = "WorldContextObject"))
-	FQuartzTransportTimeStamp GetCurrentTimestamp(const UObject* WorldContextObject);
+	AUDIOMIXER_API FQuartzTransportTimeStamp GetCurrentTimestamp(const UObject* WorldContextObject);
 
 	// Returns the amount of time, in seconds, the clock has been running. Caution: due to latency, this will not be perfectly accurate
 	UFUNCTION(BlueprintCallable, Category = "Quartz Clock Handle", meta = (WorldContext = "WorldContextObject"))
-	float GetEstimatedRunTime(const UObject* WorldContextObject);
+	AUDIOMIXER_API float GetEstimatedRunTime(const UObject* WorldContextObject);
 
 	// "other" clock manipulation
 	UFUNCTION(BlueprintCallable, Category = "Quartz Clock", meta = (WorldContext = "WorldContextObject", AutoCreateRefTerm = "InDelegate", Keywords = "Transport, Counter"))
-	void StartOtherClock(const UObject* WorldContextObject, FName OtherClockName, FQuartzQuantizationBoundary InQuantizationBoundary, const FOnQuartzCommandEventBP& InDelegate);
+	AUDIOMIXER_API void StartOtherClock(const UObject* WorldContextObject, FName OtherClockName, FQuartzQuantizationBoundary InQuantizationBoundary, const FOnQuartzCommandEventBP& InDelegate);
 
 	// Metronome subscription
 	UFUNCTION(BlueprintCallable, Category = "Quartz Clock", meta = (WorldContext = "WorldContextObject"))
-	void SubscribeToQuantizationEvent(const UObject* WorldContextObject, EQuartzCommandQuantization InQuantizationBoundary, const FOnQuartzMetronomeEventBP& OnQuantizationEvent, UQuartzClockHandle*& ClockHandle);
+	AUDIOMIXER_API void SubscribeToQuantizationEvent(const UObject* WorldContextObject, EQuartzCommandQuantization InQuantizationBoundary, const FOnQuartzMetronomeEventBP& OnQuantizationEvent, UQuartzClockHandle*& ClockHandle);
 
 	UFUNCTION(BlueprintCallable, Category = "Quartz Clock", meta = (WorldContext = "WorldContextObject"))
-	void SubscribeToAllQuantizationEvents(const UObject* WorldContextObject, const FOnQuartzMetronomeEventBP& OnQuantizationEvent, UQuartzClockHandle*& ClockHandle);
+	AUDIOMIXER_API void SubscribeToAllQuantizationEvents(const UObject* WorldContextObject, const FOnQuartzMetronomeEventBP& OnQuantizationEvent, UQuartzClockHandle*& ClockHandle);
 
 	UFUNCTION(BlueprintCallable, Category = "Quartz Clock", meta = (WorldContext = "WorldContextObject"))
-	void UnsubscribeFromTimeDivision(const UObject* WorldContextObject, EQuartzCommandQuantization InQuantizationBoundary, UQuartzClockHandle*& ClockHandle);
+	AUDIOMIXER_API void UnsubscribeFromTimeDivision(const UObject* WorldContextObject, EQuartzCommandQuantization InQuantizationBoundary, UQuartzClockHandle*& ClockHandle);
 
 	UFUNCTION(BlueprintCallable, Category = "Quartz Clock", meta = (WorldContext = "WorldContextObject"))
-	void UnsubscribeFromAllTimeDivisions(const UObject* WorldContextObject, UQuartzClockHandle*& ClockHandle);
+	AUDIOMIXER_API void UnsubscribeFromAllTimeDivisions(const UObject* WorldContextObject, UQuartzClockHandle*& ClockHandle);
 
 	// Metronome Alteration (setters)
 	UFUNCTION(BlueprintCallable, Category = "Quantization", meta = (WorldContext = "WorldContextObject", AdvancedDisplay = "QuantizationBoundary, Delegate", AutoCreateRefTerm = "QuantizationBoundary, Delegate", Keywords = "BPM, Tempo"))
-	void SetMillisecondsPerTick(const UObject* WorldContextObject, UPARAM(ref) const FQuartzQuantizationBoundary& QuantizationBoundary, const FOnQuartzCommandEventBP& Delegate, UQuartzClockHandle*& ClockHandle, float MillisecondsPerTick = 100.f);
+	AUDIOMIXER_API void SetMillisecondsPerTick(const UObject* WorldContextObject, UPARAM(ref) const FQuartzQuantizationBoundary& QuantizationBoundary, const FOnQuartzCommandEventBP& Delegate, UQuartzClockHandle*& ClockHandle, float MillisecondsPerTick = 100.f);
 
 	UFUNCTION(BlueprintCallable, Category = "Quantization", meta = (WorldContext = "WorldContextObject", AdvancedDisplay = "QuantizationBoundary, Delegate", AutoCreateRefTerm = "QuantizationBoundary, Delegate", Keywords = "BPM, Tempo"))
-	void SetTicksPerSecond(const UObject* WorldContextObject, UPARAM(ref) const FQuartzQuantizationBoundary& QuantizationBoundary, const FOnQuartzCommandEventBP& Delegate, UQuartzClockHandle*& ClockHandle, float TicksPerSecond = 10.f);
+	AUDIOMIXER_API void SetTicksPerSecond(const UObject* WorldContextObject, UPARAM(ref) const FQuartzQuantizationBoundary& QuantizationBoundary, const FOnQuartzCommandEventBP& Delegate, UQuartzClockHandle*& ClockHandle, float TicksPerSecond = 10.f);
 
 	UFUNCTION(BlueprintCallable, Category = "Quantization", meta = (WorldContext = "WorldContextObject", AdvancedDisplay = "QuantizationBoundary, Delegate", AutoCreateRefTerm = "QuantizationBoundary, Delegate", Keywords = "BPM, Tempo"))
-	void SetSecondsPerTick(const UObject* WorldContextObject, UPARAM(ref) const FQuartzQuantizationBoundary& QuantizationBoundary, const FOnQuartzCommandEventBP& Delegate, UQuartzClockHandle*& ClockHandle, float SecondsPerTick = 0.25f);
+	AUDIOMIXER_API void SetSecondsPerTick(const UObject* WorldContextObject, UPARAM(ref) const FQuartzQuantizationBoundary& QuantizationBoundary, const FOnQuartzCommandEventBP& Delegate, UQuartzClockHandle*& ClockHandle, float SecondsPerTick = 0.25f);
 
 	UFUNCTION(BlueprintCallable, Category = "Quantization", meta = (WorldContext = "WorldContextObject", AdvancedDisplay = "QuantizationBoundary, Delegate", AutoCreateRefTerm = "QuantizationBoundary, Delegate", Keywords = "BPM, Tempo"))
-	void SetThirtySecondNotesPerMinute(const UObject* WorldContextObject, UPARAM(ref) const FQuartzQuantizationBoundary& QuantizationBoundary, const FOnQuartzCommandEventBP& Delegate, UQuartzClockHandle*& ClockHandle, float ThirtySecondsNotesPerMinute = 960.f);
+	AUDIOMIXER_API void SetThirtySecondNotesPerMinute(const UObject* WorldContextObject, UPARAM(ref) const FQuartzQuantizationBoundary& QuantizationBoundary, const FOnQuartzCommandEventBP& Delegate, UQuartzClockHandle*& ClockHandle, float ThirtySecondsNotesPerMinute = 960.f);
 
 	UFUNCTION(BlueprintCallable, Category = "Quantization", meta = (WorldContext = "WorldContextObject", AdvancedDisplay = "QuantizationBoundary, Delegate", AutoCreateRefTerm = "QuantizationBoundary, Delegate", Keywords = "BPM, Tempo"))
-	void SetBeatsPerMinute(const UObject* WorldContextObject, UPARAM(ref) const FQuartzQuantizationBoundary& QuantizationBoundary, const FOnQuartzCommandEventBP& Delegate, UQuartzClockHandle*& ClockHandle, float BeatsPerMinute = 60.f);
+	AUDIOMIXER_API void SetBeatsPerMinute(const UObject* WorldContextObject, UPARAM(ref) const FQuartzQuantizationBoundary& QuantizationBoundary, const FOnQuartzCommandEventBP& Delegate, UQuartzClockHandle*& ClockHandle, float BeatsPerMinute = 60.f);
 
 	// Metronome getters
 	UFUNCTION(BlueprintCallable, Category = "Quantization", meta = (WorldContext = "WorldContextObject", AutoCreateRefTerm = "InDelegate", Keywords = "BPM, Tempo"))
-	float GetMillisecondsPerTick(const UObject* WorldContextObject) const;
+	AUDIOMIXER_API float GetMillisecondsPerTick(const UObject* WorldContextObject) const;
 
 	UFUNCTION(BlueprintCallable, Category = "Quantization", meta = (WorldContext = "WorldContextObject", AutoCreateRefTerm = "InDelegate", Keywords = "BPM, Tempo"))
-	float GetTicksPerSecond(const UObject* WorldContextObject) const;
+	AUDIOMIXER_API float GetTicksPerSecond(const UObject* WorldContextObject) const;
 
 	UFUNCTION(BlueprintCallable, Category = "Quantization", meta = (WorldContext = "WorldContextObject", AutoCreateRefTerm = "InDelegate", Keywords = "BPM, Tempo"))
-	float GetSecondsPerTick(const UObject* WorldContextObject) const;
+	AUDIOMIXER_API float GetSecondsPerTick(const UObject* WorldContextObject) const;
 
 	UFUNCTION(BlueprintCallable, Category = "Quantization", meta = (WorldContext = "WorldContextObject", AutoCreateRefTerm = "InDelegate", Keywords = "BPM, Tempo"))
-	float GetThirtySecondNotesPerMinute(const UObject* WorldContextObject) const;
+	AUDIOMIXER_API float GetThirtySecondNotesPerMinute(const UObject* WorldContextObject) const;
 
 	UFUNCTION(BlueprintCallable, Category = "Quantization", meta = (WorldContext = "WorldContextObject", AutoCreateRefTerm = "InDelegate", Keywords = "BPM, Tempo"))
-	float GetBeatsPerMinute(const UObject* WorldContextObject) const;
+	AUDIOMIXER_API float GetBeatsPerMinute(const UObject* WorldContextObject) const;
+
+	/**
+	 * Returns the current progress until the next occurrence of the provided musical duration as a float value from 0 (previous beat) to 1 (next beat).
+	 * This is useful for indexing into curves to animate parameters to musical time.
+	 * Ms and Phase offsets are combined internally.
+	 */
+	UFUNCTION(BlueprintCallable, Category = "Quantization", meta = ( AutoCreateRefTerm = "PhaseOffset", Keywords = "BPM, Tempo"))
+	AUDIOMIXER_API float GetBeatProgressPercent(EQuartzCommandQuantization QuantizationBoundary = EQuartzCommandQuantization::Beat, float PhaseOffset = 0.f, float MsOffset = 0.f);
 
 	// todo: un-comment when metronome events support the offset
 	// Set how early we would like to receive Metronome (not yet supported) and "About To Start" Delegates. (all other command delegates will execute as normal)
@@ -140,9 +150,9 @@ public:
 	// void SetNotificationAnticipationAmountAsMusicalDuration(const UObject* WorldContextObject, UQuartzClockHandle*& ClockHandle, const EQuartzCommandQuantization MusicalDuration = EQuartzCommandQuantization::QuarterNote, const double Multiplier = 1.0);
 
 // End Blueprint Interface
-	void QueueQuantizedSound(const UObject* WorldContextObject, UQuartzClockHandle*& ClockHandle, const FAudioComponentCommandInfo& AudioComponentData, const FOnQuartzCommandEventBP& InDelegate, const FQuartzQuantizationBoundary& InTargetBoundary);
+	AUDIOMIXER_API void QueueQuantizedSound(const UObject* WorldContextObject, UQuartzClockHandle*& ClockHandle, const FAudioComponentCommandInfo& AudioComponentData, const FOnQuartzCommandEventBP& InDelegate, const FQuartzQuantizationBoundary& InTargetBoundary);
 
-	UQuartzClockHandle* SubscribeToClock(const UObject* WorldContextObject, FName ClockName, Audio::FQuartzClockProxy const* InHandlePtr = nullptr);
+	AUDIOMIXER_API UQuartzClockHandle* SubscribeToClock(const UObject* WorldContextObject, FName ClockName, Audio::FQuartzClockProxy const* InHandlePtr = nullptr);
 
 	FName GetClockName() const { return CurrentClockId; }
 
@@ -157,10 +167,10 @@ public:
 	UE_DEPRECATED(5.1, "This function should not be called directly, and the original functionality has been moved into FQuartzTickable")
 	virtual void ProcessCommand(const Audio::FQuartzMetronomeDelegateData& Data) override {};
 
-	bool GetCurrentTickRate(const UObject* WorldContextObject, Audio::FQuartzClockTickRate& OutTickRate) const;
+	AUDIOMIXER_API bool GetCurrentTickRate(const UObject* WorldContextObject, Audio::FQuartzClockTickRate& OutTickRate) const;
 
 private:
-	void SetTickRateInternal(const FQuartzQuantizationBoundary& InQuantizationBoundary, const FOnQuartzCommandEventBP& InDelegate, const Audio::FQuartzClockTickRate& NewTickRate);
+	AUDIOMIXER_API void SetTickRateInternal(const FQuartzQuantizationBoundary& InQuantizationBoundary, const FOnQuartzCommandEventBP& InDelegate, const Audio::FQuartzClockTickRate& NewTickRate);
 
 	Audio::FQuartzClockProxy RawHandle;
 
diff --git a/Engine/Source/Runtime/Engine/Classes/Sound/QuartzSubscription.cpp b/Engine/Source/Runtime/Engine/Classes/Sound/QuartzSubscription.cpp
--- a/Engine/Source/Runtime/Engine/Classes/Sound/QuartzSubscription.cpp
+++ b/Engine/Source/Runtime/Engine/Classes/Sound/QuartzSubscription.cpp
@@ -22,9 +22,6 @@ namespace Audio
 	}
 } // namespace Audio
 
-FQuartzTickableObject::FQuartzTickableObject()
-{}
-
 PRAGMA_DISABLE_DEPRECATION_WARNINGS
 FQuartzTickableObject::~FQuartzTickableObject()
 {
diff --git a/Engine/Source/Runtime/Engine/Classes/Sound/QuartzSubscription.h b/Engine/Source/Runtime/Engine/Classes/Sound/QuartzSubscription.h
--- a/Engine/Source/Runtime/Engine/Classes/Sound/QuartzSubscription.h
+++ b/Engine/Source/Runtime/Engine/Classes/Sound/QuartzSubscription.h
@@ -127,7 +127,7 @@ class FQuartzTickableObject
 {
 public:
 	// ctor
-	FQuartzTickableObject();
+	ENGINE_API FQuartzTickableObject() {}
 
 	// explicitly defaulted ctors (to disable deprecation warnings in compiler-generated functions)
     PRAGMA_DISABLE_DEPRECATION_WARNINGS
diff --git a/Engine/Source/Runtime/Engine/Classes/Sound/QuartzSubscriptionToken.cpp b/Engine/Source/Runtime/Engine/Classes/Sound/QuartzSubscriptionToken.cpp
--- a/Engine/Source/Runtime/Engine/Classes/Sound/QuartzSubscriptionToken.cpp
+++ b/Engine/Source/Runtime/Engine/Classes/Sound/QuartzSubscriptionToken.cpp
@@ -1,4 +1,6 @@
-﻿#include "QuartzSubscriptionToken.h"
+﻿// Copyright Epic Games, Inc. All Rights Reserved.
+
+#include "QuartzSubscriptionToken.h"
 
 #include "Quartz/QuartzSubsystem.h"
 #include "Sound/QuartzQuantizationUtilities.h"
