Spec fix for when submix modulation does not apply to binaural sources

diff --git a/Engine/Source/Runtime/AudioMixer/Private/AudioMixerSource.cpp b/Engine/Source/Runtime/AudioMixer/Private/AudioMixerSource.cpp
--- a/Engine/Source/Runtime/AudioMixer/Private/AudioMixerSource.cpp
+++ b/Engine/Source/Runtime/AudioMixer/Private/AudioMixerSource.cpp
@@ -488,6 +488,8 @@ namespace Audio
 		, MixerDevice(static_cast<FMixerDevice*>(InAudioDevice))
 		, MixerBuffer(nullptr)
 		, MixerSourceVoice(nullptr)
+		, BinauralVolModulators(nullptr)
+		, BinauralDryModulators(nullptr)
 		, bPreviousBusEnablement(false)
 		, bPreviousBaseSubmixEnablement(false)
 		, PreviousAzimuth(-1.0f)
@@ -741,6 +743,11 @@ namespace Audio
 						UE_LOG(LogAudioMixer, Warning, TEXT("Ignoring soundfield Base Submix destination being set on SoundWave (%s) because spatializaition method is set to Binaural.")
 							, *InWaveInstance->GetName());
 					}
+					
+					// Get the modulation info associated with the SubmixPtr here
+					// The audio never goes to this submix, but we can still use this to modulate the volume ourselves later
+					BinauralVolModulators = SubmixPtr->GetOutputVolumeDestination();
+					BinauralDryModulators = SubmixPtr->GetDryVolumeDestination();
 				}
 			}
 
@@ -1606,7 +1613,17 @@ namespace Audio
 			CurrentVolume *= WaveInstance->GetVolume();
 			CurrentVolume *= WaveInstance->GetDynamicVolume();
 
-			// 3. Apply editor gain stage(s)
+			// 3. Submix Volume Modulation (this only happens if the asset is binaural and we're sending to an external submix)
+			if (BinauralVolModulators)
+			{
+				CurrentVolume *= BinauralVolModulators->GetValue();
+			}
+			if (BinauralDryModulators)
+			{
+				CurrentVolume *= BinauralDryModulators->GetValue();
+			}
+
+			// 4. Apply editor gain stage(s)
 			CurrentVolume = FMath::Clamp<float>(GetDebugVolume(CurrentVolume), 0.0f, MAX_VOLUME);
 
 			FActiveSound* ActiveSound = WaveInstance->ActiveSound;
diff --git a/Engine/Source/Runtime/AudioMixer/Private/AudioMixerSource.h b/Engine/Source/Runtime/AudioMixer/Private/AudioMixerSource.h
--- a/Engine/Source/Runtime/AudioMixer/Private/AudioMixerSource.h
+++ b/Engine/Source/Runtime/AudioMixer/Private/AudioMixerSource.h
@@ -127,6 +127,10 @@ namespace Audio
 		FMixerSourceVoice* MixerSourceVoice;
 		IAudioLinkFactory::FAudioLinkSourcePushedSharedPtr AudioLink;
 
+		// These modulators are obtained from the submix and used only on binaural assets
+		FModulationDestination* BinauralVolModulators;
+		FModulationDestination* BinauralDryModulators;
+
 		uint32 bPreviousBusEnablement;
 		uint32 bPreviousBaseSubmixEnablement;
 
diff --git a/Engine/Source/Runtime/AudioMixer/Private/AudioMixerSubmix.cpp b/Engine/Source/Runtime/AudioMixer/Private/AudioMixerSubmix.cpp
--- a/Engine/Source/Runtime/AudioMixer/Private/AudioMixerSubmix.cpp
+++ b/Engine/Source/Runtime/AudioMixer/Private/AudioMixerSubmix.cpp
@@ -1099,6 +1099,18 @@ namespace Audio
 			}
 		}
 
+		// Query Modulation; if any of the submix's Modulation Destinations are being modulated they need to be processed,
+		// because binaural sources still use these Destinations when the submix is set
+		{
+			if (MixerDevice->IsModulationPluginEnabled() && MixerDevice->ModulationInterface.IsValid())
+			{
+				if (DryLevelMod.IsActive() || VolumeMod.IsActive())
+				{
+					return true;
+				}
+			}
+		}
+
 		// If this submix is not rendering any sources directly and silence has been detected, we need to check it's children submixes
 		if (MixerSourceVoices.Num() == 0 && SilenceTimeStartSeconds >= 0.0)
 		{
@@ -2350,6 +2362,16 @@ namespace Audio
 		DryModBaseDb = InDryModBaseDb;
 	}
 
+	FModulationDestination* FMixerSubmix::GetOutputVolumeDestination()
+	{
+		return &VolumeMod;
+	}
+
+	FModulationDestination* FMixerSubmix::GetDryVolumeDestination()
+	{
+		return &DryLevelMod;
+	}
+
 	void FMixerSubmix::BroadcastDelegates()
 	{
 		if (bIsEnvelopeFollowing)
diff --git a/Engine/Source/Runtime/AudioMixer/Public/AudioMixerSubmix.h b/Engine/Source/Runtime/AudioMixer/Public/AudioMixerSubmix.h
--- a/Engine/Source/Runtime/AudioMixer/Public/AudioMixerSubmix.h
+++ b/Engine/Source/Runtime/AudioMixer/Public/AudioMixerSubmix.h
@@ -147,6 +147,10 @@ namespace Audio
 		// Update modulation settings of the submix with Decibel values
 		AUDIOMIXER_API void SetModulationBaseLevels(float InVolumeModBaseDb, float InWetModeBaseDb, float InDryModBaseDb);
 
+		FModulationDestination* GetOutputVolumeDestination();
+
+		FModulationDestination* GetDryVolumeDestination();
+
 		// Gets the submix channels channels
 		AUDIOMIXER_API int32 GetSubmixChannels() const;
 
diff --git a/Engine/Source/Runtime/Engine/Classes/Sound/SoundModulationDestination.h b/Engine/Source/Runtime/Engine/Classes/Sound/SoundModulationDestination.h
--- a/Engine/Source/Runtime/Engine/Classes/Sound/SoundModulationDestination.h
+++ b/Engine/Source/Runtime/Engine/Classes/Sound/SoundModulationDestination.h
@@ -165,7 +165,7 @@ namespace Audio
 			ENGINE_API void Init(FDeviceId InDeviceId, FName InParameterName, bool bInIsBuffered = false, bool bInValueNormalized = false);
 
 			/** returns whether or not destination references an active modulator */
-			ENGINE_API bool IsActive();
+			ENGINE_API bool IsActive() const;
 
 			/* Updates internal value (or buffer if set to bIsBuffered) to current modulated result using the provided value as the base carrier value to modulate.
 			 * Returns true if value was updated.
diff --git a/Engine/Source/Runtime/Engine/Private/SoundModulationDestination.cpp b/Engine/Source/Runtime/Engine/Private/SoundModulationDestination.cpp
--- a/Engine/Source/Runtime/Engine/Private/SoundModulationDestination.cpp
+++ b/Engine/Source/Runtime/Engine/Private/SoundModulationDestination.cpp
@@ -147,7 +147,7 @@ namespace Audio
 		ResetHandles();
 	}
 
-	bool FModulationDestination::IsActive()
+	bool FModulationDestination::IsActive() const
 	{
 		FScopeLock Lock(&HandleCritSection);
 		return Algo::AnyOf(Handles, [](const FModulatorHandle& Handle) { return Handle.IsValid(); });
