Fix Bug where not attaching modulators to a submix causes it to mute all binaural sounds

diff --git a/Engine/Source/Runtime/AudioMixer/Private/AudioMixerSource.cpp b/Engine/Source/Runtime/AudioMixer/Private/AudioMixerSource.cpp
--- a/Engine/Source/Runtime/AudioMixer/Private/AudioMixerSource.cpp
+++ b/Engine/Source/Runtime/AudioMixer/Private/AudioMixerSource.cpp
@@ -489,7 +489,7 @@ namespace Audio
 		, MixerBuffer(nullptr)
 		, MixerSourceVoice(nullptr)
 		, BinauralVolModulators(nullptr)
-		, BinauralDryModulators(nullptr)
+		, BinauralWetModulators(nullptr)
 		, bPreviousBusEnablement(false)
 		, bPreviousBaseSubmixEnablement(false)
 		, PreviousAzimuth(-1.0f)
@@ -747,7 +747,7 @@ namespace Audio
 					// Get the modulation info associated with the SubmixPtr here
 					// The audio never goes to this submix, but we can still use this to modulate the volume ourselves later
 					BinauralVolModulators = SubmixPtr->GetOutputVolumeDestination();
-					BinauralDryModulators = SubmixPtr->GetDryVolumeDestination();
+					BinauralWetModulators = SubmixPtr->GetWetVolumeDestination();
 				}
 			}
 
@@ -1618,9 +1618,9 @@ namespace Audio
 			{
 				CurrentVolume *= BinauralVolModulators->GetValue();
 			}
-			if (BinauralDryModulators)
+			if (BinauralWetModulators)
 			{
-				CurrentVolume *= BinauralDryModulators->GetValue();
+				CurrentVolume *= BinauralWetModulators->GetValue();
 			}
 
 			// 4. Apply editor gain stage(s)
diff --git a/Engine/Source/Runtime/AudioMixer/Private/AudioMixerSource.h b/Engine/Source/Runtime/AudioMixer/Private/AudioMixerSource.h
--- a/Engine/Source/Runtime/AudioMixer/Private/AudioMixerSource.h
+++ b/Engine/Source/Runtime/AudioMixer/Private/AudioMixerSource.h
@@ -129,7 +129,7 @@ namespace Audio
 
 		// These modulators are obtained from the submix and used only on binaural assets
 		FModulationDestination* BinauralVolModulators;
-		FModulationDestination* BinauralDryModulators;
+		FModulationDestination* BinauralWetModulators;
 
 		uint32 bPreviousBusEnablement;
 		uint32 bPreviousBaseSubmixEnablement;
diff --git a/Engine/Source/Runtime/AudioMixer/Private/AudioMixerSubmix.cpp b/Engine/Source/Runtime/AudioMixer/Private/AudioMixerSubmix.cpp
--- a/Engine/Source/Runtime/AudioMixer/Private/AudioMixerSubmix.cpp
+++ b/Engine/Source/Runtime/AudioMixer/Private/AudioMixerSubmix.cpp
@@ -2367,9 +2367,9 @@ namespace Audio
 		return &VolumeMod;
 	}
 
-	FModulationDestination* FMixerSubmix::GetDryVolumeDestination()
+	FModulationDestination* FMixerSubmix::GetWetVolumeDestination()
 	{
-		return &DryLevelMod;
+		return &WetLevelMod;
 	}
 
 	void FMixerSubmix::BroadcastDelegates()
diff --git a/Engine/Source/Runtime/AudioMixer/Public/AudioMixerSubmix.h b/Engine/Source/Runtime/AudioMixer/Public/AudioMixerSubmix.h
--- a/Engine/Source/Runtime/AudioMixer/Public/AudioMixerSubmix.h
+++ b/Engine/Source/Runtime/AudioMixer/Public/AudioMixerSubmix.h
@@ -149,7 +149,7 @@ namespace Audio
 
 		FModulationDestination* GetOutputVolumeDestination();
 
-		FModulationDestination* GetDryVolumeDestination();
+		FModulationDestination* GetWetVolumeDestination();
 
 		// Gets the submix channels channels
 		AUDIOMIXER_API int32 GetSubmixChannels() const;
