[Quartz] Fix for potential hang when calling GetBeatProgressPercent() with Quantization == None. (Just return 0.0 instead of doing the math.)
#rnx

diff --git a/Engine/Source/Runtime/AudioMixer/Private/Quartz/AudioMixerClockHandle.cpp b/Engine/Source/Runtime/AudioMixer/Private/Quartz/AudioMixerClockHandle.cpp
--- a/Engine/Source/Runtime/AudioMixer/Private/Quartz/AudioMixerClockHandle.cpp
+++ b/Engine/Source/Runtime/AudioMixer/Private/Quartz/AudioMixerClockHandle.cpp
@@ -316,7 +316,7 @@ float UQuartzClockHandle::GetBeatsPerMinute(const UObject* WorldContextObject) c
 
 float UQuartzClockHandle::GetBeatProgressPercent(EQuartzCommandQuantization QuantizationBoundary, float PhaseOffset, float MsOffset)
 {
-	if(RawHandle.IsValid())
+	if(RawHandle.IsValid() && QuantizationBoundary != EQuartzCommandQuantization::None)
 	{
 		constexpr float ToMilliseconds = 1000.f;
 	    const float MsInQuantizationType = ToMilliseconds * RawHandle.GetDurationOfQuantizationTypeInSeconds(QuantizationBoundary, 1.f);
