[Quartz] Fix for headless clocks not being delete-able
#rnx

diff --git a/Engine/Source/Runtime/AudioMixer/Private/Quartz/AudioMixerClockManager.cpp b/Engine/Source/Runtime/AudioMixer/Private/Quartz/AudioMixerClockManager.cpp
--- a/Engine/Source/Runtime/AudioMixer/Private/Quartz/AudioMixerClockManager.cpp
+++ b/Engine/Source/Runtime/AudioMixer/Private/Quartz/AudioMixerClockManager.cpp
@@ -161,9 +161,9 @@ namespace Audio
 		return INDEX_NONE;
 	}
 
-	void FQuartzClockManager::RemoveClock(const FName& InName)
+	void FQuartzClockManager::RemoveClock(const FName& InName, bool bForceSynchronous)
 	{
-		if (MixerDevice && !MixerDevice->IsAudioRenderingThread())
+		if (!bForceSynchronous && MixerDevice && !MixerDevice->IsAudioRenderingThread())
 		{
 			MixerDevice->AudioRenderThreadCommand([this, InName]()
 			{
diff --git a/Engine/Source/Runtime/AudioMixer/Private/Quartz/QuartzSubsystem.cpp b/Engine/Source/Runtime/AudioMixer/Private/Quartz/QuartzSubsystem.cpp
--- a/Engine/Source/Runtime/AudioMixer/Private/Quartz/QuartzSubsystem.cpp
+++ b/Engine/Source/Runtime/AudioMixer/Private/Quartz/QuartzSubsystem.cpp
@@ -467,15 +467,25 @@ UQuartzClockHandle* UQuartzSubsystem::CreateNewClock(const UObject* WorldContext
 
 void UQuartzSubsystem::DeleteClockByName(const UObject* WorldContextObject, FName ClockName)
 {
-	Audio::FQuartzClockManager* ClockManager = GetClockManager(WorldContextObject);
+	// first look for the clock on the audio device's clock manager
+	bool bShouldDeleteSynchronous = false;
+	Audio::FQuartzClockManager* ClockManager = GetClockManager(WorldContextObject, /*bUseAudioEngineClockManager*/ true);
 
-	if (!ClockManager)
+	if (ClockManager && !ClockManager->DoesClockExist(ClockName))
 	{
-		return;
+		// if we didn't find it, assume the clock is on the subsystem's clock manager
+		ClockManager = GetClockManager(WorldContextObject, /*bUseAudioEngineClockManager*/ false);
+		bShouldDeleteSynchronous = true;
 	}
 
-	ClockManager->RemoveClock(ClockName);
-	PruneStaleProxies();
+	// if the clock is managed by the audio mixer device,
+	// bShouldDeleteSynchronous is false, and it will be deleted on the correct thread.
+	// if its managed by the subsystem, bShouldDeleteSynchronous will be true
+	if(ClockManager)
+	{
+		ClockManager->RemoveClock(ClockName, bShouldDeleteSynchronous);
+		PruneStaleProxies();
+	}
 }
 
 void UQuartzSubsystem::DeleteClockByHandle(const UObject* WorldContextObject, UQuartzClockHandle*& InClockHandle)
diff --git a/Engine/Source/Runtime/AudioMixer/Public/Quartz/AudioMixerClockManager.h b/Engine/Source/Runtime/AudioMixer/Public/Quartz/AudioMixerClockManager.h
--- a/Engine/Source/Runtime/AudioMixer/Public/Quartz/AudioMixerClockManager.h
+++ b/Engine/Source/Runtime/AudioMixer/Public/Quartz/AudioMixerClockManager.h
@@ -56,7 +56,7 @@ namespace Audio
 
 		// remove existing clock
 		// safe to call from AudioThread (uses Audio Render Thread Command)
-		AUDIOMIXER_API void RemoveClock(const FName& InName);
+		AUDIOMIXER_API void RemoveClock(const FName& InName, bool bForceSynchronous = false);
 
 		// get Tick rate for clock
 		// safe to call from AudioThread (uses critical section)
