[Quartz] Fixed crash that occurs when playing a sound with a quantization set to "BarRelative" and with a musical duration longer than a Bar.

diff --git a/Engine/Source/Runtime/AudioMixer/Private/Quartz/QuartzMetronome.cpp b/Engine/Source/Runtime/AudioMixer/Private/Quartz/QuartzMetronome.cpp
--- a/Engine/Source/Runtime/AudioMixer/Private/Quartz/QuartzMetronome.cpp
+++ b/Engine/Source/Runtime/AudioMixer/Private/Quartz/QuartzMetronome.cpp
@@ -196,9 +196,25 @@ namespace Audio
 		// counting from the current bar
 		else if (InQuantizationBoundary.CountingReferencePoint == EQuarztQuantizationReference::BarRelative)
 		{
-			const int32 NumSubdivisionsPerBar = CountNumSubdivisionsPerBar(InQuantizationBoundary.Quantization);
-			const int32 NumSubdivisionsAlreadyOccuredInCurrentBar = CountNumSubdivisionsSinceBarStart(InQuantizationBoundary.Quantization);
-			NumDurationsLeft = (NumDurationsLeft % NumSubdivisionsPerBar) - NumSubdivisionsAlreadyOccuredInCurrentBar;
+			const float NumSubdivisionsPerBar = CountNumSubdivisionsPerBar(InQuantizationBoundary.Quantization);
+			const float NumSubdivisionsAlreadyOccuredInCurrentBar = CountNumSubdivisionsSinceBarStart(InQuantizationBoundary.Quantization);
+
+			// the requested duration is longer than our current bar
+			// do the math in bars instead
+			if (NumSubdivisionsPerBar < 1.f && ensure(!FMath::IsNearlyZero(NumSubdivisionsPerBar)))
+			{
+				const float NumBarsPerSubdivision = 1.f / NumSubdivisionsPerBar;
+				const float NumBarsRemaining = NumBarsPerSubdivision - (NumSubdivisionsAlreadyOccuredInCurrentBar - 1.f);
+
+				InQuantizationBoundary.Multiplier = NumBarsRemaining;
+				InQuantizationBoundary.Quantization = EQuartzCommandQuantization::Bar;
+				
+				NumDurationsLeft = static_cast<int32>(InQuantizationBoundary.Multiplier) - 1;
+			}
+			else
+			{
+				NumDurationsLeft = NumDurationsLeft % static_cast<int32>(NumSubdivisionsPerBar) - static_cast<int32>(NumSubdivisionsAlreadyOccuredInCurrentBar);
+			}
 
 			// if NumDurationsLeft is negative, it means the target has already passed this bar.
 			// instead we will schedule the sound for the same target in the next bar
@@ -208,7 +224,7 @@ namespace Audio
 			}
 		}
 
-		const double FrationalPortion = FMath::Fractional(InQuantizationBoundary.Multiplier);
+		const double FractionalPortion = FMath::Fractional(InQuantizationBoundary.Multiplier);
 
 		// for Beats, the lengths are not uniform for complex meters
 		if ((InQuantizationBoundary.Quantization == EQuartzCommandQuantization::Beat) && PulseDurations.Num())
@@ -233,11 +249,11 @@ namespace Audio
 				TempPulseDurationIndex = 0;
 			}
 
-			FramesUntilBoundary += FrationalPortion * PulseDurations[TempPulseDurationIndex];
+			FramesUntilBoundary += FractionalPortion * PulseDurations[TempPulseDurationIndex];
 		}
 		else
 		{
-			const float Multiplier = NumDurationsLeft + FrationalPortion;
+			const float Multiplier = NumDurationsLeft + FractionalPortion;
 			const float Duration = static_cast<float>(MusicalDurationsInFrames[InQuantizationBoundary.Quantization]);
 			FramesUntilBoundary += Multiplier * Duration;
 		}
