Fixing EQ bandwidth issue

diff --git a/Engine/Source/Runtime/AudioMixer/Classes/SubmixEffects/AudioMixerSubmixEffectEQ.h b/Engine/Source/Runtime/AudioMixer/Classes/SubmixEffects/AudioMixerSubmixEffectEQ.h
--- a/Engine/Source/Runtime/AudioMixer/Classes/SubmixEffects/AudioMixerSubmixEffectEQ.h
+++ b/Engine/Source/Runtime/AudioMixer/Classes/SubmixEffects/AudioMixerSubmixEffectEQ.h
@@ -18,10 +18,10 @@ struct FSubmixEffectEQBand
 {
 	GENERATED_USTRUCT_BODY()
 
-	UPROPERTY(EditAnywhere, BlueprintReadWrite, Category = "SubmixEffect|Preset", meta = (ClampMin = "0.0", ClampMax = "20000.0", UIMin = "0.0", UIMax = "15000.0"))
+	UPROPERTY(EditAnywhere, BlueprintReadWrite, Category = "SubmixEffect|Preset", meta = (ClampMin = "20.0", ClampMax = "20000.0", UIMin = "20.0", UIMax = "15000.0"))
 	float Frequency;
 
-	UPROPERTY(EditAnywhere, BlueprintReadWrite, Category = "SubmixEffect|Preset", meta = (ClampMin = "0.0", ClampMax = "10.0", UIMin = "0.0", UIMax = "10.0"))
+	UPROPERTY(EditAnywhere, BlueprintReadWrite, Category = "SubmixEffect|Preset", meta = (ClampMin = "0.1", ClampMax = "2.0", UIMin = "0.1", UIMax = "2.0"))
 	float Bandwidth;
 
 	UPROPERTY(EditAnywhere, BlueprintReadWrite, Category = "SubmixEffect|Preset", meta = (DisplayName = "Gain (dB)", ClampMin = "-90.0", ClampMax = "20.0", UIMin = "-90.0", UIMax = "20.0"))
diff --git a/Engine/Source/Runtime/AudioMixer/Private/Effects/AudioMixerSubmixEffectEQ.cpp b/Engine/Source/Runtime/AudioMixer/Private/Effects/AudioMixerSubmixEffectEQ.cpp
--- a/Engine/Source/Runtime/AudioMixer/Private/Effects/AudioMixerSubmixEffectEQ.cpp
+++ b/Engine/Source/Runtime/AudioMixer/Private/Effects/AudioMixerSubmixEffectEQ.cpp
@@ -223,6 +223,14 @@ void FSubmixEffectSubmixEQ::UpdateParameters(const int32 InNumOutputChannels)
 	if (PendingSettings.GetParams(&NewSettings))
 	{
 		bParamsChanged = true;
+
+		// Make sure we clamp our freq and bandwidth to reasonable values here
+		for (FSubmixEffectEQBand& Band : NewSettings.EQBands)
+		{
+			Band.Frequency = GetClampedFrequency(Band.Frequency);
+			Band.Bandwidth = GetClampedBandwidth(Band.Bandwidth);
+		}
+
 		RenderThreadEQSettings = NewSettings;
 	}
 
