[Quartz] Fixed timing issues with quantized transport reset.

diff --git a/Engine/Source/Runtime/AudioMixer/Private/Quartz/AudioMixerClock.cpp b/Engine/Source/Runtime/AudioMixer/Private/Quartz/AudioMixerClock.cpp
--- a/Engine/Source/Runtime/AudioMixer/Private/Quartz/AudioMixerClock.cpp
+++ b/Engine/Source/Runtime/AudioMixer/Private/Quartz/AudioMixerClock.cpp
@@ -289,7 +289,6 @@ namespace Audio
 		
 		PreTickCommands->PumpCommandQueue(this);
 
-		TRACE_CPUPROFILER_EVENT_SCOPE(QuartzClock::TickInternal);
 		if (!bIsRunning)
 		{
 			return;
@@ -302,23 +301,19 @@ namespace Audio
 		}
 
 		const int32 FramesOfLatency = (ThreadLatencyInMilliseconds / 1000) * Metronome.GetTickRate().GetSampleRate();
+		int32 FramesToTick = InNumFramesUntilNextTick - TickDelayLengthInFrames;
 
-		if (TickDelayLengthInFrames == 0)
-		{
-			TickInternal(InNumFramesUntilNextTick, ClockAlteringPendingCommands, FramesOfLatency); // (process things like BPM changes first)
-			TickInternal(InNumFramesUntilNextTick, PendingCommands, FramesOfLatency);
-			Metronome.Tick(InNumFramesUntilNextTick, FramesOfLatency);
-		}
-		else
-		{
-			const int32 FramesToTick = InNumFramesUntilNextTick - TickDelayLengthInFrames; 
-			TickInternal(FramesToTick, ClockAlteringPendingCommands, FramesOfLatency, TickDelayLengthInFrames);
-			TickInternal(FramesToTick, PendingCommands, FramesOfLatency, TickDelayLengthInFrames);
-			Metronome.Tick(FramesToTick, FramesOfLatency);
+        // commands executed in TickInternal may alter "TickDelayLengthInFrames" for the metronome's benefit
+        // for the 2nd TickInternal() call we want to use the unmodified value (OriginalTickDelayLengthInFrames).
+        const int32 OriginalTickDelayLengthInFrames = TickDelayLengthInFrames;
+        TickInternal(FramesToTick, ClockAlteringPendingCommands, FramesOfLatency, OriginalTickDelayLengthInFrames);
+        TickInternal(FramesToTick, PendingCommands, FramesOfLatency, OriginalTickDelayLengthInFrames);
 
-			TickDelayLengthInFrames = 0;
-		}
+		// FramesToTick may have been updated by TickInternal, recalculate
+		FramesToTick = InNumFramesUntilNextTick - TickDelayLengthInFrames;
+		Metronome.Tick(FramesToTick, FramesOfLatency);
 
+		TickDelayLengthInFrames = 0;
 
 		UpdateCachedState();
 	}
@@ -335,6 +330,7 @@ namespace Audio
 
 	void FQuartzClock::TickInternal(int32 InNumFramesUntilNextTick, TArray<PendingCommand>& CommandsToTick, int32 FramesOfLatency, int32 FramesOfDelay)
 	{
+		TRACE_CPUPROFILER_EVENT_SCOPE(QuartzClock::TickInternal);
 		bool bHaveCommandsToRemove = false;
 
 		// Update all pending commands
@@ -685,8 +681,13 @@ namespace Audio
 		return nullptr;
 	}
 
-	void FQuartzClock::ResetTransport()
+	void FQuartzClock::ResetTransport(const int32 NumFramesToTickBeforeReset)
 	{
+		if (NumFramesToTickBeforeReset != 0)
+		{
+			Metronome.Tick(NumFramesToTickBeforeReset);
+		}
+		
 		Metronome.ResetTransport();
 	}
 
diff --git a/Engine/Source/Runtime/AudioMixer/Private/Quartz/AudioMixerQuantizedCommands.cpp b/Engine/Source/Runtime/AudioMixer/Private/Quartz/AudioMixerQuantizedCommands.cpp
--- a/Engine/Source/Runtime/AudioMixer/Private/Quartz/AudioMixerQuantizedCommands.cpp
+++ b/Engine/Source/Runtime/AudioMixer/Private/Quartz/AudioMixerQuantizedCommands.cpp
@@ -189,7 +189,10 @@ namespace Audio
 
 	void FQuantizedTransportReset::OnFinalCallbackCustom(int32 InNumFramesLeft)
 	{
-		OwningClockPtr->ResetTransport();
+		// todo: guard against multiple reset commands executing in the same clock tick
+		// OwningClockPtr->ResetTransport(InNumFramesLeft - 1); // - 1 to triggering events w/o double trigger of events on the last frame
+		OwningClockPtr->ResetTransport(0);
+		OwningClockPtr->AddToTickDelay(InNumFramesLeft); // next metronome tick will be less by InNumFramesLeft
 	}
 
 	static const FName TransportResetCommandName("Transport Reset Command");
diff --git a/Engine/Source/Runtime/AudioMixer/Public/Quartz/AudioMixerClock.h b/Engine/Source/Runtime/AudioMixer/Public/Quartz/AudioMixerClock.h
--- a/Engine/Source/Runtime/AudioMixer/Public/Quartz/AudioMixerClock.h
+++ b/Engine/Source/Runtime/AudioMixer/Public/Quartz/AudioMixerClock.h
@@ -130,7 +130,7 @@ namespace Audio
 
 		AUDIOMIXER_API void SetSampleRate(float InNewSampleRate);
 
-		AUDIOMIXER_API void ResetTransport();
+		AUDIOMIXER_API void ResetTransport(const int32 NumFramesToTickBeforeReset = 0);
 
 		// (used for StartOtherClock command to handle the sub-tick as the target clock)
 		AUDIOMIXER_API void AddToTickDelay(int32 NumFramesOfDelayToAdd);
diff --git a/Engine/Source/Runtime/AudioMixer/Public/Quartz/AudioMixerQuantizedCommands.h b/Engine/Source/Runtime/AudioMixer/Public/Quartz/AudioMixerQuantizedCommands.h
--- a/Engine/Source/Runtime/AudioMixer/Public/Quartz/AudioMixerQuantizedCommands.h
+++ b/Engine/Source/Runtime/AudioMixer/Public/Quartz/AudioMixerQuantizedCommands.h
@@ -104,7 +104,7 @@ namespace Audio
 
 		AUDIOMIXER_API virtual void OnFinalCallbackCustom(int32 InNumFramesLeft) override;
 
-		virtual bool IsClockAltering() override { return false; }
+		virtual bool IsClockAltering() override { return true; }
 
 		AUDIOMIXER_API virtual FName GetCommandName() const override;
 		virtual EQuartzCommandType GetCommandType() const { return EQuartzCommandType::TransportReset; };
