Cover for edge cases where the Audio Mixer is cleaned up before Audio thread/audio render thread calls are made when updating submix data

diff --git a/Engine/Source/Runtime/AudioMixer/Private/AudioMixerDevice.cpp b/Engine/Source/Runtime/AudioMixer/Private/AudioMixerDevice.cpp
--- a/Engine/Source/Runtime/AudioMixer/Private/AudioMixerDevice.cpp
+++ b/Engine/Source/Runtime/AudioMixer/Private/AudioMixerDevice.cpp
@@ -1488,13 +1488,18 @@ namespace Audio
 
 		if (IsModulationPluginEnabled() && ModulationInterface.IsValid())
 		{
-			FMixerSubmixPtr MixerSubmixPtr = GetSubmixInstance(InSoundSubmix).Pin();
+			TWeakObjectPtr<USoundSubmix> SubmixWeakPtr = InSoundSubmix;
+			FMixerSubmixWeakPtr MixerSubmixWeakPtr = GetSubmixInstance(InSoundSubmix);
 
-			if (MixerSubmixPtr.IsValid())
+			if (SubmixWeakPtr.IsValid())
 			{
-				AudioRenderThreadCommand([MixerSubmixPtr, VolumeMod = InOutputModulation, WetMod = InOutputModulation, DryMod = InOutputModulation]()
+				AudioRenderThreadCommand([MixerSubmixWeakPtr, VolumeMod = InOutputModulation, WetMod = InOutputModulation, DryMod = InOutputModulation]()
 				{
-					MixerSubmixPtr->UpdateModulationSettings(VolumeMod, WetMod, DryMod);
+					FMixerSubmixPtr MixerSubmixPtr = MixerSubmixWeakPtr.Pin();
+				    if (MixerSubmixPtr.IsValid())
+				    {
+						MixerSubmixPtr->UpdateModulationSettings(VolumeMod, WetMod, DryMod);
+					}
 				});
 			}
 		}
diff --git a/Engine/Source/Runtime/Engine/Private/SoundSubmix.cpp b/Engine/Source/Runtime/Engine/Private/SoundSubmix.cpp
--- a/Engine/Source/Runtime/Engine/Private/SoundSubmix.cpp
+++ b/Engine/Source/Runtime/Engine/Private/SoundSubmix.cpp
@@ -546,7 +546,10 @@ void USoundSubmix::PushModulationChanges()
 		TSet<TObjectPtr<USoundModulatorBase>> NewDryLevelMod = DryLevelModulation.Modulators;
 		AudioDeviceManager->IterateOverAllDevices([SoundSubmix = this, VolMod = MoveTemp(NewVolumeMod), WetMod = MoveTemp(NewWetLevelMod), DryMod = MoveTemp(NewDryLevelMod)](Audio::FDeviceId Id, FAudioDevice* Device) mutable
 		{
-			Device->UpdateSubmixModulationSettings(SoundSubmix, VolMod, WetMod, DryMod);
+		    if (Device)
+		    {
+		    	Device->UpdateSubmixModulationSettings(SoundSubmix, VolMod, WetMod, DryMod);
+			}
 		});
 	}
 }
