- Move to using SharedPtr/Refs for Submix listeners & deprecate old references, moving clients where possible to be constructed as shared pointers
- Deprecate more references to "master" submixes, labeling the single "main" submix as such and reverb, eq, etc. as "required" submixes
#rnx

diff --git a/Engine/Plugins/Experimental/MotoSynth/Source/MotoSynth/Private/MotoSynthPreset.cpp b/Engine/Plugins/Experimental/MotoSynth/Source/MotoSynth/Private/MotoSynthPreset.cpp
--- a/Engine/Plugins/Experimental/MotoSynth/Source/MotoSynth/Private/MotoSynthPreset.cpp
+++ b/Engine/Plugins/Experimental/MotoSynth/Source/MotoSynth/Private/MotoSynthPreset.cpp
@@ -2,6 +2,7 @@
 
 #include "MotoSynthPreset.h"
 #include "AudioDevice.h"
+#include "AudioMixerDevice.h"
 #include "MotoSynthEngine.h"
 #include "MotoSynthSourceAsset.h"
 #include "MotoSynthModule.h"
@@ -155,7 +156,7 @@ void FMotoSynthEnginePreviewer::StartPreviewing()
 			if (FAudioDevice* AudioDevice = AudioDeviceManager->GetMainAudioDeviceRaw())
 			{
 				bRegistered = true;
-				AudioDevice->RegisterSubmixBufferListener(this);
+				AudioDevice->RegisterSubmixBufferListener(AsShared(), AudioDevice->GetMainSubmixObject());
 			}
 		}
 	}
@@ -173,7 +174,7 @@ void FMotoSynthEnginePreviewer::StopPreviewing()
 			{
 				bRegistered = false;
 				bEngineInitialized = false;
-				AudioDevice->UnregisterSubmixBufferListener(this);
+				AudioDevice->UnregisterSubmixBufferListener(AsShared(), AudioDevice->GetMainSubmixObject());
 			}
 		}
 	}
@@ -253,11 +254,15 @@ void FMotoSynthEnginePreviewer::OnNewSubmixBuffer(const USoundSubmix* OwningSubm
 	}
 }
 
+const FString& FMotoSynthEnginePreviewer::GetListenerName() const
+{
+	static const FString ListenerName = TEXT("MotoSynthEnginePreviewer");
+	return ListenerName;
+}
+
 void FMotoSynthEnginePreviewer::Reset()
 {
 	CurrentPreviewCurveStartTime = 0.0f;
 	bEngineInitialized = false;
 }
 #endif // WITH_EDITOR
-
-
diff --git a/Engine/Plugins/Experimental/MotoSynth/Source/MotoSynth/Private/MotoSynthSourceAsset.cpp b/Engine/Plugins/Experimental/MotoSynth/Source/MotoSynth/Private/MotoSynthSourceAsset.cpp
--- a/Engine/Plugins/Experimental/MotoSynth/Source/MotoSynth/Private/MotoSynthSourceAsset.cpp
+++ b/Engine/Plugins/Experimental/MotoSynth/Source/MotoSynth/Private/MotoSynthSourceAsset.cpp
@@ -39,12 +39,18 @@ void FRPMEstimationPreviewTone::StartTestTone(float InVolume)
 			if (FAudioDevice* AudioDevice = AudioDeviceManager->GetMainAudioDeviceRaw())
 			{
 				bRegistered = true;
-				AudioDevice->RegisterSubmixBufferListener(this);
+				AudioDevice->RegisterSubmixBufferListener(AsShared(), AudioDevice->GetMainSubmixObject());
 			}
 		}
 	}
 }
 
+const FString& FRPMEstimationPreviewTone::GetListenerName() const
+{
+	static const FString ListenerName = TEXT("RPM Estimation Preview Tone Listener");
+	return ListenerName;
+}
+
 void FRPMEstimationPreviewTone::StopTestTone()
 {
 	Reset();
@@ -56,7 +62,7 @@ void FRPMEstimationPreviewTone::StopTestTone()
 			if (FAudioDevice* AudioDevice = AudioDeviceManager->GetMainAudioDeviceRaw())
 			{
 				bRegistered = false;
-				AudioDevice->UnregisterSubmixBufferListener(this);
+				AudioDevice->UnregisterSubmixBufferListener(AsShared(), AudioDevice->GetMainSubmixObject());
 			}
 		}
 	}
@@ -154,6 +160,9 @@ void FRPMEstimationPreviewTone::OnNewSubmixBuffer(const USoundSubmix* OwningSubm
 #endif
 
 UMotoSynthSource::UMotoSynthSource()
+#if WITH_EDITORONLY_DATA
+	: MotoSynthSineToneTest(MakeShared<FRPMEstimationPreviewTone, ESPMode::ThreadSafe>())
+#endif // WITH_EDITORONLY_DATA
 {
 }
 
@@ -172,6 +181,11 @@ void UMotoSynthSource::BeginDestroy()
 		MotoSynthDataManager.UnRegisterData(SourceDataID);
 		SourceDataID = INDEX_NONE;
 	}
+
+#if WITH_EDITORONLY_DATA
+	MotoSynthSineToneTest->Reset();
+	MotoSynthSineToneTest.Reset();
+#endif // WITH_EDITORONLY_DATA
 }
 
 uint32 UMotoSynthSource::GetNextSourceID() const
@@ -240,23 +254,23 @@ void UMotoSynthSource::PostLoad()
 #if WITH_EDITOR
 void UMotoSynthSource::PlayToneMatch()
 {
-	MotoSynthSineToneTest.Reset();
+	MotoSynthSineToneTest->Reset();
 
 	FRichCurve* RichRPMCurve = RPMCurve.GetRichCurve();
 	if (RichRPMCurve)
 	{
-		MotoSynthSineToneTest.SetPitchCurve(*RichRPMCurve);
+		MotoSynthSineToneTest->SetPitchCurve(*RichRPMCurve);
 	}
 
 	TArrayView<const int16> Audiobuffer = MakeArrayView(SourceDataPCM.GetData(), SourceDataPCM.Num());
-	MotoSynthSineToneTest.SetAudioFile(Audiobuffer, SourceSampleRate);
+	MotoSynthSineToneTest->SetAudioFile(Audiobuffer, SourceSampleRate);
 
-	MotoSynthSineToneTest.StartTestTone(RPMSynthVolume);
+	MotoSynthSineToneTest->StartTestTone(RPMSynthVolume);
 }
 
 void UMotoSynthSource::StopToneMatch()
 {
-	MotoSynthSineToneTest.StopTestTone();
+	MotoSynthSineToneTest->StopTestTone();
 }
 
 void UMotoSynthSource::UpdateSourceData()
diff --git a/Engine/Plugins/Experimental/MotoSynth/Source/MotoSynth/Public/MotoSynthPreset.h b/Engine/Plugins/Experimental/MotoSynth/Source/MotoSynth/Public/MotoSynthPreset.h
--- a/Engine/Plugins/Experimental/MotoSynth/Source/MotoSynth/Public/MotoSynthPreset.h
+++ b/Engine/Plugins/Experimental/MotoSynth/Source/MotoSynth/Public/MotoSynthPreset.h
@@ -186,6 +186,7 @@ public:
 
 	// ISubmixBufferListener
 	void OnNewSubmixBuffer(const USoundSubmix* OwningSubmix, float* AudioData, int32 NumSamples, int32 NumChannels, const int32 InSampleRate, double AudioClock) override;
+	const FString& GetListenerName() const override;
 	// ~ ISubmixBufferListener
 
 private:
diff --git a/Engine/Plugins/Experimental/MotoSynth/Source/MotoSynth/Public/MotoSynthSourceAsset.h b/Engine/Plugins/Experimental/MotoSynth/Source/MotoSynth/Public/MotoSynthSourceAsset.h
--- a/Engine/Plugins/Experimental/MotoSynth/Source/MotoSynth/Public/MotoSynthSourceAsset.h
+++ b/Engine/Plugins/Experimental/MotoSynth/Source/MotoSynth/Public/MotoSynthSourceAsset.h
@@ -47,7 +47,8 @@ public:
 	bool IsDone() const { return CurrentFrame >= AudioFileBuffer.Num(); }
 
 	// ISubmixBufferListener
-	void OnNewSubmixBuffer(const USoundSubmix* OwningSubmix, float* AudioData, int32 NumSamples, int32 NumChannels, const int32 SampleRate, double AudioClock) override;
+	virtual void OnNewSubmixBuffer(const USoundSubmix* OwningSubmix, float* AudioData, int32 NumSamples, int32 NumChannels, const int32 SampleRate, double AudioClock) override;
+	virtual const FString& GetListenerName() const override;
 	// ~ISubmixBufferListener
 
 private:
@@ -239,7 +240,7 @@ protected:
 	TArray<FGrainTableEntry> GrainTable;
 
 #if WITH_EDITORONLY_DATA
-	FRPMEstimationPreviewTone MotoSynthSineToneTest;
+	TSharedPtr<FRPMEstimationPreviewTone, ESPMode::ThreadSafe> MotoSynthSineToneTest;
 #endif
 
 	// Data ID used to track the source data with the data manager
diff --git a/Engine/Plugins/Experimental/TextToSpeech/Source/TextToSpeech/Private/Flite/FliteTextToSpeech.cpp b/Engine/Plugins/Experimental/TextToSpeech/Source/TextToSpeech/Private/Flite/FliteTextToSpeech.cpp
--- a/Engine/Plugins/Experimental/TextToSpeech/Source/TextToSpeech/Private/Flite/FliteTextToSpeech.cpp
+++ b/Engine/Plugins/Experimental/TextToSpeech/Source/TextToSpeech/Private/Flite/FliteTextToSpeech.cpp
@@ -21,7 +21,6 @@
  */
 FFliteTextToSpeech::FFliteTextToSpeech()
 {
-
 }
 
 FFliteTextToSpeech::~FFliteTextToSpeech()
@@ -138,14 +137,17 @@ void FFliteTextToSpeech::OnActivated()
 	{
 		if (FAudioDeviceHandle AudioDeviceHandle = AudioDeviceManager->GetMainAudioDeviceHandle())
 		{
-			FliteAdapter = MakeUnique<FFliteAdapter>();
-			TTSSubmixListener = MakeUnique<FFliteTextToSpeechSubmixListener>(GetId());
-			AudioDeviceHandle->RegisterSubmixBufferListener(TTSSubmixListener.Get());
-			// RegisterSubmixBufferListener lazily enqueues the registration on the audio thread,
-		// so we have to wait for the audio thread to destroy.
-			FAudioCommandFence Fence;
-			Fence.BeginFence();
-			Fence.Wait();
+			if (FAudioDevice* AudioDevice = AudioDeviceHandle.GetAudioDevice())
+			{
+				FliteAdapter = MakeUnique<FFliteAdapter>();
+				TTSSubmixListener = MakeShared<FFliteTextToSpeechSubmixListener>(GetId());
+				AudioDeviceHandle->RegisterSubmixBufferListener(TTSSubmixListener->AsShared(), AudioDevice->GetMainSubmixObject());
+				// RegisterSubmixBufferListener lazily enqueues the registration on the audio thread,
+			// so we have to wait for the audio thread to destroy.
+				FAudioCommandFence Fence;
+				Fence.BeginFence();
+				Fence.Wait();
+			}
 		}
 	}
 }
@@ -166,13 +168,13 @@ void FFliteTextToSpeech::OnDeactivated()
 	{
 		if (FAudioDeviceHandle AudioDeviceHandle = AudioDeviceManager->GetMainAudioDeviceHandle())
 		{
-			AudioDeviceHandle->UnregisterSubmixBufferListener(TTSSubmixListener.Get());
+			AudioDeviceHandle->UnregisterSubmixBufferListener(TTSSubmixListener->AsShared(), AudioDeviceHandle->GetMainSubmixObject());
 			FAudioCommandFence Fence;
 			Fence.BeginFence();
 			Fence.Wait();
 		}
 		FliteAdapter = nullptr;
-		TTSSubmixListener = nullptr;
+		TTSSubmixListener.Reset();
 	}
 }
 
diff --git a/Engine/Plugins/Experimental/TextToSpeech/Source/TextToSpeech/Private/Flite/FliteTextToSpeech.h b/Engine/Plugins/Experimental/TextToSpeech/Source/TextToSpeech/Private/Flite/FliteTextToSpeech.h
--- a/Engine/Plugins/Experimental/TextToSpeech/Source/TextToSpeech/Private/Flite/FliteTextToSpeech.h
+++ b/Engine/Plugins/Experimental/TextToSpeech/Source/TextToSpeech/Private/Flite/FliteTextToSpeech.h
@@ -40,6 +40,6 @@ private:
 	*/
 	TUniquePtr<class FFliteAdapter> FliteAdapter;
 	/** Handles routing synthesized speech audio data directly to the hardware */
-	TUniquePtr<FFliteTextToSpeechSubmixListener> TTSSubmixListener;
+	TSharedPtr<FFliteTextToSpeechSubmixListener, ESPMode::ThreadSafe> TTSSubmixListener;
 };
 #endif
diff --git a/Engine/Plugins/Experimental/TextToSpeech/Source/TextToSpeech/Private/Flite/FliteTextToSpeechSubmixListener.cpp b/Engine/Plugins/Experimental/TextToSpeech/Source/TextToSpeech/Private/Flite/FliteTextToSpeechSubmixListener.cpp
--- a/Engine/Plugins/Experimental/TextToSpeech/Source/TextToSpeech/Private/Flite/FliteTextToSpeechSubmixListener.cpp
+++ b/Engine/Plugins/Experimental/TextToSpeech/Source/TextToSpeech/Private/Flite/FliteTextToSpeechSubmixListener.cpp
@@ -11,7 +11,7 @@ extern TMap<TextToSpeechId, TWeakPtr<FTextToSpeechBase>> ActiveTextToSpeechMap;
 
 FFliteTextToSpeechSubmixListener::FFliteTextToSpeechSubmixListener(TextToSpeechId InOwningTTSId)
 	: OwningTTSId(InOwningTTSId)
-, Volume(1.0f)
+	, Volume(1.0f)
 	, bAllowPlayback(false)
 	, bMuted(false)
 {
@@ -27,6 +27,12 @@ FFliteTextToSpeechSubmixListener::~FFliteTextToSpeechSubmixListener()
 	StopPlayback_GameThread();
 }
 
+const FString& FFliteTextToSpeechSubmixListener::GetListenerName() const
+{
+	static const FString ListenerName = TEXT("Flite Text-To-Speech Listener");
+	return ListenerName;
+}
+
 void FFliteTextToSpeechSubmixListener::OnNewSubmixBuffer(const USoundSubmix* OwningSubmix, float* AudioData, int32 NumSamples, int32 NumChannels, const int32 SampleRate, double AudioClock)
 {
 	if (!bAllowPlayback)
diff --git a/Engine/Plugins/Experimental/TextToSpeech/Source/TextToSpeech/Private/Flite/FliteTextToSpeechSubmixListener.h b/Engine/Plugins/Experimental/TextToSpeech/Source/TextToSpeech/Private/Flite/FliteTextToSpeechSubmixListener.h
--- a/Engine/Plugins/Experimental/TextToSpeech/Source/TextToSpeech/Private/Flite/FliteTextToSpeechSubmixListener.h
+++ b/Engine/Plugins/Experimental/TextToSpeech/Source/TextToSpeech/Private/Flite/FliteTextToSpeechSubmixListener.h
@@ -40,6 +40,7 @@ public:
 	void SetVolume(float InVolume);
 
 	// ISubmixBufferListener
+	virtual const FString& GetListenerName() const override;
 	virtual void OnNewSubmixBuffer(const USoundSubmix* OwningSubmix, float* AudioData, int32 NumSamples, int32 NumChannels, const int32 SampleRate, double AudioClock) override;
 	virtual bool IsRenderingAudio() const override;
 	// ~ISubmixBufferListener
diff --git a/Engine/Plugins/FX/Niagara/Source/Niagara/Classes/NiagaraDataInterfaceAudio.h b/Engine/Plugins/FX/Niagara/Source/Niagara/Classes/NiagaraDataInterfaceAudio.h
--- a/Engine/Plugins/FX/Niagara/Source/Niagara/Classes/NiagaraDataInterfaceAudio.h
+++ b/Engine/Plugins/FX/Niagara/Source/Niagara/Classes/NiagaraDataInterfaceAudio.h
@@ -44,6 +44,7 @@ public:
 	NIAGARA_API int32 GetNumChannels() const;
 
 	// Begin ISubmixBufferListener overrides
+	NIAGARA_API virtual const FString& GetListenerName() const override;
 	NIAGARA_API virtual void OnNewSubmixBuffer(const USoundSubmix* OwningSubmix, float* AudioData, int32 NumSamples, int32 NumChannels, const int32 SampleRate, double AudioClock) override;
 	// End ISubmixBufferListener overrides
 
@@ -113,7 +114,7 @@ private:
 
 	// Map of audio devices to submix listeners. Needed for editor where multiple audio devices
 	// may exist.
-	TMap<Audio::FDeviceId, TUniquePtr<FNiagaraSubmixListener>> SubmixListeners;
+	TMap<Audio::FDeviceId, TSharedPtr<FNiagaraSubmixListener>> SubmixListeners;
 
 	// Mixer for sending audio 
 	Audio::FPatchMixer PatchMixer;
diff --git a/Engine/Plugins/FX/Niagara/Source/Niagara/Classes/NiagaraDataInterfaceAudioOscilloscope.h b/Engine/Plugins/FX/Niagara/Source/Niagara/Classes/NiagaraDataInterfaceAudioOscilloscope.h
--- a/Engine/Plugins/FX/Niagara/Source/Niagara/Classes/NiagaraDataInterfaceAudioOscilloscope.h
+++ b/Engine/Plugins/FX/Niagara/Source/Niagara/Classes/NiagaraDataInterfaceAudioOscilloscope.h
@@ -58,7 +58,7 @@ private:
 	void OnNewDeviceCreated(Audio::FDeviceId InID);
 	void OnDeviceDestroyed(Audio::FDeviceId InID);
 
-	TMap<Audio::FDeviceId, TUniquePtr<FNiagaraSubmixListener>> SubmixListeners;
+	TMap<Audio::FDeviceId, TSharedPtr<FNiagaraSubmixListener>> SubmixListeners;
 
 	// This mixer is patched into by all instances of FNiagaraSubmixListener in the SubmixListeners map, and is consumed by DownsampleAudioToBuffer().
 	Audio::FPatchMixer PatchMixer;
diff --git a/Engine/Plugins/FX/Niagara/Source/Niagara/Private/NiagaraDataInterfaceAudio.cpp b/Engine/Plugins/FX/Niagara/Source/Niagara/Private/NiagaraDataInterfaceAudio.cpp
--- a/Engine/Plugins/FX/Niagara/Source/Niagara/Private/NiagaraDataInterfaceAudio.cpp
+++ b/Engine/Plugins/FX/Niagara/Source/Niagara/Private/NiagaraDataInterfaceAudio.cpp
@@ -4,6 +4,8 @@
 
 #include "AudioDevice.h"
 #include "AudioDeviceManager.h"
+#include "Sound/SoundSubmix.h"
+#include "UObject/WeakObjectPtrTemplates.h"
 
 #include UE_INLINE_GENERATED_CPP_BY_NAME(NiagaraDataInterfaceAudio)
 
@@ -57,11 +59,12 @@ FNiagaraSubmixListener::~FNiagaraSubmixListener()
 
 void FNiagaraSubmixListener::RegisterToSubmix()
 {
-	if (FAudioDevice* DeviceHandle = FAudioDeviceManager::Get()->GetAudioDeviceRaw(AudioDeviceId))
+	if (FAudioDevice* AudioDevice = FAudioDeviceManager::Get()->GetAudioDeviceRaw(AudioDeviceId))
 	{
 		bIsRegistered = true;
 
-		DeviceHandle->RegisterSubmixBufferListener(this, Submix);
+		USoundSubmix* SubmixToRegister = Submix ? Submix : &AudioDevice->GetMainSubmixObject();
+		AudioDevice->RegisterSubmixBufferListener(AsShared(), *SubmixToRegister);
 
 		// RegisterSubmixBufferListener lazily enqueues the registration on the audio thread,
 		// so we have to wait for the audio thread to destroy.
@@ -77,29 +80,35 @@ void FNiagaraSubmixListener::UnregisterFromSubmix()
 	{
 		bIsRegistered = false;
 		
-		if (FAudioDevice* DeviceHandle = FAudioDeviceManager::Get()->GetAudioDeviceRaw(AudioDeviceId))
+		if (IsInGameThread())
 		{
-			if (IsInGameThread())
+			if (FAudioDevice* AudioDevice = FAudioDeviceManager::Get()->GetAudioDeviceRaw(AudioDeviceId))
 			{
-				DeviceHandle->UnregisterSubmixBufferListener(this, Submix);
-				
+				USoundSubmix* SubmixToUnregister = Submix ? Submix : &AudioDevice->GetMainSubmixObject();
+				AudioDevice->UnregisterSubmixBufferListener(AsShared(), *SubmixToUnregister);
+
 				// UnregisterSubmixBufferListener lazily enqueues the unregistration on the audio thread,
 				// so we have to wait for the audio thread to destroy.
 				FAudioCommandFence Fence;
 				Fence.BeginFence();
 				Fence.Wait();
 			}
-			else
+		}
+		else
+		{
+			UE::Tasks::FTaskEvent CompletionEvent { UE_SOURCE_LOCATION };
+
+			FAudioThread::RunCommandOnAudioThread([this, &CompletionEvent]()
 			{
-				UE::Tasks::FTaskEvent CompletionEvent{ UE_SOURCE_LOCATION };
-				FAudioThread::RunCommandOnAudioThread([this, DeviceHandle, &CompletionEvent]()
+				if (FAudioDevice* AudioDevice = FAudioDeviceManager::Get()->GetAudioDeviceRaw(AudioDeviceId))
 				{
-					DeviceHandle->UnregisterSubmixBufferListener(this, Submix);
+					USoundSubmix* SubmixToUnregister = Submix ? Submix : &AudioDevice->GetMainSubmixObject();
+					AudioDevice->UnregisterSubmixBufferListener(AsShared(), *SubmixToUnregister);
 					CompletionEvent.Trigger();
-				});
+				}
+			});
 
-				CompletionEvent.Wait();
-			}
+			CompletionEvent.Wait();
 		}
 	}
 }
@@ -114,6 +123,12 @@ int32 FNiagaraSubmixListener::GetNumChannels() const
 	return NumChannelsInSubmix;
 }
 
+const FString& FNiagaraSubmixListener::GetListenerName() const
+{
+	static const FString ListenerName = TEXT("NiagaraSubmixListener");
+	return ListenerName;
+}
+
 void FNiagaraSubmixListener::OnNewSubmixBuffer(const USoundSubmix* OwningSubmix, float* AudioData, int32 NumSamples, int32 NumChannels, const int32 SampleRate, double AudioClock)
 {
 	NumChannelsInSubmix = NumChannels;
@@ -203,7 +218,7 @@ void FNiagaraDataInterfaceProxySubmix::OnDeviceDestroyed(Audio::FDeviceId InDevi
 void FNiagaraDataInterfaceProxySubmix::AddSubmixListener(const Audio::FDeviceId InDeviceId)
 {
 	check(!SubmixListeners.Contains(InDeviceId));
-	SubmixListeners.Emplace(InDeviceId, new FNiagaraSubmixListener(PatchMixer, NumSamplesToBuffer, InDeviceId, SubmixRegisteredTo));
+	SubmixListeners.Emplace(InDeviceId, MakeShared<FNiagaraSubmixListener>(PatchMixer, NumSamplesToBuffer, InDeviceId, SubmixRegisteredTo));
 	SubmixListeners[InDeviceId]->RegisterToSubmix();
 }
 
diff --git a/Engine/Plugins/FX/Niagara/Source/Niagara/Private/NiagaraDataInterfaceAudioOscilloscope.cpp b/Engine/Plugins/FX/Niagara/Source/Niagara/Private/NiagaraDataInterfaceAudioOscilloscope.cpp
--- a/Engine/Plugins/FX/Niagara/Source/Niagara/Private/NiagaraDataInterfaceAudioOscilloscope.cpp
+++ b/Engine/Plugins/FX/Niagara/Source/Niagara/Private/NiagaraDataInterfaceAudioOscilloscope.cpp
@@ -54,7 +54,7 @@ void FNiagaraDataInterfaceProxyOscilloscope::RegisterToAllAudioDevices()
 		{
 			check(!SubmixListeners.Contains(DeviceId));
 			const int32 NumSamplesToPop = Align(FMath::FloorToInt(ScopeInMilliseconds / 1000.0f * InDevice->GetSampleRate()) * AUDIO_MIXER_MAX_OUTPUT_CHANNELS, 4);
-			SubmixListeners.Emplace(DeviceId, new FNiagaraSubmixListener(PatchMixer, NumSamplesToPop, DeviceId, SubmixRegisteredTo));
+			SubmixListeners.Emplace(DeviceId, MakeShared<FNiagaraSubmixListener>(PatchMixer, NumSamplesToPop, DeviceId, SubmixRegisteredTo));
 			SubmixListeners[DeviceId]->RegisterToSubmix();
 		});
 	}
@@ -101,7 +101,7 @@ void FNiagaraDataInterfaceProxyOscilloscope::OnNewDeviceCreated(Audio::FDeviceId
 		if (ensure(DeviceHandle))
 		{
 			const int32 NumSamplesToPop = Align(FMath::FloorToInt(ScopeInMilliseconds / 1000.0f * DeviceHandle->GetSampleRate()) * AUDIO_MIXER_MAX_OUTPUT_CHANNELS, 4);
-			SubmixListeners.Emplace(InID, new FNiagaraSubmixListener(PatchMixer, NumSamplesToPop, InID, SubmixRegisteredTo));
+			SubmixListeners.Emplace(InID, MakeShared<FNiagaraSubmixListener>(PatchMixer, NumSamplesToPop, InID, SubmixRegisteredTo));
 			SubmixListeners[InID]->RegisterToSubmix();
 		}
 	}
diff --git a/Engine/Plugins/Media/MediaIOFramework/Source/MediaIOCore/Private/MediaIOCoreAudioOutput.cpp b/Engine/Plugins/Media/MediaIOFramework/Source/MediaIOCore/Private/MediaIOCoreAudioOutput.cpp
--- a/Engine/Plugins/Media/MediaIOFramework/Source/MediaIOCore/Private/MediaIOCoreAudioOutput.cpp
+++ b/Engine/Plugins/Media/MediaIOFramework/Source/MediaIOCore/Private/MediaIOCoreAudioOutput.cpp
@@ -163,15 +163,21 @@ void FMediaIOAudioCapture::RegisterBufferListener(FAudioDevice* AudioDevice)
 		NumChannels = MixerDevice->GetDeviceOutputChannels();
 		SampleRate = MixerDevice->GetSampleRate();
 		PrimarySubmixName = *GetDefault<UAudioSettings>()->MasterSubmix.GetAssetName();
-		AudioDevice->RegisterSubmixBufferListener(this);
+		AudioDevice->RegisterSubmixBufferListener(AsShared(), AudioDevice->GetMainSubmixObject());
 	}
 }
 
+const FString& FMediaIOAudioCapture::GetListenerName() const
+{
+	static const FString ListenerName = TEXT("MediaIO AudioCapture Listener");
+	return ListenerName;
+}
+
 void FMediaIOAudioCapture::UnregisterBufferListener(FAudioDevice* AudioDevice)
 {
 	if (AudioDevice)
 	{
-		AudioDevice->UnregisterSubmixBufferListener(this);
+		AudioDevice->UnregisterSubmixBufferListener(AsShared(), AudioDevice->GetMainSubmixObject());
 	}
 }
 
diff --git a/Engine/Plugins/Media/MediaIOFramework/Source/MediaIOCore/Private/MediaIOCoreSubsystem.cpp b/Engine/Plugins/Media/MediaIOFramework/Source/MediaIOCore/Private/MediaIOCoreSubsystem.cpp
--- a/Engine/Plugins/Media/MediaIOFramework/Source/MediaIOCore/Private/MediaIOCoreSubsystem.cpp
+++ b/Engine/Plugins/Media/MediaIOFramework/Source/MediaIOCore/Private/MediaIOCoreSubsystem.cpp
@@ -21,13 +21,13 @@ TSharedPtr<FMediaIOAudioOutput> UMediaIOCoreSubsystem::CreateAudioOutput(const F
 
 	if (InArgs.AudioDeviceHandle.IsValid())
 	{
-		if (const TUniquePtr<FMediaIOAudioCapture>* FoundMediaIOAudioCapture = MediaIOAudioCaptures.Find(InArgs.AudioDeviceHandle.GetDeviceID()))
+		if (const TSharedPtr<FMediaIOAudioCapture, ESPMode::ThreadSafe>* FoundMediaIOAudioCapture = MediaIOAudioCaptures.Find(InArgs.AudioDeviceHandle.GetDeviceID()))
 		{
 			MediaIOAudioCapture = FoundMediaIOAudioCapture->Get();
 		}
 		else
 		{
-			MediaIOAudioCapture = MediaIOAudioCaptures.Add(InArgs.AudioDeviceHandle.GetDeviceID(), MakeUnique<FMediaIOAudioCapture>(InArgs.AudioDeviceHandle)).Get();
+			MediaIOAudioCapture = MediaIOAudioCaptures.Add(InArgs.AudioDeviceHandle.GetDeviceID(), MakeShared<FMediaIOAudioCapture, ESPMode::ThreadSafe>(InArgs.AudioDeviceHandle)).Get();
 			MainMediaIOAudioCapture->OnAudioCaptured_RenderThread().BindUObject(this, &UMediaIOCoreSubsystem::OnBufferReceivedByCapture, InArgs.AudioDeviceHandle.GetDeviceID());
 		}
 	}
diff --git a/Engine/Plugins/Media/MediaIOFramework/Source/MediaIOCore/Public/MediaIOCoreAudioOutput.h b/Engine/Plugins/Media/MediaIOFramework/Source/MediaIOCore/Public/MediaIOCoreAudioOutput.h
--- a/Engine/Plugins/Media/MediaIOFramework/Source/MediaIOCore/Public/MediaIOCoreAudioOutput.h
+++ b/Engine/Plugins/Media/MediaIOFramework/Source/MediaIOCore/Public/MediaIOCoreAudioOutput.h
@@ -163,6 +163,7 @@ public:
 
 	//~ ISubmixBufferListener interface
 	virtual void OnNewSubmixBuffer(const USoundSubmix* InOwningSubmix, float* InAudioData, int32 InNumSamples, int32 InNumChannels, const int32 InSampleRate, double InAudioClock) override;
+	virtual const FString& GetListenerName() const override;
 
 	/** Create an audio output that will receive audio samples. */
 	TSharedPtr<FMediaIOAudioOutput> CreateAudioOutput(int32 InNumOutputChannels, FFrameRate InTargetFrameRate, uint32 InMaxSampleLatency, uint32 InOutputSampleRate);
diff --git a/Engine/Plugins/Media/MediaIOFramework/Source/MediaIOCore/Public/MediaIOCoreSubsystem.h b/Engine/Plugins/Media/MediaIOFramework/Source/MediaIOCore/Public/MediaIOCoreSubsystem.h
--- a/Engine/Plugins/Media/MediaIOFramework/Source/MediaIOCore/Public/MediaIOCoreSubsystem.h
+++ b/Engine/Plugins/Media/MediaIOFramework/Source/MediaIOCore/Public/MediaIOCoreSubsystem.h
@@ -58,7 +58,7 @@ private:
 private:
 	TUniquePtr<FMediaIOAudioCapture> MainMediaIOAudioCapture;
 	
-	TMap<Audio::FDeviceId, TUniquePtr<FMediaIOAudioCapture>> MediaIOAudioCaptures;
+	TMap<Audio::FDeviceId, TSharedPtr<FMediaIOAudioCapture, ESPMode::ThreadSafe>> MediaIOAudioCaptures;
 
 	FDelegateHandle DeviceDestroyedHandle;
 
diff --git a/Engine/Plugins/Media/PixelStreaming/Source/PixelStreaming/Private/AudioCapturer.cpp b/Engine/Plugins/Media/PixelStreaming/Source/PixelStreaming/Private/AudioCapturer.cpp
--- a/Engine/Plugins/Media/PixelStreaming/Source/PixelStreaming/Private/AudioCapturer.cpp
+++ b/Engine/Plugins/Media/PixelStreaming/Source/PixelStreaming/Private/AudioCapturer.cpp
@@ -30,6 +30,22 @@ namespace UE::PixelStreaming
 	constexpr int FAudioCapturer::SampleRate;
 	constexpr int FAudioCapturer::NumChannels;
 
+	FAudioCapturerListener::FAudioCapturerListener(FAudioCapturer& InParent)
+		: Parent(&InParent)
+	{
+	}
+
+	void FAudioCapturerListener::OnNewSubmixBuffer(const USoundSubmix* OwningSubmix, float* AudioData, int32 NumSamples, int32 InNumChannels, const int32 InSampleRate, double AudioClock)
+	{
+		check(Parent);
+		Parent->OnNewSubmixBuffer(OwningSubmix, AudioData, NumSamples, InNumChannels, InSampleRate, AudioClock);
+	}
+
+	FAudioCapturer::FAudioCapturer()
+		: Listener(MakeShared<FAudioCapturerListener>(*this))
+	{
+	}
+
 	void FAudioCapturer::OnNewSubmixBuffer(const USoundSubmix* OwningSubmix, float* AudioData, int32 NumSamples, int32 InNumChannels, const int32 InSampleRate, double AudioClock)
 	{
 		TRACE_CPUPROFILER_EVENT_SCOPE_ON_CHANNEL_STR("PixelStreaming FAudioCapturer::OnNewSubmixBuffer", PixelStreamingChannel);
@@ -130,7 +146,7 @@ namespace UE::PixelStreaming
 		}
 
 		bInitialized = true;
-		AudioDevice->RegisterSubmixBufferListener(this);
+		AudioDevice->RegisterSubmixBufferListener(Listener->AsShared(), AudioDevice->GetMainSubmixObject());
 
 		UE_LOG(LogPixelStreamingAudioCapturer, Verbose, TEXT("Init"));
 
@@ -154,7 +170,7 @@ namespace UE::PixelStreaming
 			return -1;
 		}
 
-		AudioDevice->UnregisterSubmixBufferListener(this);
+		AudioDevice->UnregisterSubmixBufferListener(Listener->AsShared(), AudioDevice->GetMainSubmixObject());
 		bInitialized = false;
 
 		{
diff --git a/Engine/Plugins/Media/PixelStreaming/Source/PixelStreaming/Private/AudioCapturer.h b/Engine/Plugins/Media/PixelStreaming/Source/PixelStreaming/Private/AudioCapturer.h
--- a/Engine/Plugins/Media/PixelStreaming/Source/PixelStreaming/Private/AudioCapturer.h
+++ b/Engine/Plugins/Media/PixelStreaming/Source/PixelStreaming/Private/AudioCapturer.h
@@ -9,13 +9,29 @@
 
 namespace UE::PixelStreaming 
 {
+	class FAudioCapturer;
+
+	struct FAudioCapturerListener : public ISubmixBufferListener
+	{
+		FAudioCapturerListener(FAudioCapturer& InParent);
+
+		// ISubmixBufferListener interface
+		virtual void OnNewSubmixBuffer(const USoundSubmix* OwningSubmix, float* AudioData, int32 NumSamples, int32 NumChannels, const int32 SampleRate, double AudioClock) override;
+
+		FAudioCapturer* Parent = nullptr;
+	};
+
 	// This class is only used if -PixelStreamingWebRTCUseLegacyAudioDevice or related CVar are used.
 	// This class will likely be removed in the future once the new ADM is confirmed stable.
-	class FAudioCapturer : public webrtc::AudioDeviceModule, public ISubmixBufferListener
+	class FAudioCapturer : public webrtc::AudioDeviceModule
 	{
+	public:
+		FAudioCapturer();
+
 	private:
-		// ISubmixBufferListener interface
-		void OnNewSubmixBuffer(const USoundSubmix* OwningSubmix, float* AudioData, int32 NumSamples, int32 NumChannels, const int32 SampleRate, double AudioClock) override;
+		friend struct FAudioCapturerListener;
+
+		void OnNewSubmixBuffer(const USoundSubmix* OwningSubmix, float* AudioData, int32 NumSamples, int32 NumChannels, const int32 SampleRate, double AudioClock);
 
 		//
 		// webrtc::AudioDeviceModule interface
@@ -192,5 +208,6 @@ namespace UE::PixelStreaming
 		bool bFormatChecked = false;
 
 		std::unique_ptr<webrtc::TaskQueueFactory> m_taskQueueFactory;
+		TSharedPtr<FAudioCapturerListener, ESPMode::ThreadSafe> Listener;
 	};
 } // namespace UE::PixelStreaming
\ No newline at end of file
diff --git a/Engine/Plugins/Media/PixelStreaming/Source/PixelStreaming/Private/AudioSubmixCapturer.cpp b/Engine/Plugins/Media/PixelStreaming/Source/PixelStreaming/Private/AudioSubmixCapturer.cpp
--- a/Engine/Plugins/Media/PixelStreaming/Source/PixelStreaming/Private/AudioSubmixCapturer.cpp
+++ b/Engine/Plugins/Media/PixelStreaming/Source/PixelStreaming/Private/AudioSubmixCapturer.cpp
@@ -59,11 +59,17 @@ namespace UE::PixelStreaming
 			return false;
 		}
 
-		AudioDevice->RegisterSubmixBufferListener(this);
+		AudioDevice->RegisterSubmixBufferListener(AsShared(), AudioDevice->GetMainSubmixObject());
 		bInitialised = true;
 		return true;
 	}
 
+	const FString& FAudioSubmixCapturer::GetListenerName() const
+	{
+		static const FString ListenerName = TEXT("PixelStreaming:SubmixCapturer");
+		return ListenerName;
+	}
+
 	void FAudioSubmixCapturer::OnNewSubmixBuffer(const USoundSubmix* OwningSubmix,
 		float* AudioData, int32 NumSamples,
 		int32 NumChannels,
@@ -184,7 +190,7 @@ namespace UE::PixelStreaming
 			FAudioDeviceHandle AudioDevice = GEngine->GetMainAudioDevice();
 			if (AudioDevice)
 			{
-				AudioDevice->UnregisterSubmixBufferListener(this);
+				AudioDevice->UnregisterSubmixBufferListener(AsShared(), AudioDevice->GetMainSubmixObject());
 			}
 		}
 
diff --git a/Engine/Plugins/Media/PixelStreaming/Source/PixelStreaming/Private/AudioSubmixCapturer.h b/Engine/Plugins/Media/PixelStreaming/Source/PixelStreaming/Private/AudioSubmixCapturer.h
--- a/Engine/Plugins/Media/PixelStreaming/Source/PixelStreaming/Private/AudioSubmixCapturer.h
+++ b/Engine/Plugins/Media/PixelStreaming/Source/PixelStreaming/Private/AudioSubmixCapturer.h
@@ -36,6 +36,7 @@ namespace UE::PixelStreaming
 		void SetAudioInputMixerPatch(TSharedPtr<class FAudioInput> InPatchInput);
 
 		// ISubmixBufferListener interface
+		const FString& GetListenerName() const override;
 		void OnNewSubmixBuffer(const USoundSubmix* OwningSubmix, float* AudioData,
 			int32 NumSamples, int32 NumChannels,
 			const int32 SampleRate, double AudioClock) override;
diff --git a/Engine/Plugins/Media/PixelStreaming/Source/PixelStreamingEditor/Private/EditorSubmixListener.cpp b/Engine/Plugins/Media/PixelStreaming/Source/PixelStreamingEditor/Private/EditorSubmixListener.cpp
--- a/Engine/Plugins/Media/PixelStreaming/Source/PixelStreamingEditor/Private/EditorSubmixListener.cpp
+++ b/Engine/Plugins/Media/PixelStreaming/Source/PixelStreamingEditor/Private/EditorSubmixListener.cpp
@@ -13,7 +13,7 @@ namespace UE::EditorPixelStreaming
 	{
 		if (AudioDevice.IsValid())
 		{
-			AudioDevice->RegisterSubmixBufferListener(this);
+			AudioDevice->RegisterSubmixBufferListener(AsShared(), AudioDevice->GetMainSubmixObject());
 		}
 	}
 
@@ -21,7 +21,7 @@ namespace UE::EditorPixelStreaming
 	{
 		if (FAudioDevice* AudioDevice = FAudioDeviceManager::Get()->GetAudioDeviceRaw(AudioDeviceId))
 		{
-			AudioDevice->UnregisterSubmixBufferListener(this);
+			AudioDevice->UnregisterSubmixBufferListener(AsShared(), AudioDevice->GetMainSubmixObject());
 		}
 	}
 
diff --git a/Engine/Plugins/Media/PixelStreaming/Source/PixelStreamingEditor/Private/PixelStreamingEditorModule.cpp b/Engine/Plugins/Media/PixelStreaming/Source/PixelStreamingEditor/Private/PixelStreamingEditorModule.cpp
--- a/Engine/Plugins/Media/PixelStreaming/Source/PixelStreamingEditor/Private/PixelStreamingEditorModule.cpp
+++ b/Engine/Plugins/Media/PixelStreaming/Source/PixelStreamingEditor/Private/PixelStreamingEditorModule.cpp
@@ -51,18 +51,20 @@ void FPixelStreamingEditorModule::StartupModule()
 	Settings::InitialiseSettings();
 	bUseExternalSignallingServer = Settings::CVarEditorPixelStreamingUseRemoteSignallingServer.GetValueOnAnyThread();
 
-    // Enable experimental audio so we can mix the different editor audio outputs
-    IConsoleVariable* CVar = IConsoleManager::Get().FindConsoleVariable(TEXT("PixelStreaming.ExperimentalAudioInput")); 
-    CVar->Set(true);
-
-    FAudioDeviceManagerDelegates::OnAudioDeviceCreated.AddLambda([&](Audio::FDeviceId AudioDeviceId) {
-	    FAudioDeviceManager* DeviceManager = GEngine->GetAudioDeviceManager();
-	    FAudioDeviceHandle Device = DeviceManager->GetAudioDevice(AudioDeviceId);
-        AudioInputs.Add(AudioDeviceId, MakeShared<FEditorSubmixListener>(Device));
-    });
-    FAudioDeviceManagerDelegates::OnAudioDeviceDestroyed.AddLambda([&](Audio::FDeviceId AudioDeviceId) {
-        AudioInputs.Remove(AudioDeviceId);
-    });
+	// Enable experimental audio so we can mix the different editor audio outputs
+	IConsoleVariable* CVar = IConsoleManager::Get().FindConsoleVariable(TEXT("PixelStreaming.ExperimentalAudioInput")); 
+	CVar->Set(true);
+
+	FAudioDeviceManagerDelegates::OnAudioDeviceCreated.AddLambda([&](Audio::FDeviceId AudioDeviceId)
+	{
+		FAudioDeviceManager* DeviceManager = GEngine->GetAudioDeviceManager();
+		FAudioDeviceHandle Device = DeviceManager->GetAudioDevice(AudioDeviceId);
+		AudioInputs.Add(AudioDeviceId, MakeShared<FEditorSubmixListener, ESPMode::ThreadSafe>(Device));
+	});
+	FAudioDeviceManagerDelegates::OnAudioDeviceDestroyed.AddLambda([&](Audio::FDeviceId AudioDeviceId)
+	{
+		AudioInputs.Remove(AudioDeviceId);
+	});
 
 	IPixelStreamingModule& Module = IPixelStreamingModule::Get();
 	Module.OnReady().AddRaw(this, &FPixelStreamingEditorModule::InitEditorStreaming);
diff --git a/Engine/Plugins/Media/PixelStreaming/Source/PixelStreamingEditor/Private/PixelStreamingEditorModule.h b/Engine/Plugins/Media/PixelStreaming/Source/PixelStreamingEditor/Private/PixelStreamingEditorModule.h
--- a/Engine/Plugins/Media/PixelStreaming/Source/PixelStreamingEditor/Private/PixelStreamingEditorModule.h
+++ b/Engine/Plugins/Media/PixelStreaming/Source/PixelStreamingEditor/Private/PixelStreamingEditorModule.h
@@ -68,5 +68,5 @@ private:
 
 	bool bUseExternalSignallingServer = false;
 
-	TMap<Audio::FDeviceId, TSharedPtr<UE::EditorPixelStreaming::FEditorSubmixListener>> AudioInputs;
+	TMap<Audio::FDeviceId, TSharedPtr<UE::EditorPixelStreaming::FEditorSubmixListener, ESPMode::ThreadSafe>> AudioInputs;
 };
diff --git a/Engine/Source/Runtime/AudioLink/AudioLinkEngine/Private/BufferedSubmixListener.cpp b/Engine/Source/Runtime/AudioLink/AudioLinkEngine/Private/BufferedSubmixListener.cpp
--- a/Engine/Source/Runtime/AudioLink/AudioLinkEngine/Private/BufferedSubmixListener.cpp
+++ b/Engine/Source/Runtime/AudioLink/AudioLinkEngine/Private/BufferedSubmixListener.cpp
@@ -1,6 +1,7 @@
 // Copyright Epic Games, Inc. All Rights Reserved.
 
 #include "BufferedSubmixListener.h"
+#include "AudioMixerDevice.h"
 
 namespace BufferedSubmixListenerPrivate
 {
@@ -8,11 +9,15 @@ namespace BufferedSubmixListenerPrivate
 }
 
 /** Buffered Submix Listener. */
-FBufferedSubmixListener::FBufferedSubmixListener(int32 InDefaultCircularBufferSize, bool bInZeroInputBuffer)
+FBufferedSubmixListener::FBufferedSubmixListener(int32 InDefaultCircularBufferSize, bool bInZeroInputBuffer, const FString* InName)
 	: FBufferedListenerBase{ InDefaultCircularBufferSize }
 	, DeviceId{ BufferedSubmixListenerPrivate::InvalidAudioDeviceId }
 	, bZeroInputBuffer{ bInZeroInputBuffer }
 {
+	if (InName)
+	{
+		Name = *InName;
+	}
 }
 
 FBufferedSubmixListener::~FBufferedSubmixListener()
@@ -22,18 +27,27 @@ FBufferedSubmixListener::~FBufferedSubmixListener()
 	check(!IsStartedNonAtomic());
 }
 
+const FString& FBufferedSubmixListener::GetListenerName() const
+{
+	return Name;
+}
+
 void FBufferedSubmixListener::RegisterWithAudioDevice(FAudioDevice* InDevice)
 {
+	check(InDevice);
+
 	DeviceId = InDevice->DeviceID;
-	InDevice->RegisterSubmixBufferListener(this);
+	InDevice->RegisterSubmixBufferListener(AsShared(), InDevice->GetMainSubmixObject());
 }
 
 void FBufferedSubmixListener::UnregsiterWithAudioDevice(FAudioDevice* InDevice)
 {
+	check(InDevice);
+
 	if (ensure(DeviceId != BufferedSubmixListenerPrivate::InvalidAudioDeviceId))
 	{
 		DeviceId = BufferedSubmixListenerPrivate::InvalidAudioDeviceId;
-		InDevice->UnregisterSubmixBufferListener(this);
+		InDevice->UnregisterSubmixBufferListener(AsShared(), InDevice->GetMainSubmixObject());
 	}
 }
 
diff --git a/Engine/Source/Runtime/AudioLink/AudioLinkEngine/Private/BufferedSubmixListener.h b/Engine/Source/Runtime/AudioLink/AudioLinkEngine/Private/BufferedSubmixListener.h
--- a/Engine/Source/Runtime/AudioLink/AudioLinkEngine/Private/BufferedSubmixListener.h
+++ b/Engine/Source/Runtime/AudioLink/AudioLinkEngine/Private/BufferedSubmixListener.h
@@ -7,15 +7,15 @@
 
 /** Concrete Submix Buffer Listener  
 */
-class FBufferedSubmixListener : public ISubmixBufferListener, 
-													public FBufferedListenerBase
+class FBufferedSubmixListener : public ISubmixBufferListener, public FBufferedListenerBase
 {
 public:
 	/** Constructor
 	 * @param  InDefaultCircularBufferSize: Size of the circular buffer in samples by default.
-	 * @param  bInZeroInputBuffer: The passed in buffer will be set zerod after we've buffered it. Not 
+	 * @param  bInZeroInputBuffer: The passed in buffer will be set zeroed after we've buffered it.
+	 * @param InName(Optional): Optional name to track listener lifetime with. 
 	 */
-	AUDIOLINKENGINE_API FBufferedSubmixListener(int32 InDefaultCircularBufferSize, bool bInZeroInputBuffer);
+	AUDIOLINKENGINE_API FBufferedSubmixListener(int32 InDefaultCircularBufferSize, bool bInZeroInputBuffer, const FString* InName);
 
 	AUDIOLINKENGINE_API virtual ~FBufferedSubmixListener();
 
@@ -35,8 +35,10 @@ private:
 	
 	//~ Begin ISubmixBufferListener
 	AUDIOLINKENGINE_API void OnNewSubmixBuffer(const USoundSubmix* OwningSubmix, float* AudioData, int32 InNumSamples, int32 InNumChannels, const int32 InSampleRate, double) override;
+	AUDIOLINKENGINE_API const FString& GetListenerName() const override;
 	//~ End ISubmixBufferListener
 
 	Audio::FDeviceId DeviceId;
 	bool bZeroInputBuffer = false;
+	FString Name;
 };
diff --git a/Engine/Source/Runtime/AudioLink/AudioLinkEngine/Private/IAudioLinkFactory.cpp b/Engine/Source/Runtime/AudioLink/AudioLinkEngine/Private/IAudioLinkFactory.cpp
--- a/Engine/Source/Runtime/AudioLink/AudioLinkEngine/Private/IAudioLinkFactory.cpp
+++ b/Engine/Source/Runtime/AudioLink/AudioLinkEngine/Private/IAudioLinkFactory.cpp
@@ -80,7 +80,8 @@ FSharedBufferedOutputPtr IAudioLinkFactory::CreatePushableBufferListener(const F
 
 FSharedBufferedOutputPtr IAudioLinkFactory::CreateSubmixBufferListener(const FSubmixBufferListenerCreateParams& InSubmixCreateParams)
 {
-	return MakeShared<FBufferedSubmixListener, ESPMode::ThreadSafe>(InSubmixCreateParams.SizeOfBufferInFrames, InSubmixCreateParams.bShouldZeroBuffer);	
+	const FString ListenerName = FString::Format(TEXT("IAudioLinkFactory:{0}"), { *GetFactoryName().ToString() });
+	return MakeShared<FBufferedSubmixListener, ESPMode::ThreadSafe>(InSubmixCreateParams.SizeOfBufferInFrames, InSubmixCreateParams.bShouldZeroBuffer, &ListenerName);
 }
 
 TArray<IAudioLinkFactory*> IAudioLinkFactory::GetAllRegisteredFactories()
diff --git a/Engine/Source/Runtime/AudioMixer/Private/AudioMixerDevice.cpp b/Engine/Source/Runtime/AudioMixer/Private/AudioMixerDevice.cpp
--- a/Engine/Source/Runtime/AudioMixer/Private/AudioMixerDevice.cpp
+++ b/Engine/Source/Runtime/AudioMixer/Private/AudioMixerDevice.cpp
@@ -2,8 +2,11 @@
 
 #include "AudioMixerDevice.h"
 
+#include "AssetRegistry/IAssetRegistry.h"
+#include "Async/Async.h"
 #include "AudioAnalytics.h"
 #include "AudioBusSubsystem.h"
+#include "AudioDeviceNotificationSubsystem.h"
 #include "AudioMixerSource.h"
 #include "AudioMixerSourceManager.h"
 #include "AudioMixerSourceDecode.h"
@@ -22,11 +25,9 @@
 #include "UObject/UObjectHash.h"
 #include "UObject/UObjectIterator.h"
 #include "IHeadMountedDisplayModule.h"
+#include "ISubmixBufferListener.h"
 #include "Misc/App.h"
 #include "ProfilingDebugging/CsvProfiler.h"
-#include "AssetRegistry/IAssetRegistry.h"
-#include "Async/Async.h"
-#include "AudioDeviceNotificationSubsystem.h"
 #include "Sound/AudioFormatSettings.h"
 
 #if WITH_EDITOR
@@ -560,7 +561,7 @@ namespace Audio
 		}
 
 		// Reset existing submixes if they exist
-		MasterSubmixInstances.Reset();
+		RequiredSubmixInstances.Reset();
 		Submixes.Reset();
 	}
 
@@ -795,16 +796,16 @@ namespace Audio
 		// Compute the next block of audio in the source manager
 		SourceManager->ComputeNextBlockOfSamples();
 
-		FMixerSubmixWeakPtr MasterSubmix = GetMasterSubmix();
+		FMixerSubmixWeakPtr MainSubmix = GetMasterSubmix();
 		{
 			CSV_SCOPED_TIMING_STAT(Audio, Submixes);
 			SCOPE_CYCLE_COUNTER(STAT_AudioMixerSubmixes);
 
-			FMixerSubmixPtr MasterSubmixPtr = MasterSubmix.Pin();
-			if (MasterSubmixPtr.IsValid())
+			FMixerSubmixPtr MainSubmixPtr = MainSubmix.Pin();
+			if (MainSubmixPtr.IsValid())
 			{
 				// Process the audio output from the master submix
-				MasterSubmixPtr->ProcessAudio(Output);
+				MainSubmixPtr->ProcessAudio(Output);
 			}
 		}
 
@@ -873,23 +874,23 @@ namespace Audio
 		SourceManager->UpdatePendingReleaseData(true);
 	}
 
-	void FMixerDevice::LoadMasterSoundSubmix(EMasterSubmixType::Type InType, const FString& InDefaultName, bool bInDefaultMuteWhenBackgrounded, FSoftObjectPath& InObjectPath)
+	void FMixerDevice::LoadRequiredSubmix(ERequiredSubmixes::Type InType, const FString& InDefaultName, bool bInDefaultMuteWhenBackgrounded, FSoftObjectPath& InObjectPath)
 	{
 		check(IsInGameThread());
 
-		const int32 MasterSubmixCount = static_cast<int32>(EMasterSubmixType::Type::Count);
-		if(MasterSubmixes.Num() < MasterSubmixCount)
+		const int32 RequiredSubmixCount = static_cast<int32>(ERequiredSubmixes::Count);
+		if(RequiredSubmixes.Num() < RequiredSubmixCount)
 		{
-			MasterSubmixes.AddZeroed(MasterSubmixCount - MasterSubmixes.Num());
+			RequiredSubmixes.AddZeroed(RequiredSubmixCount - RequiredSubmixes.Num());
 		}
 
-		if (MasterSubmixInstances.Num() < MasterSubmixCount)
+		if (RequiredSubmixInstances.Num() < RequiredSubmixCount)
 		{
-			MasterSubmixInstances.AddZeroed(MasterSubmixCount - MasterSubmixInstances.Num());
+			RequiredSubmixInstances.AddZeroed(RequiredSubmixCount - RequiredSubmixInstances.Num());
 		}
 
 		const int32 TypeIndex = static_cast<int32>(InType);
-		if (USoundSubmix* OldSubmix = MasterSubmixes[TypeIndex])
+		if (USoundSubmix* OldSubmix = RequiredSubmixes[TypeIndex])
 		{
 			// Don't bother swapping if new path is invalid...
 			if (!InObjectPath.IsValid())
@@ -903,13 +904,13 @@ namespace Audio
 				return;
 			}
 			OldSubmix->RemoveFromRoot();
-			FMixerSubmixPtr OldSubmixPtr = MasterSubmixInstances[TypeIndex];
+			FMixerSubmixPtr OldSubmixPtr = RequiredSubmixInstances[TypeIndex];
 			if (OldSubmixPtr.IsValid())
 			{
-				FMixerSubmixPtr ParentSubmixPtr = MasterSubmixInstances[TypeIndex]->GetParentSubmix().Pin();
+				FMixerSubmixPtr ParentSubmixPtr = RequiredSubmixInstances[TypeIndex]->GetParentSubmix().Pin();
 				if (ParentSubmixPtr.IsValid())
 				{
-					ParentSubmixPtr->RemoveChildSubmix(MasterSubmixInstances[TypeIndex]);
+					ParentSubmixPtr->RemoveChildSubmix(RequiredSubmixInstances[TypeIndex]);
 				}
 			}
 		}
@@ -944,10 +945,10 @@ namespace Audio
 		check(NewSubmix);
 		NewSubmix->AddToRoot();
 
-		// If sharing submix with other explicitly defined MasterSubmix, create
+		// If sharing submix with other explicitly defined MainSubmix, create
 		// shared pointer directed to already existing submix instance. Otherwise,
 		// create a new version.
-		FMixerSubmixPtr NewMixerSubmix = GetMasterSubmixInstance(NewSubmix);
+		FMixerSubmixPtr NewMixerSubmix = GetRequiredSubmixInstance(NewSubmix);
 		if (!NewMixerSubmix.IsValid())
 		{
 			UE_LOG(LogAudioMixer, Display, TEXT("Creating Master Submix '%s'"), *NewSubmix->GetName());
@@ -955,15 +956,15 @@ namespace Audio
 		}
 
 		// Ensure that master submixes are ONLY tracked in master submix array.
-		// MasterSubmixes array can share instances, but should not be duplicated in Submixes Map.
+		// RequiredSubmixes array can share instances, but should not be duplicated in Submixes Map.
 		if (Submixes.Remove(NewSubmix->GetUniqueID()) > 0)
 		{
 			UE_LOG(LogAudioMixer, Display, TEXT("Submix '%s' has been promoted to master array."), *NewSubmix->GetName());
 		}
 
 		// Update/add new submix and instance to respective master arrays
-		MasterSubmixes[TypeIndex] = NewSubmix;
-		MasterSubmixInstances[TypeIndex] = NewMixerSubmix;
+		RequiredSubmixes[TypeIndex] = NewSubmix;
+		RequiredSubmixInstances[TypeIndex] = NewMixerSubmix;
 
 		//Note: If we support using endpoint/soundfield submixes as a master submix in the future, we will need to call NewMixerSubmix->SetSoundfieldFactory here.
 		NewMixerSubmix->Init(NewSubmix, false /* bAllowReInit */);
@@ -1022,7 +1023,7 @@ namespace Audio
 			UAudioSettings* AudioSettings = GetMutableDefault<UAudioSettings>();
 			check(AudioSettings);
 
-			if (MasterSubmixes.Num() > 0)
+			if (RequiredSubmixes.Num() > 0)
 			{
 				UE_LOG(LogAudioMixer, Display, TEXT("Re-initializing Sound Submixes..."));
 			}
@@ -1032,19 +1033,19 @@ namespace Audio
 			}
 
 			// 1. Load or reload all sound submixes/instances
-			LoadMasterSoundSubmix(EMasterSubmixType::Master, TEXT("MasterSubmixDefault"), false /* DefaultMuteWhenBackgrounded */, AudioSettings->MasterSubmix);
+			LoadRequiredSubmix(ERequiredSubmixes::Main, TEXT("MasterSubmixDefault"), false /* DefaultMuteWhenBackgrounded */, AudioSettings->MasterSubmix);
 
 			// BaseDefaultSubmix is an optional master submix type set by project settings
 			if (AudioSettings->BaseDefaultSubmix.IsValid())
 			{
-				LoadMasterSoundSubmix(EMasterSubmixType::BaseDefault, TEXT("BaseDefault"), false /* DefaultMuteWhenBackgrounded */, AudioSettings->BaseDefaultSubmix);
+				LoadRequiredSubmix(ERequiredSubmixes::BaseDefault, TEXT("BaseDefault"), false /* DefaultMuteWhenBackgrounded */, AudioSettings->BaseDefaultSubmix);
 			}
 
-			LoadMasterSoundSubmix(EMasterSubmixType::Reverb, TEXT("MasterReverbSubmixDefault"), true /* DefaultMuteWhenBackgrounded */, AudioSettings->ReverbSubmix);
+			LoadRequiredSubmix(ERequiredSubmixes::Reverb, TEXT("MasterReverbSubmixDefault"), true /* DefaultMuteWhenBackgrounded */, AudioSettings->ReverbSubmix);
 
 			if (!DisableSubmixEffectEQCvar)
 			{
-				LoadMasterSoundSubmix(EMasterSubmixType::EQ, TEXT("MasterEQSubmixDefault"), false /* DefaultMuteWhenBackgrounded */, AudioSettings->EQSubmix);
+				LoadRequiredSubmix(ERequiredSubmixes::EQ, TEXT("MasterEQSubmixDefault"), false /* DefaultMuteWhenBackgrounded */, AudioSettings->EQSubmix);
 			}
 
 			LoadPluginSoundSubmixes();
@@ -1054,7 +1055,7 @@ namespace Audio
 				USoundSubmixBase* SubmixToLoad = *It;
 				check(SubmixToLoad);
 
-				if (!IsMasterSubmixType(SubmixToLoad))
+				if (!IsRequiredSubmixType(SubmixToLoad))
 				{
 					LoadSoundSubmix(*SubmixToLoad);
 					InitSoundfieldAndEndpointDataForSubmix(*SubmixToLoad, GetSubmixInstance(SubmixToLoad).Pin(), false);
@@ -1075,19 +1076,19 @@ namespace Audio
 			return;
 		}
 
-		for (int32 i = 0; i < static_cast<int32>(EMasterSubmixType::Count); ++i)
+		for (int32 i = 0; i < static_cast<int32>(ERequiredSubmixes::Count); ++i)
 		{
-			if (DisableSubmixEffectEQCvar && i == static_cast<int32>(EMasterSubmixType::EQ))
+			if (DisableSubmixEffectEQCvar && i == static_cast<int32>(ERequiredSubmixes::EQ))
 			{
 				continue;
 			}
 
-			USoundSubmixBase* SoundSubmix = MasterSubmixes[i];
-			if (SoundSubmix && SoundSubmix != MasterSubmixes[static_cast<int32>(EMasterSubmixType::Master)])
+			USoundSubmixBase* SoundSubmix = RequiredSubmixes[i];
+			if (SoundSubmix && SoundSubmix != RequiredSubmixes[static_cast<int32>(ERequiredSubmixes::Main)])
 			{
-				FMixerSubmixPtr& MasterSubmixInstance = MasterSubmixInstances[i];
+				FMixerSubmixPtr& MainSubmixInstance = RequiredSubmixInstances[i];
 
-				RebuildSubmixLinks(*SoundSubmix, MasterSubmixInstance);
+				RebuildSubmixLinks(*SoundSubmix, MainSubmixInstance);
 			}
 		}
 
@@ -1095,7 +1096,7 @@ namespace Audio
 		{
 			if (const USoundSubmixBase* SubmixBase = *It)
 			{
-				if (IsMasterSubmixType(SubmixBase))
+				if (IsRequiredSubmixType(SubmixBase))
 				{
 					continue;
 				}
@@ -1122,7 +1123,7 @@ namespace Audio
 			else
 			{
 				// If this submix is itself the broadcast submix, set its parent to the master submix
-				if (SubmixInstance == MasterSubmixInstances[static_cast<int32>(EMasterSubmixType::BaseDefault)])
+				if (SubmixInstance == RequiredSubmixInstances[static_cast<int32>(ERequiredSubmixes::BaseDefault)])
 				{
 					ParentSubmixInstance = GetMasterSubmix().Pin();
 				}
@@ -1158,49 +1159,79 @@ namespace Audio
 
 	FMixerSubmixWeakPtr FMixerDevice::GetMasterSubmix()
 	{
-		return MasterSubmixInstances[EMasterSubmixType::Master];
+		return GetMainSubmix();
+	}
+
+	FMixerSubmixWeakPtr FMixerDevice::GetMainSubmix()
+	{
+		return RequiredSubmixInstances[ERequiredSubmixes::Main];
 	}
 
 	FMixerSubmixWeakPtr FMixerDevice::GetBaseDefaultSubmix()
 	{
-		if (MasterSubmixInstances[EMasterSubmixType::BaseDefault].IsValid())
+		if (RequiredSubmixInstances[ERequiredSubmixes::BaseDefault].IsValid())
 		{
-			return MasterSubmixInstances[EMasterSubmixType::BaseDefault];
+			return RequiredSubmixInstances[ERequiredSubmixes::BaseDefault];
 		}
 		return GetMasterSubmix();
 	}
 
+	FMixerSubmixWeakPtr FMixerDevice::GetReverbSubmix()
+	{
+		return RequiredSubmixInstances[ERequiredSubmixes::Reverb];
+	}
+
+	FMixerSubmixWeakPtr FMixerDevice::GetEQSubmix()
+	{
+		return RequiredSubmixInstances[ERequiredSubmixes::EQ];
+	}
+
 	FMixerSubmixWeakPtr FMixerDevice::GetMasterReverbSubmix()
 	{
-		return MasterSubmixInstances[EMasterSubmixType::Reverb];
+		return RequiredSubmixInstances[ERequiredSubmixes::Reverb];
 	}
 
 	FMixerSubmixWeakPtr FMixerDevice::GetMasterEQSubmix()
 	{
-		return MasterSubmixInstances[EMasterSubmixType::EQ];
+		return RequiredSubmixInstances[ERequiredSubmixes::EQ];
 	}
 
-	void FMixerDevice::AddMasterSubmixEffect(FSoundEffectSubmixPtr SoundEffectSubmix)
+	void FMixerDevice::AddMainSubmixEffect(FSoundEffectSubmixPtr SoundEffectSubmix)
 	{
 		AudioRenderThreadCommand([this, SoundEffectSubmix]()
 		{
-			MasterSubmixInstances[EMasterSubmixType::Master]->AddSoundEffectSubmix(SoundEffectSubmix);
+			RequiredSubmixInstances[ERequiredSubmixes::Main]->AddSoundEffectSubmix(SoundEffectSubmix);
 		});
 	}
 
-	void FMixerDevice::RemoveMasterSubmixEffect(uint32 SubmixEffectId)
+	void FMixerDevice::AddMasterSubmixEffect(FSoundEffectSubmixPtr SoundEffectSubmix)
+	{
+		AddMainSubmixEffect(SoundEffectSubmix);
+	}
+
+	void FMixerDevice::RemoveMainSubmixEffect(uint32 SubmixEffectId)
 	{
 		AudioRenderThreadCommand([this, SubmixEffectId]()
 		{
-			MasterSubmixInstances[EMasterSubmixType::Master]->RemoveSoundEffectSubmix(SubmixEffectId);
+			RequiredSubmixInstances[ERequiredSubmixes::Main]->RemoveSoundEffectSubmix(SubmixEffectId);
 		});
 	}
 
+	void FMixerDevice::RemoveMasterSubmixEffect(uint32 SubmixEffectId)
+	{
+		RemoveMainSubmixEffect(SubmixEffectId);
+	}
+
 	void FMixerDevice::ClearMasterSubmixEffects()
+	{
+		ClearMainSubmixEffects();
+	}
+
+	void FMixerDevice::ClearMainSubmixEffects()
 	{
 		AudioRenderThreadCommand([this]()
 		{
-			MasterSubmixInstances[EMasterSubmixType::Master]->ClearSoundEffectSubmixes();
+			RequiredSubmixInstances[ERequiredSubmixes::Main]->ClearSoundEffectSubmixes();
 		});
 	}
 
@@ -1594,11 +1625,11 @@ namespace Audio
 		}
 	}
 
-	bool FMixerDevice::IsMasterSubmixType(const USoundSubmixBase* InSubmix) const
+	bool FMixerDevice::IsRequiredSubmixType(const USoundSubmixBase* InSubmix) const
 	{
-		for (int32 i = 0; i < EMasterSubmixType::Count; ++i)
+		for (int32 i = 0; i < ERequiredSubmixes::Count; ++i)
 		{
-			if (InSubmix == MasterSubmixes[i])
+			if (InSubmix == RequiredSubmixes[i])
 			{
 				return true;
 			}
@@ -1606,27 +1637,27 @@ namespace Audio
 		return false;
 	}
 
-	FMixerSubmixPtr FMixerDevice::GetMasterSubmixInstance(uint32 InObjectId)
+	FMixerSubmixPtr FMixerDevice::GetRequiredSubmixInstance(uint32 InObjectId)
 	{
-		check(MasterSubmixes.Num() == EMasterSubmixType::Count);
-		for (int32 i = 0; i < EMasterSubmixType::Count; ++i)
+		check(RequiredSubmixes.Num() == ERequiredSubmixes::Count);
+		for (int32 i = 0; i < ERequiredSubmixes::Count; ++i)
 		{
-			if (InObjectId == MasterSubmixes[i]->GetUniqueID())
+			if (InObjectId == RequiredSubmixes[i]->GetUniqueID())
 			{
-				return MasterSubmixInstances[i];
+				return RequiredSubmixInstances[i];
 			}
 		}
 		return nullptr;
 	}
 
-	FMixerSubmixPtr FMixerDevice::GetMasterSubmixInstance(const USoundSubmixBase* InSubmix)
+	FMixerSubmixPtr FMixerDevice::GetRequiredSubmixInstance(const USoundSubmixBase* InSubmix)
 	{
-		check(MasterSubmixes.Num() == EMasterSubmixType::Count);
-		for (int32 i = 0; i < EMasterSubmixType::Count; ++i)
+		check(RequiredSubmixes.Num() == ERequiredSubmixes::Count);
+		for (int32 i = 0; i < ERequiredSubmixes::Count; ++i)
 		{
-			if (InSubmix == MasterSubmixes[i])
+			if (InSubmix == RequiredSubmixes[i])
 			{
-				return MasterSubmixInstances[i];
+				return RequiredSubmixInstances[i];
 			}
 		}
 		return nullptr;
@@ -1660,9 +1691,9 @@ namespace Audio
 
 		UE_LOG(LogAudioMixer, Display, TEXT("Registering submix %s."), *InSoundSubmix->GetFullName());
 
-		const bool bIsMasterSubmix = IsMasterSubmixType(InSoundSubmix);
+		const bool bIsMainSubmix = IsRequiredSubmixType(InSoundSubmix);
 
-		if (!bIsMasterSubmix)
+		if (!bIsMainSubmix)
 		{
 			// Ensure parent structure is registered prior to current submix if missing
 			const USoundSubmixWithParentBase* SubmixWithParent = Cast<const USoundSubmixWithParentBase>(InSoundSubmix);
@@ -1689,7 +1720,7 @@ namespace Audio
 			InitSoundfieldAndEndpointDataForSubmix(*InSoundSubmix, SubmixPtr, true);
 		}
 
-		if (!bIsMasterSubmix)
+		if (!bIsMainSubmix)
 		{
 			RebuildSubmixLinks(*InSoundSubmix, SubmixPtr);
 		}
@@ -1748,7 +1779,7 @@ namespace Audio
 
 	void FMixerDevice::UnregisterSoundSubmix(const USoundSubmixBase* InSoundSubmix)
 	{
-		if (!InSoundSubmix || bSubmixRegistrationDisabled || IsMasterSubmixType(InSoundSubmix))
+		if (!InSoundSubmix || bSubmixRegistrationDisabled || IsRequiredSubmixType(InSoundSubmix))
 		{
 			return;
 		}
@@ -1776,7 +1807,7 @@ namespace Audio
 	{
 		check(IsInAudioThread());
 
-		FMixerSubmixWeakPtr MasterSubmix = GetMasterSubmix();
+		FMixerSubmixWeakPtr MainSubmix = GetMasterSubmix();
 
 		// Check if this is a submix type that has a parent.
 		FMixerSubmixPtr ParentSubmixInstance;
@@ -1784,7 +1815,7 @@ namespace Audio
 		{
 			ParentSubmixInstance = InSoundSubmixWithParent->ParentSubmix
 				? GetSubmixInstance(InSoundSubmixWithParent->ParentSubmix).Pin()
-				: MasterSubmix.Pin();
+				: MainSubmix.Pin();
 		}
 
 		if (ParentSubmixInstance.IsValid())
@@ -1799,7 +1830,7 @@ namespace Audio
 			{
 				ChildSubmixPtr->SetParentSubmix(ParentSubmixInstance.IsValid()
 					? ParentSubmixInstance
-					: MasterSubmix);
+					: MainSubmix);
 			}
 		}
 
@@ -1830,23 +1861,23 @@ namespace Audio
 
 	FMixerSubmixPtr FMixerDevice::FindSubmixInstanceByObjectId(uint32 InObjectId)
 	{
-		for (int32 i = 0; i < MasterSubmixes.Num(); i++)
+		for (int32 i = 0; i < RequiredSubmixes.Num(); i++)
 		{
-			if (const USoundSubmix* MasterSubmix = MasterSubmixes[i])
+			if (const USoundSubmix* MainSubmix = RequiredSubmixes[i])
 			{
-				if (MasterSubmix->GetUniqueID() == InObjectId)
+				if (MainSubmix->GetUniqueID() == InObjectId)
 				{
-					return GetMasterSubmixInstance(MasterSubmix);
+					return GetRequiredSubmixInstance(MainSubmix);
 				}
 			}
 			else
 			{
-				const EMasterSubmixType::Type SubmixType = static_cast<EMasterSubmixType::Type>(i);
-				ensureAlwaysMsgf(EMasterSubmixType::Master != SubmixType,
-					TEXT("Top-level master has to be registered before anything else, and is required for the lifetime of the application.")
+				const ERequiredSubmixes::Type SubmixType = static_cast<ERequiredSubmixes::Type>(i);
+				ensureAlwaysMsgf(ERequiredSubmixes::Main != SubmixType,
+					TEXT("Main submix has to be registered before anything else, and is required for the lifetime of the application.")
 				);
 
-				if (!DisableSubmixEffectEQCvar && EMasterSubmixType::EQ == SubmixType)
+				if (!DisableSubmixEffectEQCvar && ERequiredSubmixes::EQ == SubmixType)
 				{
 					UE_LOG(LogAudioMixer, Warning, TEXT("Failed to query EQ Submix when it was expected to be loaded."));
 				}
@@ -1868,7 +1899,7 @@ namespace Audio
 	{
 		LLM_SCOPE(ELLMTag::AudioMixer);
 
-		FMixerSubmixPtr MixerSubmix = GetMasterSubmixInstance(SoundSubmix);
+		FMixerSubmixPtr MixerSubmix = GetRequiredSubmixInstance(SoundSubmix);
 		if (MixerSubmix.IsValid())
 		{
 			return MixerSubmix;
@@ -2144,11 +2175,11 @@ namespace Audio
 		}
 		else
 		{
-			FMixerSubmixWeakPtr MasterSubmix = GetMasterSubmix();
-			FMixerSubmixPtr MasterSubmixPtr = MasterSubmix.Pin();
-			check(MasterSubmixPtr.IsValid());
+			FMixerSubmixWeakPtr MainSubmix = GetMasterSubmix();
+			FMixerSubmixPtr MainSubmixPtr = MainSubmix.Pin();
+			check(MainSubmixPtr.IsValid());
 
-			MasterSubmixPtr->OnStartRecordingOutput(ExpectedRecordingDuration);
+			MainSubmixPtr->OnStartRecordingOutput(ExpectedRecordingDuration);
 		}
 	}
 
@@ -2162,10 +2193,10 @@ namespace Audio
 		}
 		else
 		{
-			FMixerSubmixPtr MasterSubmixPtr = GetMasterSubmix().Pin();
-			check(MasterSubmixPtr.IsValid());
+			FMixerSubmixPtr MainSubmixPtr = GetMasterSubmix().Pin();
+			check(MainSubmixPtr.IsValid());
 
-			return MasterSubmixPtr->OnStopRecordingOutput(OutNumChannels, OutSampleRate);
+			return MainSubmixPtr->OnStopRecordingOutput(OutNumChannels, OutSampleRate);
 		}
 	}
 
@@ -2191,11 +2222,11 @@ namespace Audio
 		}
 		else
 		{
-			FMixerSubmixWeakPtr MasterSubmix = GetMasterSubmix();
-			FMixerSubmixPtr MasterSubmixPtr = MasterSubmix.Pin();
-			check(MasterSubmixPtr.IsValid());
+			FMixerSubmixWeakPtr MainSubmix = GetMasterSubmix();
+			FMixerSubmixPtr MainSubmixPtr = MainSubmix.Pin();
+			check(MainSubmixPtr.IsValid());
 
-			MasterSubmixPtr->PauseRecordingOutput();
+			MainSubmixPtr->PauseRecordingOutput();
 		}
 	}
 
@@ -2221,11 +2252,11 @@ namespace Audio
 		}
 		else
 		{
-			FMixerSubmixWeakPtr MasterSubmix = GetMasterSubmix();
-			FMixerSubmixPtr MasterSubmixPtr = MasterSubmix.Pin();
-			check(MasterSubmixPtr.IsValid());
+			FMixerSubmixWeakPtr MainSubmix = GetMasterSubmix();
+			FMixerSubmixPtr MainSubmixPtr = MainSubmix.Pin();
+			check(MainSubmixPtr.IsValid());
 
-			MasterSubmixPtr->ResumeRecordingOutput();
+			MainSubmixPtr->ResumeRecordingOutput();
 		}
 	}
 
@@ -2251,11 +2282,11 @@ namespace Audio
 		}
 		else
 		{
-			FMixerSubmixWeakPtr MasterSubmix = GetMasterSubmix();
-			FMixerSubmixPtr MasterSubmixPtr = MasterSubmix.Pin();
-			check(MasterSubmixPtr.IsValid());
+			FMixerSubmixWeakPtr MainSubmix = GetMasterSubmix();
+			FMixerSubmixPtr MainSubmixPtr = MainSubmix.Pin();
+			check(MainSubmixPtr.IsValid());
 
-			MasterSubmixPtr->StartEnvelopeFollowing(InSubmix->EnvelopeFollowerAttackTime, InSubmix->EnvelopeFollowerReleaseTime);
+			MainSubmixPtr->StartEnvelopeFollowing(InSubmix->EnvelopeFollowerAttackTime, InSubmix->EnvelopeFollowerReleaseTime);
 		}
 
 		DelegateBoundSubmixes.AddUnique(InSubmix);
@@ -2283,11 +2314,11 @@ namespace Audio
 		}
 		else
 		{
-			FMixerSubmixWeakPtr MasterSubmix = GetMasterSubmix();
-			FMixerSubmixPtr MasterSubmixPtr = MasterSubmix.Pin();
-			check(MasterSubmixPtr.IsValid());
+			FMixerSubmixWeakPtr MainSubmix = GetMasterSubmix();
+			FMixerSubmixPtr MainSubmixPtr = MainSubmix.Pin();
+			check(MainSubmixPtr.IsValid());
 
-			MasterSubmixPtr->StopEnvelopeFollowing();
+			MainSubmixPtr->StopEnvelopeFollowing();
 		}
 
 		DelegateBoundSubmixes.RemoveSingleSwap(InSubmix);
@@ -2315,11 +2346,11 @@ namespace Audio
 		}
 		else
 		{
-			FMixerSubmixWeakPtr MasterSubmix = GetMasterSubmix();
-			FMixerSubmixPtr MasterSubmixPtr = MasterSubmix.Pin();
-			check(MasterSubmixPtr.IsValid());
+			FMixerSubmixWeakPtr MainSubmix = GetMasterSubmix();
+			FMixerSubmixPtr MainSubmixPtr = MainSubmix.Pin();
+			check(MainSubmixPtr.IsValid());
 
-			MasterSubmixPtr->AddEnvelopeFollowerDelegate(OnSubmixEnvelopeBP);
+			MainSubmixPtr->AddEnvelopeFollowerDelegate(OnSubmixEnvelopeBP);
 		}
 	}
 
@@ -2345,11 +2376,11 @@ namespace Audio
 		}
 		else
 		{
-			FMixerSubmixWeakPtr MasterSubmix = GetMasterSubmix();
-			FMixerSubmixPtr MasterSubmixPtr = MasterSubmix.Pin();
-			check(MasterSubmixPtr.IsValid());
+			FMixerSubmixWeakPtr MainSubmix = GetMasterSubmix();
+			FMixerSubmixPtr MainSubmixPtr = MainSubmix.Pin();
+			check(MainSubmixPtr.IsValid());
 
-			MasterSubmixPtr->StartSpectrumAnalysis(InSettings);
+			MainSubmixPtr->StartSpectrumAnalysis(InSettings);
 		}
 
 		DelegateBoundSubmixes.AddUnique(InSubmix);
@@ -2376,11 +2407,11 @@ namespace Audio
 		}
 		else
 		{
-			FMixerSubmixWeakPtr MasterSubmix = GetMasterSubmix();
-			FMixerSubmixPtr MasterSubmixPtr = MasterSubmix.Pin();
-			check(MasterSubmixPtr.IsValid());
+			FMixerSubmixWeakPtr MainSubmix = GetMasterSubmix();
+			FMixerSubmixPtr MainSubmixPtr = MainSubmix.Pin();
+			check(MainSubmixPtr.IsValid());
 
-			MasterSubmixPtr->StopSpectrumAnalysis();
+			MainSubmixPtr->StopSpectrumAnalysis();
 		}
 
 		DelegateBoundSubmixes.RemoveSingleSwap(InSubmix);
@@ -2396,11 +2427,11 @@ namespace Audio
 		}
 		else
 		{
-			FMixerSubmixWeakPtr MasterSubmix = GetMasterSubmix();
-			FMixerSubmixPtr MasterSubmixPtr = MasterSubmix.Pin();
-			check(MasterSubmixPtr.IsValid());
+			FMixerSubmixWeakPtr MainSubmix = GetMasterSubmix();
+			FMixerSubmixPtr MainSubmixPtr = MainSubmix.Pin();
+			check(MainSubmixPtr.IsValid());
 
-			MasterSubmixPtr->GetMagnitudeForFrequencies(InFrequencies, OutMagnitudes);
+			MainSubmixPtr->GetMagnitudeForFrequencies(InFrequencies, OutMagnitudes);
 		}
 	}
 
@@ -2413,11 +2444,11 @@ namespace Audio
 		}
 		else
 		{
-			FMixerSubmixWeakPtr MasterSubmix = GetMasterSubmix();
-			FMixerSubmixPtr MasterSubmixPtr = MasterSubmix.Pin();
-			check(MasterSubmixPtr.IsValid());
+			FMixerSubmixWeakPtr MainSubmix = GetMasterSubmix();
+			FMixerSubmixPtr MainSubmixPtr = MainSubmix.Pin();
+			check(MainSubmixPtr.IsValid());
 
-			MasterSubmixPtr->GetPhaseForFrequencies(InFrequencies, OutPhases);
+			MainSubmixPtr->GetPhaseForFrequencies(InFrequencies, OutPhases);
 		}
 	}
 
@@ -2442,8 +2473,8 @@ namespace Audio
 		if (!FoundSubmix.IsValid())
 		{
 			// If can't find the submix isntance, use master submix.
-			FMixerSubmixWeakPtr MasterSubmix = GetMasterSubmix();
-			FoundSubmix = MasterSubmix.Pin();
+			FMixerSubmixWeakPtr MainSubmix = GetMasterSubmix();
+			FoundSubmix = MainSubmix.Pin();
 		}
 
 		if (ensure(FoundSubmix.IsValid()))
@@ -2472,8 +2503,8 @@ namespace Audio
 		if (!FoundSubmix.IsValid())
 		{
 			// If can't find the submix isntance, use master submix.
-			FMixerSubmixWeakPtr MasterSubmix = GetMasterSubmix();
-			FoundSubmix = MasterSubmix.Pin();
+			FMixerSubmixWeakPtr MainSubmix = GetMasterSubmix();
+			FoundSubmix = MainSubmix.Pin();
 		}
 
 		if (ensure(FoundSubmix.IsValid()))
@@ -2482,6 +2513,14 @@ namespace Audio
 		}
 	}
 
+	USoundSubmix& FMixerDevice::GetMainSubmixObject() const
+	{
+		const int32 SubmixIndex = static_cast<int32>(ERequiredSubmixes::Main);
+		USoundSubmix* MainSubmix = RequiredSubmixes[SubmixIndex];
+		check(MainSubmix);
+		return *MainSubmix;
+	}
+
 	void FMixerDevice::RegisterSubmixBufferListener(ISubmixBufferListener* InSubmixBufferListener, USoundSubmix* InSubmix)
 	{
 		DECLARE_CYCLE_STAT(TEXT("FAudioThreadTask.RegisterSubmixBufferListener"), STAT_RegisterSubmixBufferListener, STATGROUP_AudioThreadCommands);
@@ -2506,7 +2545,9 @@ namespace Audio
 
 			if (FoundSubmix.IsValid())
 			{
+				PRAGMA_DISABLE_DEPRECATION_WARNINGS
 				FoundSubmix->RegisterBufferListener(InSubmixBufferListener);
+				PRAGMA_ENABLE_DEPRECATION_WARNINGS
 			}
 			else
 			{
@@ -2517,6 +2558,40 @@ namespace Audio
 		FAudioThread::RunCommandOnAudioThread(MoveTemp(RegisterLambda));
 	}
 
+	void FMixerDevice::RegisterSubmixBufferListener(TSharedRef<ISubmixBufferListener, ESPMode::ThreadSafe> InSubmixBufferListener, USoundSubmix& InSubmix)
+	{
+		DECLARE_CYCLE_STAT(TEXT("FAudioThreadTask.RegisterSubmixBufferListener"), STAT_RegisterSubmixBufferListener, STATGROUP_AudioThreadCommands);
+
+		const TWeakObjectPtr<USoundSubmix> SubmixPtr(&InSubmix);
+
+		auto RegisterLambda = [this, InSubmixBufferListener, SubmixPtr]()
+		{
+			CSV_SCOPED_TIMING_STAT(Audio, RegisterSubmixBufferListener);
+
+			FMixerSubmixPtr FoundSubmix = GetSubmixInstance(SubmixPtr.Get()).Pin();
+
+			// Attempt to register submix if instance not found and is not master (i.e. default) submix
+			if (!FoundSubmix.IsValid() && SubmixPtr.IsValid())
+			{
+				RegisterSoundSubmix(SubmixPtr.Get(), true /* bInit */);
+				FoundSubmix = GetSubmixInstance(SubmixPtr.Get()).Pin();
+			}
+
+			const FString& ListenerName = InSubmixBufferListener->GetListenerName();
+			if (FoundSubmix.IsValid())
+			{
+				FoundSubmix->RegisterBufferListener(InSubmixBufferListener);
+				UE_LOG(LogAudioMixer, Display, TEXT("Submix buffer listener '%s' registered with submix '%s'"), *ListenerName, *FoundSubmix->SubmixName);
+			}
+			else
+			{
+				UE_LOG(LogAudioMixer, Warning, TEXT("Submix buffer listener '%s' not registered. Submix not loaded."), *ListenerName);
+			}
+		};
+
+		FAudioThread::RunCommandOnAudioThread(MoveTemp(RegisterLambda));
+	}
+
 	void FMixerDevice::UnregisterSubmixBufferListener(ISubmixBufferListener* InSubmixBufferListener, USoundSubmix* InSubmix)
 	{
 		DECLARE_CYCLE_STAT(TEXT("FAudioThreadTask.UnregisterSubmixBufferListener"), STAT_UnregisterSubmixBufferListener, STATGROUP_AudioThreadCommands);
@@ -2534,7 +2609,9 @@ namespace Audio
 
 			if (FoundSubmix.IsValid())
 			{
+				PRAGMA_DISABLE_DEPRECATION_WARNINGS
 				FoundSubmix->UnregisterBufferListener(InSubmixBufferListener);
+				PRAGMA_ENABLE_DEPRECATION_WARNINGS
 			}
 			else
 			{
@@ -2545,6 +2622,31 @@ namespace Audio
 		FAudioThread::RunCommandOnAudioThread(MoveTemp(UnregisterLambda));
 	}
 
+	void FMixerDevice::UnregisterSubmixBufferListener(TSharedRef<ISubmixBufferListener, ESPMode::ThreadSafe> InSubmixBufferListener, USoundSubmix& InSubmix)
+	{
+		DECLARE_CYCLE_STAT(TEXT("FAudioThreadTask.UnregisterSubmixBufferListener"), STAT_UnregisterSubmixBufferListener, STATGROUP_AudioThreadCommands);
+
+		const TWeakObjectPtr<USoundSubmix> SubmixPtr(&InSubmix);
+
+		auto UnregisterLambda = [this, InSubmixBufferListener, SubmixPtr]()
+		{
+			CSV_SCOPED_TIMING_STAT(Audio, UnregisterSubmixBufferListener);
+
+			FMixerSubmixPtr FoundSubmix = GetSubmixInstance(SubmixPtr.Get()).Pin();
+			if (FoundSubmix.IsValid())
+			{
+				UE_LOG(LogAudioMixer, Display, TEXT("Unregistering submix buffer listener '%s' from submix '%s'"), *InSubmixBufferListener->GetListenerName(), *FoundSubmix->SubmixName);
+				FoundSubmix->UnregisterBufferListener(InSubmixBufferListener);
+			}
+			else
+			{
+				UE_LOG(LogAudioMixer, Display, TEXT("Submix buffer listener '%s' not unregistered. Submix not loaded."), *InSubmixBufferListener->GetListenerName());
+			}
+		};
+
+		FAudioThread::RunCommandOnAudioThread(MoveTemp(UnregisterLambda));
+	}
+
 	void FMixerDevice::FlushExtended(UWorld* WorldToFlush, bool bClearActivatedReverb)
 	{
 		QuantizedEventClockManager.Flush();
diff --git a/Engine/Source/Runtime/AudioMixer/Private/AudioMixerSubmix.cpp b/Engine/Source/Runtime/AudioMixer/Private/AudioMixerSubmix.cpp
--- a/Engine/Source/Runtime/AudioMixer/Private/AudioMixerSubmix.cpp
+++ b/Engine/Source/Runtime/AudioMixer/Private/AudioMixerSubmix.cpp
@@ -1090,9 +1090,9 @@ namespace Audio
 		{
 			// query the SubmixBufferListeners to see if they plan to render audio into this buffer
 			FScopeLock Lock(&BufferListenerCriticalSection);
-			for (const ISubmixBufferListener* BufferListener : BufferListeners)
+			for (const TSharedRef<ISubmixBufferListener, ESPMode::ThreadSafe>& BufferListener : BufferListeners)
 			{
-				if (BufferListener && BufferListener->IsRenderingAudio())
+				if (BufferListener->IsRenderingAudio())
 				{
 					return true;
 				}
@@ -1569,9 +1569,8 @@ namespace Audio
 			double AudioClock = MixerDevice->GetAudioTime();
 			float SampleRate = MixerDevice->GetSampleRate();
 			FScopeLock Lock(&BufferListenerCriticalSection);
-			for (ISubmixBufferListener* BufferListener : BufferListeners)
+			for (TSharedRef<ISubmixBufferListener, ESPMode::ThreadSafe>& BufferListener : BufferListeners)
 			{
-				check(BufferListener);
 				BufferListener->OnNewSubmixBuffer(SoundSubmix, OutAudioBuffer.GetData(), OutAudioBuffer.Num(), NumChannels, SampleRate, AudioClock);
 			}
 
@@ -2093,13 +2092,25 @@ namespace Audio
 	{
 		FScopeLock Lock(&BufferListenerCriticalSection);
 		check(BufferListener);
-		BufferListeners.AddUnique(BufferListener);
+		BufferListeners.Add(BufferListener->AsShared());
+	}
+
+	void FMixerSubmix::RegisterBufferListener(TSharedRef<ISubmixBufferListener, ESPMode::ThreadSafe> BufferListener)
+	{
+		FScopeLock Lock(&BufferListenerCriticalSection);
+		BufferListeners.Add(BufferListener);
 	}
 
 	void FMixerSubmix::UnregisterBufferListener(ISubmixBufferListener* BufferListener)
 	{
 		FScopeLock Lock(&BufferListenerCriticalSection);
 		check(BufferListener);
+		BufferListeners.Remove(BufferListener->AsShared());
+	}
+
+	void FMixerSubmix::UnregisterBufferListener(TSharedRef<ISubmixBufferListener, ESPMode::ThreadSafe> BufferListener)
+	{
+		FScopeLock Lock(&BufferListenerCriticalSection);
 		BufferListeners.Remove(BufferListener);
 	}
 
diff --git a/Engine/Source/Runtime/AudioMixer/Public/AudioMixerDevice.h b/Engine/Source/Runtime/AudioMixer/Public/AudioMixerDevice.h
--- a/Engine/Source/Runtime/AudioMixer/Public/AudioMixerDevice.h
+++ b/Engine/Source/Runtime/AudioMixer/Public/AudioMixerDevice.h
@@ -54,16 +54,29 @@ namespace Audio
 		{}
 	};
 
-	// Master submixes
+	// Required submixes
+	namespace ERequiredSubmixes
+	{
+		enum Type
+		{
+			Main = 0,
+			BaseDefault = 1,
+			Reverb = 2,
+			EQ = 3,
+			Count = 4
+		};
+	}
+
+	// Deprecated, use ERequiredSubmixes above
 	namespace EMasterSubmixType
 	{
 		enum Type
 		{
-			Master,
-			BaseDefault,
-			Reverb,
-			EQ,
-			Count,
+			Master = ERequiredSubmixes::Main,
+			BaseDefault = ERequiredSubmixes::BaseDefault,
+			Reverb = ERequiredSubmixes::Reverb,
+			EQ = ERequiredSubmixes::EQ,
+			Count = ERequiredSubmixes::Count
 		};
 	}
 
@@ -178,9 +191,15 @@ namespace Audio
 		AUDIOMIXER_API virtual void RemoveSpectralAnalysisDelegate(USoundSubmix* InSubmix, const FOnSubmixSpectralAnalysisBP& OnSubmixSpectralAnalysisBP) override;
 
 		// Submix buffer listener callbacks
+		UE_DEPRECATED(5.4, "Use RegisterSubmixBufferListener version that requires a shared reference to a listener and provide explicit reference to a submix: use GetMainSubmixObject to register with the Main Output Submix (rather than nullptr for safety), and instantiate buffer listener via the shared pointer API.")
 		AUDIOMIXER_API virtual void RegisterSubmixBufferListener(ISubmixBufferListener* InSubmixBufferListener, USoundSubmix* InSubmix = nullptr) override;
+
+		UE_DEPRECATED(5.4, "Use UnregisterSubmixBufferListener version that requires a shared reference to a listener and provide explicit reference to a submix: use GetMainSubmixObject to unregister from the Main Output Submix (rather than nullptr for safety), and instantiate buffer listener via the shared pointer API.")
 		AUDIOMIXER_API virtual void UnregisterSubmixBufferListener(ISubmixBufferListener* InSubmixBufferListener, USoundSubmix* InSubmix = nullptr) override;
 
+		AUDIOMIXER_API virtual void RegisterSubmixBufferListener(TSharedRef<ISubmixBufferListener, ESPMode::ThreadSafe> InSubmixBufferListener, USoundSubmix& InSubmix) override;
+		AUDIOMIXER_API virtual void UnregisterSubmixBufferListener(TSharedRef<ISubmixBufferListener, ESPMode::ThreadSafe> InSubmixBufferListener, USoundSubmix& InSubmix) override;
+
 		AUDIOMIXER_API virtual FPatchOutputStrongPtr AddPatchForSubmix(uint32 InObjectId, float InPatchGain) override;
 
 		AUDIOMIXER_API virtual void FlushExtended(UWorld* WorldToFlush, bool bClearActivatedReverb);
@@ -251,19 +270,29 @@ namespace Audio
 		AUDIOMIXER_API FMixerSourceManager* GetSourceManager();
 		AUDIOMIXER_API const FMixerSourceManager* GetSourceManager() const;
 
-		AUDIOMIXER_API FMixerSubmixWeakPtr GetMasterSubmix(); 
+		AUDIOMIXER_API virtual USoundSubmix& GetMainSubmixObject() const override;
+
 		AUDIOMIXER_API FMixerSubmixWeakPtr GetBaseDefaultSubmix();
-		AUDIOMIXER_API FMixerSubmixWeakPtr GetMasterReverbSubmix();
-		AUDIOMIXER_API FMixerSubmixWeakPtr GetMasterEQSubmix();
+		AUDIOMIXER_API FMixerSubmixWeakPtr GetMainSubmix();
+		AUDIOMIXER_API FMixerSubmixWeakPtr GetReverbSubmix();
+		AUDIOMIXER_API FMixerSubmixWeakPtr GetEQSubmix();
 
-		// Add submix effect to master submix
+		// Renamed Main submix: these functions will be deprecated in a future release
 		AUDIOMIXER_API void AddMasterSubmixEffect(FSoundEffectSubmixPtr SoundEffect);
-		
-		// Remove submix effect from master submix
 		AUDIOMIXER_API void RemoveMasterSubmixEffect(uint32 SubmixEffectId);
-		
-		// Clear all submix effects from master submix
 		AUDIOMIXER_API void ClearMasterSubmixEffects();
+		AUDIOMIXER_API FMixerSubmixWeakPtr GetMasterSubmix();
+		AUDIOMIXER_API FMixerSubmixWeakPtr GetMasterReverbSubmix();
+		AUDIOMIXER_API FMixerSubmixWeakPtr GetMasterEQSubmix();
+
+		// Add submix effect to main submix
+		AUDIOMIXER_API void AddMainSubmixEffect(FSoundEffectSubmixPtr SoundEffect);
+
+		// Remove submix effect from main submix
+		AUDIOMIXER_API void RemoveMainSubmixEffect(uint32 SubmixEffectId);
+
+		// Clear all submix effects from main submix
+		AUDIOMIXER_API void ClearMainSubmixEffects();
 
 		// Add submix effect to given submix
 		AUDIOMIXER_API int32 AddSubmixEffect(USoundSubmix* InSoundSubmix, FSoundEffectSubmixPtr SoundEffect);
@@ -363,7 +392,7 @@ namespace Audio
 
 		bool IsMainAudioDevice() const;
 
-		void LoadMasterSoundSubmix(EMasterSubmixType::Type InType, const FString& InDefaultName, bool bInDefaultMuteWhenBackgrounded, FSoftObjectPath& InOutObjectPath);
+		void LoadRequiredSubmix(ERequiredSubmixes::Type InType, const FString& InDefaultName, bool bInDefaultMuteWhenBackgrounded, FSoftObjectPath& InOutObjectPath);
 		void LoadPluginSoundSubmixes();
 		void LoadSoundSubmix(USoundSubmixBase& SoundSubmix);
 
@@ -373,16 +402,16 @@ namespace Audio
 
 		ICompressedAudioInfo* CreateAudioInfo(FName InFormat) const;
 
-		bool IsMasterSubmixType(const USoundSubmixBase* InSubmix) const;
-		FMixerSubmixPtr GetMasterSubmixInstance(uint32 InSubmixId);
-		FMixerSubmixPtr GetMasterSubmixInstance(const USoundSubmixBase* InSubmix);
+		bool IsRequiredSubmixType(const USoundSubmixBase* InSubmix) const;
+		FMixerSubmixPtr GetRequiredSubmixInstance(uint32 InSubmixId);
+		FMixerSubmixPtr GetRequiredSubmixInstance(const USoundSubmixBase* InSubmix);
 		
 		// Pumps the audio render thread command queue
 		void PumpCommandQueue();
 		void PumpGameThreadCommandQueue();
 		
-		TArray<USoundSubmix*> MasterSubmixes;
-		TArray<FMixerSubmixPtr> MasterSubmixInstances;
+		TArray<USoundSubmix*> RequiredSubmixes;
+		TArray<FMixerSubmixPtr> RequiredSubmixInstances;
 
 		TArray<TStrongObjectPtr<UAudioBus>> DefaultAudioBuses;
 		/** Ptr to the platform interface, which handles streaming audio to the hardware device. */
diff --git a/Engine/Source/Runtime/AudioMixer/Public/AudioMixerSubmix.h b/Engine/Source/Runtime/AudioMixer/Public/AudioMixerSubmix.h
--- a/Engine/Source/Runtime/AudioMixer/Public/AudioMixerSubmix.h
+++ b/Engine/Source/Runtime/AudioMixer/Public/AudioMixerSubmix.h
@@ -251,11 +251,17 @@ namespace Audio
 		AUDIOMIXER_API void ResumeRecordingOutput();
 
 		// Register buffer listener with this submix
+		// Unregister buffer listener with this submix
+
+		UE_DEPRECATED(5.4, "This function is deprecated. Use RegisterBufferListener version that is provided a shared reference to a listener.")
 		AUDIOMIXER_API void RegisterBufferListener(ISubmixBufferListener* BufferListener);
 		
-		// Unregister buffer listener with this submix
+		UE_DEPRECATED(5.4, "This function is deprecated. Use UnregisterBufferListener version that is provided a shared reference to a listener.")
 		AUDIOMIXER_API void UnregisterBufferListener(ISubmixBufferListener* BufferListener);
 
+		AUDIOMIXER_API void RegisterBufferListener(TSharedRef<ISubmixBufferListener, ESPMode::ThreadSafe> BufferListener);
+		AUDIOMIXER_API void UnregisterBufferListener(TSharedRef<ISubmixBufferListener, ESPMode::ThreadSafe> BufferListener);
+
 		// Starts envelope following with the given attack time and release time
 		AUDIOMIXER_API void StartEnvelopeFollowing(int32 AttackTime, int32 ReleaseTime);
 
@@ -355,6 +361,11 @@ namespace Audio
 		AUDIOMIXER_API TUniquePtr<ISoundfieldTranscodeStream> GetTranscoderForChildSubmix(const TSharedPtr<Audio::FMixerSubmix, ESPMode::ThreadSafe>& InChildSubmix);
 
 	protected:
+		struct AUDIOMIXER_API FSubmixBufferListenerInfo
+		{
+			ISubmixBufferListener* Listener = nullptr;
+			FString Descriptor;
+		};
 
 		// Pump command queue
 		AUDIOMIXER_API void PumpCommandQueue();
@@ -565,8 +576,8 @@ namespace Audio
 		// Submix command queue to shuffle commands from audio thread to audio render thread.
 		TQueue<TFunction<void()>> CommandQueue;
 
-		// List of submix buffer listeners
-		TArray<ISubmixBufferListener*> BufferListeners;
+		// List of submix buffer listeners.
+		TArray<TSharedRef<ISubmixBufferListener, ESPMode::ThreadSafe>> BufferListeners;
 
 		// Critical section used for modifying and interacting with buffer listeners
 		mutable FCriticalSection BufferListenerCriticalSection;
diff --git a/Engine/Source/Runtime/Engine/Private/AudioDevice.cpp b/Engine/Source/Runtime/Engine/Private/AudioDevice.cpp
--- a/Engine/Source/Runtime/Engine/Private/AudioDevice.cpp
+++ b/Engine/Source/Runtime/Engine/Private/AudioDevice.cpp
@@ -6836,16 +6836,34 @@ void FAudioDevice::UnregisterSoundClass(USoundClass* InSoundClass)
 	}
 }
 
+USoundSubmix& FAudioDevice::GetMainSubmixObject() const
+{
+	UE_LOG(LogAudio, Error, TEXT("Main submix accessor only supported with the audio mixer. Run with audio mixer enabled. Class Default Object returned."));
+	USoundSubmix* CDO = USoundSubmix::StaticClass()->GetDefaultObject<USoundSubmix>();
+	check(CDO);
+	return *CDO;
+}
+
 void FAudioDevice::RegisterSubmixBufferListener(ISubmixBufferListener* InSubmixBufferListener, USoundSubmix* SoundSubmix)
 {
 	UE_LOG(LogAudio, Error, TEXT("Submix buffer listener only works with the audio mixer. Please run with audio mixer enabled."));
 }
 
+void FAudioDevice::RegisterSubmixBufferListener(TSharedRef<ISubmixBufferListener, ESPMode::ThreadSafe>, USoundSubmix& SoundSubmix)
+{
+	UE_LOG(LogAudio, Error, TEXT("Submix buffer listener only works with the audio mixer. Please run with audio mixer enabled."));
+}
+
 void FAudioDevice::UnregisterSubmixBufferListener(ISubmixBufferListener* InSubmixBufferListener, USoundSubmix* SoundSubmix)
 {
 	UE_LOG(LogAudio, Error, TEXT("Submix buffer listener only works with the audio mixer. Please run with audio mixer enabled."));
 }
 
+void FAudioDevice::UnregisterSubmixBufferListener(TSharedRef<ISubmixBufferListener, ESPMode::ThreadSafe>, USoundSubmix& SoundSubmix)
+{
+	UE_LOG(LogAudio, Error, TEXT("Submix buffer listener only works with the audio mixer. Please run with audio mixer enabled."));
+}
+
 Audio::FPatchOutputStrongPtr FAudioDevice::AddPatchForSubmix(uint32 InObjectId, float InPatchGain)
 {
 	UE_LOG(LogAudio, Error, TEXT("Submix patching only works with the audio mixer. Please run with audio mixer enabled."));
diff --git a/Engine/Source/Runtime/Engine/Public/AudioDevice.h b/Engine/Source/Runtime/Engine/Public/AudioDevice.h
--- a/Engine/Source/Runtime/Engine/Public/AudioDevice.h
+++ b/Engine/Source/Runtime/Engine/Public/AudioDevice.h
@@ -958,17 +958,24 @@ public:
 	/** Unregisters the sound submix */
 	virtual void UnregisterSoundSubmix(const USoundSubmixBase* SoundSubmix) {}
 
+	ENGINE_API virtual USoundSubmix& GetMainSubmixObject() const;
+
+	UE_DEPRECATED(5.4, "Use RegisterSubmixBufferListener version that requires a shared reference to a listener and provide explicit reference to a submix: use GetMainSubmixObject to register with the Main Output Submix (rather than nullptr for safety), and instantiate buffer listener via the shared pointer API.")
+	ENGINE_API virtual void RegisterSubmixBufferListener(ISubmixBufferListener* InSubmixBufferListener, USoundSubmix* SoundSubmix = nullptr);
+
+	UE_DEPRECATED(5.4, "Use UnregisterSubmixBufferListener version that requires a shared reference to a listener and provide explicit reference to a submix: use GetMainSubmixObject to unregister from the Main Output Submix (rather than nullptr for safety), and instantiate buffer listener via the shared pointer API.")
+	ENGINE_API virtual void UnregisterSubmixBufferListener(ISubmixBufferListener* InSubmixBufferListener, USoundSubmix* SoundSubmix = nullptr);
+
 	/**
-	 * Registers the submix buffer listener with the given submix.
-	 * A nullptr for SoundSubmix will register the listener with the master submix.
+	 * Registers the provided submix buffer listener with the given submix.
 	*/
-	ENGINE_API virtual void RegisterSubmixBufferListener(ISubmixBufferListener* InSubmixBufferListener, USoundSubmix* SoundSubmix = nullptr);
+	ENGINE_API virtual void RegisterSubmixBufferListener(TSharedRef<ISubmixBufferListener, ESPMode::ThreadSafe> InSubmixBufferListener, USoundSubmix& SoundSubmix);
 
 	/**
-	 * Unregisters the submix buffer listener with the given submix.
-	 * A nullptr for SoundSubmix will unregister the listener with the master submix.
+	 * Unregisters the provided submix buffer listener with the given submix.
 	*/
-	ENGINE_API virtual void UnregisterSubmixBufferListener(ISubmixBufferListener* InSubmixBufferListener, USoundSubmix* SoundSubmix = nullptr);
+	ENGINE_API virtual void UnregisterSubmixBufferListener(TSharedRef<ISubmixBufferListener, ESPMode::ThreadSafe> InSubmixBufferListener, USoundSubmix& SoundSubmix);
+
 
 	ENGINE_API virtual Audio::FPatchOutputStrongPtr AddPatchForSubmix(uint32 InObjectId, float InPatchGain);
 
diff --git a/Engine/Source/Runtime/Engine/Public/ISubmixBufferListener.h b/Engine/Source/Runtime/Engine/Public/ISubmixBufferListener.h
--- a/Engine/Source/Runtime/Engine/Public/ISubmixBufferListener.h
+++ b/Engine/Source/Runtime/Engine/Public/ISubmixBufferListener.h
@@ -3,14 +3,13 @@
 
 #include "HAL/Platform.h"
 
-#if UE_ENABLE_INCLUDE_ORDER_DEPRECATED_IN_5_3
-#include "Sound/SoundSubmix.h"
-#endif
+#include "Templates/SharedPointer.h"
 
+// Forward Declarations
 class USoundSubmix;
 
 /** Abstract interface for receiving audio data from a given submix. */
-class ISubmixBufferListener
+class ISubmixBufferListener : public TSharedFromThis<ISubmixBufferListener, ESPMode::ThreadSafe>
 {
 public:
 	virtual ~ISubmixBufferListener() = default;
@@ -38,4 +37,14 @@ public:
 	{
 		return false;
 	}
+
+	virtual const FString& GetListenerName() const
+	{
+		static FString UnsetName = TEXT("Unset");
+		return UnsetName;
+	}
+
+protected:
+	template <typename ObjectType, ESPMode Mode>
+	friend class SharedPointerInternals::TIntrusiveReferenceController;
 };
diff --git a/Engine/Source/Runtime/GameplayMediaEncoder/Private/GameplayMediaEncoder.cpp b/Engine/Source/Runtime/GameplayMediaEncoder/Private/GameplayMediaEncoder.cpp
--- a/Engine/Source/Runtime/GameplayMediaEncoder/Private/GameplayMediaEncoder.cpp
+++ b/Engine/Source/Runtime/GameplayMediaEncoder/Private/GameplayMediaEncoder.cpp
@@ -64,15 +64,16 @@ FAutoConsoleCommand GameplayMediaEncoderShutdown(TEXT("GameplayMediaEncoder.Shut
 //
 //////////////////////////////////////////////////////////////////////////
 
-FGameplayMediaEncoder* FGameplayMediaEncoder::Singleton = nullptr;
+TSharedPtr<FGameplayMediaEncoder, ESPMode::ThreadSafe> FGameplayMediaEncoder::Singleton = { };
 
 FGameplayMediaEncoder* FGameplayMediaEncoder::Get()
 {
-	if(!Singleton)
+	// Constructed and captured as a thread safe shared pointer as this is required by the audio submix listener interface
+	if(!Singleton.IsValid())
 	{
-		Singleton = new FGameplayMediaEncoder();
+		Singleton = MakeShared<FGameplayMediaEncoder>();
 	}
-	return Singleton;
+	return Singleton.Get();
 }
 
 FGameplayMediaEncoder::FGameplayMediaEncoder() {}
@@ -307,7 +308,7 @@ bool FGameplayMediaEncoder::Start()
 	if(AudioDevice)
 	{
 		bAudioFormatChecked = false;
-		AudioDevice->RegisterSubmixBufferListener(this);
+		AudioDevice->RegisterSubmixBufferListener(AsShared(), AudioDevice->GetMainSubmixObject());
 	}
 
 	FSlateApplication::Get().GetRenderer()->OnBackBufferReadyToPresent().AddRaw(this, &FGameplayMediaEncoder::OnFrameBufferReady);
@@ -327,10 +328,9 @@ void FGameplayMediaEncoder::Stop()
 
 	if(UGameEngine* GameEngine = Cast<UGameEngine>(GEngine))
 	{
-		FAudioDevice* AudioDevice = GameEngine->GetMainAudioDeviceRaw();
-		if(AudioDevice)
+		if(FAudioDevice* AudioDevice = GameEngine->GetMainAudioDeviceRaw())
 		{
-			AudioDevice->UnregisterSubmixBufferListener(this);
+			AudioDevice->UnregisterSubmixBufferListener(AsShared(), AudioDevice->GetMainSubmixObject());
 		}
 
 		if(FSlateApplication::IsInitialized())
@@ -388,6 +388,11 @@ void FGameplayMediaEncoder::Shutdown()
 
 FTimespan FGameplayMediaEncoder::GetMediaTimestamp() const { return FTimespan::FromSeconds(FPlatformTime::Seconds()) - StartTime; }
 
+const FString& FGameplayMediaEncoder::GetListenerName() const
+{
+	static const FString& ListenerName = TEXT("GameplayMediaEncoderListener");
+	return ListenerName;
+}
 void FGameplayMediaEncoder::OnNewSubmixBuffer(const USoundSubmix* OwningSubmix, float* AudioData, int32 NumSamples, int32 NumChannels, const int32 SampleRate, double /*AudioClock*/)
 {
 	CSV_SCOPED_TIMING_STAT(GameplayMediaEncoder, OnNewSubmixBuffer);
diff --git a/Engine/Source/Runtime/GameplayMediaEncoder/Private/GameplayMediaEncoderModule.cpp b/Engine/Source/Runtime/GameplayMediaEncoder/Private/GameplayMediaEncoderModule.cpp
--- a/Engine/Source/Runtime/GameplayMediaEncoder/Private/GameplayMediaEncoderModule.cpp
+++ b/Engine/Source/Runtime/GameplayMediaEncoder/Private/GameplayMediaEncoderModule.cpp
@@ -27,11 +27,7 @@ public:
 		// If the FGameplayMediaEncoder instance was created, then explicitly destroy it here
 		// instead of waiting for the automatic cleanup, since at that point some objects
 		// it depends to for a clean shutdown are not available any longer.
-		if (FGameplayMediaEncoder::Singleton)
-		{
-			delete FGameplayMediaEncoder::Singleton;
-			FGameplayMediaEncoder::Singleton = nullptr;
-		}
+		FGameplayMediaEncoder::Singleton.Reset();
 	}
 private:
 };
diff --git a/Engine/Source/Runtime/GameplayMediaEncoder/Public/GameplayMediaEncoder.h b/Engine/Source/Runtime/GameplayMediaEncoder/Public/GameplayMediaEncoder.h
--- a/Engine/Source/Runtime/GameplayMediaEncoder/Public/GameplayMediaEncoder.h
+++ b/Engine/Source/Runtime/GameplayMediaEncoder/Public/GameplayMediaEncoder.h
@@ -23,6 +23,8 @@ THIRD_PARTY_INCLUDES_END
 #include "AudioEncoder.h"
 #include "VideoEncoder.h"
 #include "MediaPacket.h"
+#include "Templates/SharedPointer.h"
+#include "Templates/SharedPointerInternals.h"
 #include "VideoEncoderInput.h"
 
 class SWindow;
@@ -33,7 +35,7 @@ public:
 	virtual void OnMediaSample(const AVEncoder::FMediaPacket& Sample) = 0;
 };
 
-class FGameplayMediaEncoder final : private ISubmixBufferListener, public AVEncoder::IAudioEncoderListener
+class FGameplayMediaEncoder final : public ISubmixBufferListener, public AVEncoder::IAudioEncoderListener
 {
 public:
 
@@ -89,6 +91,8 @@ public:
 	AVEncoder::FVideoConfig GetVideoConfig() const { return VideoConfig; }
 
 private:
+	template <typename ObjectType, ESPMode Mode>
+	friend class SharedPointerInternals::TIntrusiveReferenceController;
 
 	// Private to control how our single instance is created
 	GAMEPLAYMEDIAENCODER_API FGameplayMediaEncoder();
@@ -98,7 +102,9 @@ private:
 
 	// Back buffer capture
 	GAMEPLAYMEDIAENCODER_API void OnFrameBufferReady(SWindow& SlateWindow, const FTexture2DRHIRef& FrameBuffer);
+
 	// ISubmixBufferListener interface
+	GAMEPLAYMEDIAENCODER_API const FString& GetListenerName() const override;
 	GAMEPLAYMEDIAENCODER_API void OnNewSubmixBuffer(const USoundSubmix* OwningSubmix, float* AudioData, int32 NumSamples, int32 NumChannels, const int32 SampleRate, double AudioClock) override;
 
 	GAMEPLAYMEDIAENCODER_API void ProcessAudioFrame(const float* AudioData, int32 NumSamples, int32 NumChannels, int32 SampleRate);
@@ -141,7 +147,7 @@ private:
 	bool bDoFrameSkipping = false;
 
 	friend class FGameplayMediaEncoderModule;
-	static GAMEPLAYMEDIAENCODER_API FGameplayMediaEncoder* Singleton;
+	static GAMEPLAYMEDIAENCODER_API TSharedPtr<FGameplayMediaEncoder, ESPMode::ThreadSafe> Singleton;
 
 	// live streaming: quality adaptation to available uplink b/w
 	TAtomic<uint32> NewVideoBitrate{ 0 };
