- Add listener name and logging for when register command is sent to AT vs executed to get more logging data that could better determine potentially rogue, prematurally-deleted client
- Add listener name to editor pixel streaming
#rnx

diff --git a/Engine/Plugins/Media/PixelStreaming/Source/PixelStreamingEditor/Private/EditorSubmixListener.cpp b/Engine/Plugins/Media/PixelStreaming/Source/PixelStreamingEditor/Private/EditorSubmixListener.cpp
--- a/Engine/Plugins/Media/PixelStreaming/Source/PixelStreamingEditor/Private/EditorSubmixListener.cpp
+++ b/Engine/Plugins/Media/PixelStreaming/Source/PixelStreamingEditor/Private/EditorSubmixListener.cpp
@@ -28,6 +28,12 @@ namespace UE::EditorPixelStreaming
 	{
 	}
 
+	const FString& FEditorSubmixListener::GetListenerName() const
+	{
+		static const FString ListenerName(TEXT("PixelStreamingEditorListener"));
+		return ListenerName;
+	}
+
 	void FEditorSubmixListener::Shutdown()
 	{
 		if (FAudioDevice* AudioDevice = FAudioDeviceManager::Get()->GetAudioDeviceRaw(AudioDeviceId))
diff --git a/Engine/Plugins/Media/PixelStreaming/Source/PixelStreamingEditor/Private/EditorSubmixListener.h b/Engine/Plugins/Media/PixelStreaming/Source/PixelStreamingEditor/Private/EditorSubmixListener.h
--- a/Engine/Plugins/Media/PixelStreaming/Source/PixelStreamingEditor/Private/EditorSubmixListener.h
+++ b/Engine/Plugins/Media/PixelStreaming/Source/PixelStreamingEditor/Private/EditorSubmixListener.h
@@ -21,6 +21,8 @@ namespace UE::EditorPixelStreaming
 
 		void Shutdown();
 
+		virtual const FString& GetListenerName() const override;
+
 	protected:
 		template <typename ObjectType, ESPMode Mode>
 		friend class SharedPointerInternals::TIntrusiveReferenceController;
diff --git a/Engine/Source/Runtime/AudioMixer/Private/AudioMixerDevice.cpp b/Engine/Source/Runtime/AudioMixer/Private/AudioMixerDevice.cpp
--- a/Engine/Source/Runtime/AudioMixer/Private/AudioMixerDevice.cpp
+++ b/Engine/Source/Runtime/AudioMixer/Private/AudioMixerDevice.cpp
@@ -2564,7 +2564,14 @@ namespace Audio
 
 		const TWeakObjectPtr<USoundSubmix> SubmixPtr(&InSubmix);
 
-		auto RegisterLambda = [this, InSubmixBufferListener, SubmixPtr]()
+		// Pass the name vs. reconciling it inline with the command lambda, as occasionally
+		// deprecated submix buffer listeners are not constructed as a shared pointer and
+		// therefore getting destroyed before the lambda is executed on the AudioThread.
+		// This means no name is provided and thus not apparent who the caller is. If
+		// the Buffer Listener is invalid here, at least the callstack will show the
+		// requesting client directly.
+		const FString ListenerName = InSubmixBufferListener->GetListenerName();
+		auto RegisterLambda = [this, InSubmixBufferListener, ListenerName, SubmixPtr]()
 		{
 			CSV_SCOPED_TIMING_STAT(Audio, RegisterSubmixBufferListener);
 
@@ -2577,7 +2584,6 @@ namespace Audio
 				FoundSubmix = GetSubmixInstance(SubmixPtr.Get()).Pin();
 			}
 
-			const FString& ListenerName = InSubmixBufferListener->GetListenerName();
 			if (FoundSubmix.IsValid())
 			{
 				FoundSubmix->RegisterBufferListener(InSubmixBufferListener);
@@ -2589,6 +2595,7 @@ namespace Audio
 			}
 		};
 
+		UE_LOG(LogAudioMixer, Display, TEXT("Sending SubmixBufferListener '%s' register command..."), *ListenerName);
 		FAudioThread::RunCommandOnAudioThread(MoveTemp(RegisterLambda));
 	}
 
