[AudioMixer] Route Submix audio output to an AudioBus

diff --git a/Engine/Source/Runtime/AudioMixer/Private/AudioMixerBlueprintLibrary.cpp b/Engine/Source/Runtime/AudioMixer/Private/AudioMixerBlueprintLibrary.cpp
--- a/Engine/Source/Runtime/AudioMixer/Private/AudioMixerBlueprintLibrary.cpp
+++ b/Engine/Source/Runtime/AudioMixer/Private/AudioMixerBlueprintLibrary.cpp
@@ -9,6 +9,7 @@
 #include "AudioDevice.h"
 #include "AudioDeviceManager.h"
 #include "AudioMixerDevice.h"
+#include "AudioMixerSubmix.h"
 #include "ContentStreaming.h"
 #include "CoreMinimal.h"
 #include "DSP/ConstantQ.h"
@@ -709,6 +710,99 @@ bool UAudioMixerBlueprintLibrary::IsAudioBusActive(const UObject* WorldContextOb
 	}
 }
 
+void UAudioMixerBlueprintLibrary::RegisterAudioBusToSubmix(const UObject* WorldContextObject, USoundSubmix* SoundSubmix, UAudioBus* AudioBus)
+{
+	if (!SoundSubmix)
+	{
+		UE_LOG(LogAudioMixer, Error, TEXT("RegisterAudioBusToSubmix called with an invalid Submix."));
+		return;
+	}
+
+	if (!AudioBus)
+	{
+		UE_LOG(LogAudioMixer, Error, TEXT("RegisterAudioBusToSubmix called with an invalid Audio Bus."));
+		return;
+	}
+
+	if (!IsInAudioThread())
+	{
+		//Send this over to the audio thread, with the same settings
+		FAudioThread::RunCommandOnAudioThread([WorldContextObject, SoundSubmix, AudioBus]()
+		{
+			RegisterAudioBusToSubmix(WorldContextObject, SoundSubmix, AudioBus);
+		});
+
+		return;
+	}
+
+	if (Audio::FMixerDevice* MixerDevice = FAudioDeviceManager::GetAudioMixerDeviceFromWorldContext(WorldContextObject))
+	{
+		Audio::FMixerSubmixPtr MixerSubmixPtr = MixerDevice->GetSubmixInstance(SoundSubmix).Pin();
+
+		if (MixerSubmixPtr.IsValid())
+		{
+			UAudioBusSubsystem* AudioBusSubsystem = MixerDevice->GetSubsystem<UAudioBusSubsystem>();
+			check(AudioBusSubsystem);
+
+			const Audio::FAudioBusKey AudioBusKey = Audio::FAudioBusKey(AudioBus->GetUniqueID());
+			MixerSubmixPtr->RegisterAudioBus(AudioBusKey, AudioBusSubsystem->AddPatchInputForAudioBus(AudioBusKey, MixerDevice->GetNumOutputFrames(), AudioBus->GetNumChannels()));
+		}
+		else
+		{
+			UE_LOG(LogAudioMixer, Error, TEXT("Submix not found in audio mixer."));
+		}
+	}
+	else
+	{
+		UE_LOG(LogAudioMixer, Error, TEXT("Audio mixer device not found."));
+	}
+}
+
+void UAudioMixerBlueprintLibrary::UnregisterAudioBusFromSubmix(const UObject* WorldContextObject, USoundSubmix* SoundSubmix, UAudioBus* AudioBus)
+{
+	if (!SoundSubmix)
+	{
+		UE_LOG(LogAudioMixer, Error, TEXT("UnregisterAudioBusToSubmix called with an invalid Submix."));
+		return;
+	}
+
+	if (!AudioBus)
+	{
+		UE_LOG(LogAudioMixer, Error, TEXT("UnregisterAudioBusToSubmix called with an invalid Audio Bus."));
+		return;
+	}
+
+	if (!IsInAudioThread())
+	{
+		//Send this over to the audio thread, with the same settings
+		FAudioThread::RunCommandOnAudioThread([WorldContextObject, SoundSubmix, AudioBus]()
+		{
+			UnregisterAudioBusFromSubmix(WorldContextObject, SoundSubmix, AudioBus);
+		});
+
+		return;
+	}
+
+	if (Audio::FMixerDevice* MixerDevice = FAudioDeviceManager::GetAudioMixerDeviceFromWorldContext(WorldContextObject))
+	{
+		Audio::FMixerSubmixPtr MixerSubmixPtr = MixerDevice->GetSubmixInstance(SoundSubmix).Pin();
+
+		if (MixerSubmixPtr.IsValid())
+		{
+			const Audio::FAudioBusKey AudioBusKey(AudioBus->GetUniqueID());
+			MixerSubmixPtr->UnregisterAudioBus(AudioBusKey);
+		}
+		else
+		{
+			UE_LOG(LogAudioMixer, Error, TEXT("Submix not found in audio mixer."));
+		}
+	}
+	else
+	{
+		UE_LOG(LogAudioMixer, Error, TEXT("Audio mixer device not found."));
+	}
+}
+
 void UAudioMixerBlueprintLibrary::GetAvailableAudioOutputDevices(const UObject* WorldContextObject, const FOnAudioOutputDevicesObtained& OnObtainDevicesEvent)
 {
 	if (!OnObtainDevicesEvent.IsBound())
diff --git a/Engine/Source/Runtime/AudioMixer/Private/AudioMixerSubmix.cpp b/Engine/Source/Runtime/AudioMixer/Private/AudioMixerSubmix.cpp
--- a/Engine/Source/Runtime/AudioMixer/Private/AudioMixerSubmix.cpp
+++ b/Engine/Source/Runtime/AudioMixer/Private/AudioMixerSubmix.cpp
@@ -45,8 +45,8 @@ FAutoConsoleVariableRef CVarLogSubmixEnablement(
 	ECVF_Default);
 
 // Define profiling categories for submixes. 
-DEFINE_STAT(STAT_AudioMixerSubmixes);
 DEFINE_STAT(STAT_AudioMixerEndpointSubmixes);
+DEFINE_STAT(STAT_AudioMixerSubmixes);
 DEFINE_STAT(STAT_AudioMixerSubmixChildren);
 DEFINE_STAT(STAT_AudioMixerSubmixSource);
 DEFINE_STAT(STAT_AudioMixerSubmixEffectProcessing);
@@ -498,6 +498,32 @@ namespace Audio
 		});
 	}
 
+	void FMixerSubmix::RegisterAudioBus(const Audio::FAudioBusKey& InAudioBusKey, Audio::FPatchInput&& InPatchInput)
+	{
+		check(IsInAudioThread());
+
+		SubmixCommand([this, InAudioBusKey, InPatchInput = MoveTemp(InPatchInput)]()
+		{
+			if (!AudioBuses.Contains(InAudioBusKey))
+			{
+				AudioBuses.Emplace(InAudioBusKey, InPatchInput);
+			}
+		});
+	}
+
+	void FMixerSubmix::UnregisterAudioBus(const Audio::FAudioBusKey& InAudioBusKey)
+	{
+		check(IsInAudioThread());
+
+		SubmixCommand([this, InAudioBusKey]()
+		{
+			if (AudioBuses.Contains(InAudioBusKey))
+			{
+				AudioBuses.Remove(InAudioBusKey);
+			}
+		});
+	}
+
 	int32 FMixerSubmix::GetSubmixChannels() const
 	{
 		return NumChannels;
@@ -1242,6 +1268,7 @@ namespace Audio
 
 			// Even though we're silent, broadcast the buffer to any listeners (will be a silent buffer)
 			SendAudioToSubmixBufferListeners(InputBuffer);
+			SendAudioToRegisteredAudioBuses(InputBuffer);
 			return;
 		}
 
@@ -1524,6 +1551,7 @@ namespace Audio
 		}
 
 		SendAudioToSubmixBufferListeners(InputBuffer);
+		SendAudioToRegisteredAudioBuses(InputBuffer);
 
 		// Mix the audio buffer of this submix with the audio buffer of the output buffer (i.e. with other submixes)
 		Audio::ArrayMixIn(InputBuffer, OutAudioBuffer);
@@ -1578,6 +1606,16 @@ namespace Audio
 		}
 	}
 
+	void FMixerSubmix::SendAudioToRegisteredAudioBuses(FAlignedFloatBuffer& OutAudioBuffer)
+	{
+		TRACE_CPUPROFILER_EVENT_SCOPE(FMixerSubmix::SendAudioToRegisteredAudioBuses);
+
+		for (auto& [AudioBusKey, PatchInput] : AudioBuses)
+		{
+			PatchInput.PushAudio(OutAudioBuffer.GetData(), OutAudioBuffer.Num());
+		}
+	}
+
 	bool FMixerSubmix::GenerateEffectChainAudio(FSoundEffectSubmixInputData& InputData, const FAlignedFloatBuffer& InAudioBuffer, TArray<FSoundEffectSubmixPtr>& InEffectChain, FAlignedFloatBuffer& OutBuffer)
 	{
 		checkf(InAudioBuffer.Num() == OutBuffer.Num(), TEXT("Processing effect chain audio must not alter the number of samples in a buffer. Buffer size mismatch InAudioBuffer[%d] != OutBuffer[%d]"), InAudioBuffer.Num(), OutBuffer.Num());
diff --git a/Engine/Source/Runtime/AudioMixer/Public/AudioMixerBlueprintLibrary.h b/Engine/Source/Runtime/AudioMixer/Public/AudioMixerBlueprintLibrary.h
--- a/Engine/Source/Runtime/AudioMixer/Public/AudioMixerBlueprintLibrary.h
+++ b/Engine/Source/Runtime/AudioMixer/Public/AudioMixerBlueprintLibrary.h
@@ -379,6 +379,14 @@ public:
 	UFUNCTION(BlueprintCallable, Category = "Audio|Bus", meta = (WorldContext = "WorldContextObject"))
 	static AUDIOMIXER_API bool IsAudioBusActive(const UObject* WorldContextObject, UAudioBus* AudioBus);
 
+	/** Registers an audio bus to a submix so the submix output can be routed to the audiobus. */
+	UFUNCTION(BlueprintCallable, Category = "Audio|Bus", meta = (WorldContext = "WorldContextObject"))
+	static AUDIOMIXER_API void RegisterAudioBusToSubmix(const UObject* WorldContextObject, USoundSubmix* SoundSubmix, UAudioBus* AudioBus);
+
+	/** Unregisters an audio bus that could have been registered to a submix. */
+	UFUNCTION(BlueprintCallable, Category = "Audio|Bus", meta = (WorldContext = "WorldContextObject"))
+	static AUDIOMIXER_API void UnregisterAudioBusFromSubmix(const UObject* WorldContextObject, USoundSubmix* SoundSubmix, UAudioBus* AudioBus);
+
 	/**
 	* Gets information about all audio output devices available in the system
 	* @param OnObtainDevicesEvent - the event to fire when the audio endpoint devices have been retrieved
diff --git a/Engine/Source/Runtime/AudioMixer/Public/AudioMixerSubmix.h b/Engine/Source/Runtime/AudioMixer/Public/AudioMixerSubmix.h
--- a/Engine/Source/Runtime/AudioMixer/Public/AudioMixerSubmix.h
+++ b/Engine/Source/Runtime/AudioMixer/Public/AudioMixerSubmix.h
@@ -2,6 +2,7 @@
 
 #pragma once
 
+#include "AudioBusSubsystem.h"
 #include "AudioMixer.h"
 #include "AudioDefines.h"
 #include "CoreMinimal.h"
@@ -132,6 +133,12 @@ namespace Audio
 		// Removes the given submix from this submix's children
 		AUDIOMIXER_API void RemoveChildSubmix(TWeakPtr<FMixerSubmix, ESPMode::ThreadSafe> SubmixWeakPtr);
 
+		// Registers the given audiobus to this submix
+		AUDIOMIXER_API void RegisterAudioBus(const Audio::FAudioBusKey& InAudioBusKey, Audio::FPatchInput&& InPatchInput);
+
+		// Unregisters a registered audiobus from this submix (if any)
+		AUDIOMIXER_API void UnregisterAudioBus(const Audio::FAudioBusKey& InAudioBusKey);
+
 		// Sets the output level of the submix in linear gain
 		AUDIOMIXER_API void SetOutputVolume(float InOutputLevel);
 
@@ -671,5 +678,11 @@ namespace Audio
 		TUniquePtr<IAudioLink> AudioLinkInstance;
 
 		friend class FMixerDevice;
+
+	private:
+		AUDIOMIXER_API void SendAudioToRegisteredAudioBuses(FAlignedFloatBuffer& OutAudioBuffer);
+
+		// Registered audio buses
+		TMap<Audio::FAudioBusKey, Audio::FPatchInput> AudioBuses;
 	};
 }
