Don't allocate reverb memory if reverbs are disabled by CVar

diff --git a/Engine/Source/Runtime/AudioMixer/Classes/SubmixEffects/AudioMixerSubmixEffectReverb.h b/Engine/Source/Runtime/AudioMixer/Classes/SubmixEffects/AudioMixerSubmixEffectReverb.h
--- a/Engine/Source/Runtime/AudioMixer/Classes/SubmixEffects/AudioMixerSubmixEffectReverb.h
+++ b/Engine/Source/Runtime/AudioMixer/Classes/SubmixEffects/AudioMixerSubmixEffectReverb.h
@@ -163,6 +163,8 @@ private:
 
 	// Curve which maps old reverb times to new decay value
 	FRichCurve DecayCurve;
+
+	float SampleRate = 1.f;
 };
 
 UCLASS(MinimalAPI)
diff --git a/Engine/Source/Runtime/AudioMixer/Private/Effects/AudioMixerSubmixEffectReverb.cpp b/Engine/Source/Runtime/AudioMixer/Private/Effects/AudioMixerSubmixEffectReverb.cpp
--- a/Engine/Source/Runtime/AudioMixer/Private/Effects/AudioMixerSubmixEffectReverb.cpp
+++ b/Engine/Source/Runtime/AudioMixer/Private/Effects/AudioMixerSubmixEffectReverb.cpp
@@ -68,6 +68,8 @@ void FSubmixEffectReverb::Init(const FSoundEffectSubmixInitData& InitData)
 	 */
 	Audio::FPlateReverbFastSettings NewSettings;
 
+	SampleRate = InitData.SampleRate;
+
 	NewSettings.EarlyReflections.Decay = 0.9f;
 	NewSettings.EarlyReflections.Absorption = 0.7f;
 	NewSettings.EarlyReflections.Gain = 1.0f;
@@ -82,7 +84,6 @@ void FSubmixEffectReverb::Init(const FSoundEffectSubmixInitData& InitData)
 	NewSettings.LateReflections.Decay = 0.15f;
 	NewSettings.LateReflections.Density = 0.85f;
 
-
 	ReverbParams.SetParams(NewSettings);
 
 	DecayCurve.AddKey(0.0f, 0.99f);
@@ -93,7 +94,10 @@ void FSubmixEffectReverb::Init(const FSoundEffectSubmixInitData& InitData)
 	DecayCurve.AddKey(19.0f, 0.002f);
 	DecayCurve.AddKey(20.0f, 0.0001f);
 
-	PlateReverb = MakeUnique<Audio::FPlateReverbFast>(InitData.SampleRate, 512, NewSettings);
+	if (DisableSubmixReverbCVarFast == 0)
+	{
+		PlateReverb = MakeUnique<Audio::FPlateReverbFast>(SampleRate, 512, NewSettings);
+	}
 }
 
 void FSubmixEffectReverb::OnPresetChanged()
@@ -141,6 +145,13 @@ void FSubmixEffectReverb::OnProcessAudio(const FSoundEffectSubmixInputData& InDa
 		return;
 	}
 
+	if (!PlateReverb.IsValid())
+	{
+		Audio::FPlateReverbFastSettings NewSettings;
+		ReverbParams.CopyParams(NewSettings);
+		PlateReverb = MakeUnique<Audio::FPlateReverbFast>(SampleRate, 512, NewSettings);
+	}
+
 	CSV_SCOPED_TIMING_STAT(Audio, SubmixReverb);
 	SCOPE_CYCLE_COUNTER(STAT_AudioMixerSubmixReverb);
 
@@ -166,7 +177,6 @@ void FSubmixEffectReverb::OnProcessAudio(const FSoundEffectSubmixInputData& InDa
 	}
 
 	PlateReverb->ProcessAudio(WetInputBuffer, InData.NumChannels, *OutData.AudioBuffer, OutData.NumChannels);
-
 }
 
 bool FSubmixEffectReverb::SetParameters(const FAudioEffectParameters& InParams)
@@ -223,7 +233,7 @@ bool FSubmixEffectReverb::SetParameters(const FAudioEffectParameters& InParams)
 void FSubmixEffectReverb::UpdateParameters()
 {
 	Audio::FPlateReverbFastSettings NewSettings;
-	if (ReverbParams.GetParams(&NewSettings))
+	if (PlateReverb.IsValid() && ReverbParams.GetParams(&NewSettings))
 	{
 		PlateReverb->SetSettings(NewSettings);
 	}
@@ -246,6 +256,11 @@ void FSubmixEffectReverb::UpdateParameters()
 		TargetQuadBehavior = Audio::FPlateReverbFastSettings::EQuadBehavior::QuadMatched;
 	}
 
+	if (!PlateReverb.IsValid())
+	{
+		return;
+	}
+
 	// Check if settings need to be updated
 	const Audio::FPlateReverbFastSettings& ReverbSettings = PlateReverb->GetSettings();
 	if (ReverbSettings.QuadBehavior != TargetQuadBehavior)
