Dynamic Submixes Reimagined

diff --git a/Engine/Source/Runtime/AudioMixer/Private/AudioMixerDevice.cpp b/Engine/Source/Runtime/AudioMixer/Private/AudioMixerDevice.cpp
--- a/Engine/Source/Runtime/AudioMixer/Private/AudioMixerDevice.cpp
+++ b/Engine/Source/Runtime/AudioMixer/Private/AudioMixerDevice.cpp
@@ -28,12 +28,29 @@
 #include "ISubmixBufferListener.h"
 #include "Misc/App.h"
 #include "ProfilingDebugging/CsvProfiler.h"
+#include "AssetRegistry/IAssetRegistry.h"
+#include "Async/Async.h"
+#include "AudioDeviceNotificationSubsystem.h"
 #include "Sound/AudioFormatSettings.h"
+#include "HAL/PlatformMisc.h"
 
 #if WITH_EDITOR
 #include "AudioEditorModule.h"
 #endif // WITH_EDITOR
 
+#ifndef CASE_ENUM_TO_TEXT_RET
+#define CASE_ENUM_TO_TEXT_RET(OP) case OP: return TEXT(#OP);
+#endif
+
+const TCHAR* LexToString(EMasterSubmixType InType)
+{
+	switch (InType)
+	{
+		FOREACH_ENUM_EMASTERSUBMIXTYPE(CASE_ENUM_TO_TEXT_RET)
+	}
+	return TEXT("Unknown");
+}
+
 
 static int32 DisableSubmixEffectEQCvar = 1;
 FAutoConsoleVariableRef CVarDisableSubmixEQ(
@@ -127,7 +144,7 @@ namespace Audio
 		}
 	}
 
-	FMixerSubmixPtr FSubmixMap::FindRef(const FSubmixMap::FObjectId InObjectId)
+	FMixerSubmixPtr FSubmixMap::FindRef(const FSubmixMap::FObjectId InObjectId) const
 	{
 		if (DisableSubmixMutationLockCVar)
 		{
@@ -166,6 +183,67 @@ namespace Audio
 		}
 	}
 
+	static FAutoConsoleCommand DumpSubmixCmd(
+		TEXT("au.dumpsubmixes"),
+		TEXT("Draws the submix heirarchy for this world to the debug output"),
+		FConsoleCommandWithWorldArgsAndOutputDeviceDelegate::CreateLambda([](const TArray<FString>& InArgs, UWorld* InWorld, FOutputDevice& OutLog)
+		{
+			if (InWorld)
+			{
+				if (const FMixerDevice* MixerDevice = static_cast<FMixerDevice*>(InWorld->GetAudioDeviceRaw()))
+				{
+					MixerDevice->DrawSubmixes(OutLog);
+				}
+			}
+		})
+	);
+
+
+	static void DrawSubmixHeirarchy(USoundSubmixBase* InSubmix, FMixerSubmix* InInstance, const FMixerDevice* InDevice, int32 InIdent, FOutputDevice& Ar);
+	static void DrawSubmixHeirarchy(USoundSubmixBase* InSubmix, FMixerSubmix* InInstance, const FMixerDevice* InDevice, int32 InIdent, FOutputDevice& Ar)
+	{
+		if (!InSubmix)
+		{
+			return;
+		}
+
+		auto SubmixInfo = [](USoundSubmixBase* InUObj, FMixerSubmix* InInst, USoundSubmixBase* InParent, FMixerSubmix* InParentInst) -> FString
+		{
+			FMixerSubmixPtr ParentInstance = InInst->GetParent().Pin();
+			const bool bParentLinkSeemsSane = ParentInstance.IsValid() ? ParentInstance.Get() == InParentInst : false;
+			const int32 NumFxRunning = InInst->GetNumEffects();
+			const int32 NumChildren = InInst->GetChildren().Num();
+			return FString::Printf(TEXT("ParentLinkGood=%s,NumFxRunning=%d,NumChildren=%d"), ToCStr(LexToString(bParentLinkSeemsSane)),NumFxRunning, NumChildren);
+		};
+
+		for (TObjectPtr<USoundSubmixBase> i : InSubmix->ChildSubmixes)
+		{
+			FMixerSubmixPtr ChildInstance = InDevice->GetSubmixInstance(i).Pin();
+			FString Info = SubmixInfo(i.Get(), ChildInstance.Get(), InSubmix, InInstance);
+			Ar.Logf(TEXT("%s [%s](Static)(Instance=0x%p)%s"), FCString::Spc(InIdent * 2), *InSubmix->GetName(), ChildInstance.Get(), *Info);
+			DrawSubmixHeirarchy(i.Get(), ChildInstance.Get(), InDevice, InIdent + 1, Ar);
+		}
+		const auto& DynamicChildren = InSubmix->DynamicChildSubmixes.FindOrAdd(InDevice->DeviceID).ChildSubmixes;
+		for (TObjectPtr<USoundSubmixBase> i : DynamicChildren)
+		{
+			FMixerSubmixPtr ChildInstance = InDevice->GetSubmixInstance(i).Pin();
+			FString Info = SubmixInfo(i.Get(), ChildInstance.Get(), InSubmix, InInstance);
+			Ar.Logf(TEXT("%s [%s](Dynamic)(Instance=0x%p)%s"), FCString::Spc(InIdent * 2), *InSubmix->GetName(), ChildInstance.Get(), *Info);
+			DrawSubmixHeirarchy(i.Get(), ChildInstance.Get(), InDevice,  InIdent + 1, Ar);
+		}
+	}
+
+	void FMixerDevice::DrawSubmixes(FOutputDevice& Output) const
+	{
+		Output.Logf(TEXT("Submix Graph for AudioMixerDevice=%u/(0x%p) "), DeviceID, this);
+		for (int32 i = 0; i < RequiredSubmixes.Num(); i++)
+		{
+			FMixerSubmixPtr Instance = GetRequiredSubmixInstance(RequiredSubmixes[i]);
+			Output.Logf(TEXT("Slot=[%s] (Instance=0x%p)"), ToCStr(LexToString((EMasterSubmixType)i)), Instance.Get());
+			DrawSubmixHeirarchy(RequiredSubmixes[i], Instance.Get(), this, 1, Output);
+		}
+	}
+
 	FMixerDevice::FMixerDevice(IAudioMixerPlatformInterface* InAudioMixerPlatform)
 		: QuantizedEventClockManager(this)
 		, AudioMixerPlatform(InAudioMixerPlatform)
@@ -874,7 +952,7 @@ namespace Audio
 		SourceManager->UpdatePendingReleaseData(true);
 	}
 
-	void FMixerDevice::LoadRequiredSubmix(ERequiredSubmixes::Type InType, const FString& InDefaultName, bool bInDefaultMuteWhenBackgrounded, FSoftObjectPath& InObjectPath)
+	void FMixerDevice::LoadRequiredSubmix(ERequiredSubmixes InType, const FString& InDefaultName, bool bInDefaultMuteWhenBackgrounded, FSoftObjectPath& InObjectPath) 
 	{
 		check(IsInGameThread());
 
@@ -1116,9 +1194,9 @@ namespace Audio
 		FMixerSubmixPtr ParentSubmixInstance;
 		if (const USoundSubmixWithParentBase* SubmixWithParent = Cast<const USoundSubmixWithParentBase>(&SoundSubmix))
 		{
-			if (SubmixWithParent->ParentSubmix)
+			if (TObjectPtr<USoundSubmixBase> Parent = SubmixWithParent->GetParent(DeviceID); Parent)
 			{
-				ParentSubmixInstance = GetSubmixInstance(SubmixWithParent->ParentSubmix).Pin();
+				ParentSubmixInstance = GetSubmixInstance(Parent).Pin();
 			}
 			else
 			{
@@ -1164,43 +1242,43 @@ namespace Audio
 
 	FMixerSubmixWeakPtr FMixerDevice::GetMainSubmix()
 	{
-		return RequiredSubmixInstances[ERequiredSubmixes::Main];
+		return RequiredSubmixInstances[(int32)ERequiredSubmixes::Main];
 	}
 
 	FMixerSubmixWeakPtr FMixerDevice::GetBaseDefaultSubmix()
 	{
-		if (RequiredSubmixInstances[ERequiredSubmixes::BaseDefault].IsValid())
+		if (RequiredSubmixInstances[(int32)ERequiredSubmixes::BaseDefault].IsValid())
 		{
-			return RequiredSubmixInstances[ERequiredSubmixes::BaseDefault];
+			return RequiredSubmixInstances[(int32)ERequiredSubmixes::BaseDefault];
 		}
 		return GetMasterSubmix();
 	}
 
 	FMixerSubmixWeakPtr FMixerDevice::GetReverbSubmix()
 	{
-		return RequiredSubmixInstances[ERequiredSubmixes::Reverb];
+		return RequiredSubmixInstances[(int32)ERequiredSubmixes::Reverb];
 	}
 
 	FMixerSubmixWeakPtr FMixerDevice::GetEQSubmix()
 	{
-		return RequiredSubmixInstances[ERequiredSubmixes::EQ];
+		return RequiredSubmixInstances[(int32)ERequiredSubmixes::EQ];
 	}
 
 	FMixerSubmixWeakPtr FMixerDevice::GetMasterReverbSubmix()
 	{
-		return RequiredSubmixInstances[ERequiredSubmixes::Reverb];
+		return RequiredSubmixInstances[(int32)ERequiredSubmixes::Reverb];
 	}
 
 	FMixerSubmixWeakPtr FMixerDevice::GetMasterEQSubmix()
 	{
-		return RequiredSubmixInstances[ERequiredSubmixes::EQ];
+		return RequiredSubmixInstances[(int32)ERequiredSubmixes::EQ];
 	}
 
 	void FMixerDevice::AddMainSubmixEffect(FSoundEffectSubmixPtr SoundEffectSubmix)
 	{
 		AudioRenderThreadCommand([this, SoundEffectSubmix]()
 		{
-			RequiredSubmixInstances[ERequiredSubmixes::Main]->AddSoundEffectSubmix(SoundEffectSubmix);
+			RequiredSubmixInstances[(int32)ERequiredSubmixes::Main]->AddSoundEffectSubmix(SoundEffectSubmix);
 		});
 	}
 
@@ -1213,7 +1291,7 @@ namespace Audio
 	{
 		AudioRenderThreadCommand([this, SubmixEffectId]()
 		{
-			RequiredSubmixInstances[ERequiredSubmixes::Main]->RemoveSoundEffectSubmix(SubmixEffectId);
+			RequiredSubmixInstances[(int32)ERequiredSubmixes::Main]->RemoveSoundEffectSubmix(SubmixEffectId);
 		});
 	}
 
@@ -1231,7 +1309,7 @@ namespace Audio
 	{
 		AudioRenderThreadCommand([this]()
 		{
-			RequiredSubmixInstances[ERequiredSubmixes::Main]->ClearSoundEffectSubmixes();
+			RequiredSubmixInstances[(int32)ERequiredSubmixes::Main]->ClearSoundEffectSubmixes();
 		});
 	}
 
@@ -1627,7 +1705,7 @@ namespace Audio
 
 	bool FMixerDevice::IsRequiredSubmixType(const USoundSubmixBase* InSubmix) const
 	{
-		for (int32 i = 0; i < ERequiredSubmixes::Count; ++i)
+		for (int32 i = 0; i < (int32)EMasterSubmixType::Count; ++i)
 		{
 			if (InSubmix == RequiredSubmixes[i])
 			{
@@ -1637,12 +1715,12 @@ namespace Audio
 		return false;
 	}
 
-	FMixerSubmixPtr FMixerDevice::GetRequiredSubmixInstance(uint32 InObjectId)
+	FMixerSubmixPtr FMixerDevice::GetRequiredSubmixInstance(uint32 InObjectId) const
 	{
-		check(RequiredSubmixes.Num() == ERequiredSubmixes::Count);
-		for (int32 i = 0; i < ERequiredSubmixes::Count; ++i)
+		check(RequiredSubmixes.Num() == (int32)ERequiredSubmixes::Count);
+		for (int32 i = 0; i < (int32)ERequiredSubmixes::Count; ++i)
 		{
-			if (InObjectId == RequiredSubmixes[i]->GetUniqueID())
+			if (RequiredSubmixes[i] && InObjectId == RequiredSubmixes[i]->GetUniqueID())
 			{
 				return RequiredSubmixInstances[i];
 			}
@@ -1650,10 +1728,10 @@ namespace Audio
 		return nullptr;
 	}
 
-	FMixerSubmixPtr FMixerDevice::GetRequiredSubmixInstance(const USoundSubmixBase* InSubmix)
+	FMixerSubmixPtr FMixerDevice::GetRequiredSubmixInstance(const USoundSubmixBase* InSubmix) const
 	{
-		check(RequiredSubmixes.Num() == ERequiredSubmixes::Count);
-		for (int32 i = 0; i < ERequiredSubmixes::Count; ++i)
+		check(RequiredSubmixes.Num() == (int32)ERequiredSubmixes::Count);
+		for (int32 i = 0; i < (int32)ERequiredSubmixes::Count; ++i)
 		{
 			if (InSubmix == RequiredSubmixes[i])
 			{
@@ -1697,12 +1775,15 @@ namespace Audio
 		{
 			// Ensure parent structure is registered prior to current submix if missing
 			const USoundSubmixWithParentBase* SubmixWithParent = Cast<const USoundSubmixWithParentBase>(InSoundSubmix);
-			if (SubmixWithParent && SubmixWithParent->ParentSubmix)
+
+			TObjectPtr<USoundSubmixBase> Parent = SubmixWithParent->GetParent(DeviceID);
+
+			if (SubmixWithParent && Parent)
 			{
-				FMixerSubmixPtr ParentSubmix = GetSubmixInstance(SubmixWithParent->ParentSubmix).Pin();
+				FMixerSubmixPtr ParentSubmix = GetSubmixInstance(Parent).Pin();
 				if (!ParentSubmix.IsValid())
 				{
-					RegisterSoundSubmix(SubmixWithParent->ParentSubmix, bInit);
+					RegisterSoundSubmix(Parent, bInit);
 				}
 			}
 
@@ -1813,8 +1894,8 @@ namespace Audio
 		FMixerSubmixPtr ParentSubmixInstance;
 		if (const USoundSubmixWithParentBase* InSoundSubmixWithParent = Cast<const USoundSubmixWithParentBase>(&InSoundSubmix))
 		{
-			ParentSubmixInstance = InSoundSubmixWithParent->ParentSubmix
-				? GetSubmixInstance(InSoundSubmixWithParent->ParentSubmix).Pin()
+			ParentSubmixInstance = InSoundSubmixWithParent->GetParent(DeviceID)
+				? GetSubmixInstance(InSoundSubmixWithParent->GetParent(DeviceID)).Pin()
 				: MainSubmix.Pin();
 		}
 
@@ -1872,9 +1953,9 @@ namespace Audio
 			}
 			else
 			{
-				const ERequiredSubmixes::Type SubmixType = static_cast<ERequiredSubmixes::Type>(i);
+				const ERequiredSubmixes SubmixType = static_cast<ERequiredSubmixes>(i);
 				ensureAlwaysMsgf(ERequiredSubmixes::Main != SubmixType,
-					TEXT("Main submix has to be registered before anything else, and is required for the lifetime of the application.")
+					TEXT("Top-level main submix has to be registered before anything else, and is required for the lifetime of the application.")
 				);
 
 				if (!DisableSubmixEffectEQCvar && ERequiredSubmixes::EQ == SubmixType)
@@ -1895,7 +1976,7 @@ namespace Audio
 	{
 	}
 
-	FMixerSubmixWeakPtr FMixerDevice::GetSubmixInstance(const USoundSubmixBase* SoundSubmix)
+	FMixerSubmixWeakPtr FMixerDevice::GetSubmixInstance(const USoundSubmixBase* SoundSubmix) const
 	{
 		LLM_SCOPE(ELLMTag::AudioMixer);
 
diff --git a/Engine/Source/Runtime/AudioMixer/Public/AudioMixerDevice.h b/Engine/Source/Runtime/AudioMixer/Public/AudioMixerDevice.h
--- a/Engine/Source/Runtime/AudioMixer/Public/AudioMixerDevice.h
+++ b/Engine/Source/Runtime/AudioMixer/Public/AudioMixerDevice.h
@@ -20,6 +20,29 @@ class IAudioMixerPlatformInterface;
 class USoundModulatorBase;
 class IAudioLinkFactory;
 
+#include "AudioMixerDevice.generated.h"
+
+// Master submixes
+UENUM()
+enum class ERequiredSubmixes
+{
+	Main = 0,
+	BaseDefault = 1,
+	Reverb = 2,
+	EQ = 3,
+	Count = 4 UMETA(Hidden)
+};
+
+UENUM()
+enum class EMasterSubmixType : uint8
+{
+	Master = ERequiredSubmixes::Main,
+	BaseDefault = ERequiredSubmixes::BaseDefault,
+	Reverb = ERequiredSubmixes::Reverb,
+	EQ = ERequiredSubmixes::EQ,
+	Count = ERequiredSubmixes::Count
+};
+
 namespace Audio
 {
 	// Audio Namespace Forward Declarations
@@ -52,34 +75,9 @@ namespace Audio
 			, AudioRenderThreadTime(0.0)
 			, AudioThreadTimeJitterDelta(0.05)
 		{}
-	};
-
-	// Required submixes
-	namespace ERequiredSubmixes
-	{
-		enum Type
-		{
-			Main = 0,
-			BaseDefault = 1,
-			Reverb = 2,
-			EQ = 3,
-			Count = 4
-		};
-	}
-
-	// Deprecated, use ERequiredSubmixes above
-	namespace EMasterSubmixType
-	{
-		enum Type
-		{
-			Master = ERequiredSubmixes::Main,
-			BaseDefault = ERequiredSubmixes::BaseDefault,
-			Reverb = ERequiredSubmixes::Reverb,
-			EQ = ERequiredSubmixes::EQ,
-			Count = ERequiredSubmixes::Count
-		};
-	}
+	};	
 
+	
 	struct FSubmixMap
 	{
 	public:
@@ -89,14 +87,14 @@ namespace Audio
 
 		void Add(const FObjectId InObjectId, FMixerSubmixPtr InMixerSubmix);
 		void Iterate(FIterFunc InFunction);
-		FMixerSubmixPtr FindRef(FObjectId InObjectId);
+		FMixerSubmixPtr FindRef(FObjectId InObjectId) const;
 		int32 Remove(const FObjectId InObjectId);
 		void Reset();
 
 	private:
 		TMap<FObjectId, FMixerSubmixPtr> SubmixMap;
 
-		FCriticalSection MutationLock;
+		mutable FCriticalSection MutationLock;
 	};
 
 
@@ -225,7 +223,7 @@ namespace Audio
 
 		AUDIOMIXER_API FMixerSubmixPtr FindSubmixInstanceByObjectId(uint32 InObjectId);
 
-		AUDIOMIXER_API FMixerSubmixWeakPtr GetSubmixInstance(const USoundSubmixBase* SoundSubmix);
+		AUDIOMIXER_API FMixerSubmixWeakPtr GetSubmixInstance(const USoundSubmixBase* SoundSubmix) const;
 
 		// If SoundSubmix is a soundfield submix, this will return the factory used to encode 
 		// source audio to it's soundfield format.
@@ -368,8 +366,8 @@ namespace Audio
 		// Pushes the command to a MPSC queue to be executed on the game thread
 		AUDIOMIXER_API void GameThreadMPSCCommand(TFunction<void()> InCommand);
 
+		AUDIOMIXER_API void DrawSubmixes(FOutputDevice& Output) const;
 	protected:
-
 		AUDIOMIXER_API virtual void InitSoundSubmixes() override;
 
 		AUDIOMIXER_API virtual void OnListenerUpdated(const TArray<FListener>& InListeners) override;
@@ -392,7 +390,7 @@ namespace Audio
 
 		bool IsMainAudioDevice() const;
 
-		void LoadRequiredSubmix(ERequiredSubmixes::Type InType, const FString& InDefaultName, bool bInDefaultMuteWhenBackgrounded, FSoftObjectPath& InOutObjectPath);
+		void LoadRequiredSubmix(ERequiredSubmixes InType, const FString& InDefaultName, bool bInDefaultMuteWhenBackgrounded, FSoftObjectPath& InOutObjectPath);
 		void LoadPluginSoundSubmixes();
 		void LoadSoundSubmix(USoundSubmixBase& SoundSubmix);
 
@@ -403,8 +401,8 @@ namespace Audio
 		ICompressedAudioInfo* CreateAudioInfo(FName InFormat) const;
 
 		bool IsRequiredSubmixType(const USoundSubmixBase* InSubmix) const;
-		FMixerSubmixPtr GetRequiredSubmixInstance(uint32 InSubmixId);
-		FMixerSubmixPtr GetRequiredSubmixInstance(const USoundSubmixBase* InSubmix);
+		FMixerSubmixPtr GetRequiredSubmixInstance(uint32 InSubmixId) const;
+		FMixerSubmixPtr GetRequiredSubmixInstance(const USoundSubmixBase* InSubmix) const;
 		
 		// Pumps the audio render thread command queue
 		void PumpCommandQueue();
diff --git a/Engine/Source/Runtime/AudioMixer/Public/AudioMixerSubmix.h b/Engine/Source/Runtime/AudioMixer/Public/AudioMixerSubmix.h
--- a/Engine/Source/Runtime/AudioMixer/Public/AudioMixerSubmix.h
+++ b/Engine/Source/Runtime/AudioMixer/Public/AudioMixerSubmix.h
@@ -337,6 +337,9 @@ namespace Audio
 
 		AUDIOMIXER_API FSoundfieldSpeakerPositionalData GetDefaultPositionalDataForAudioDevice();
 
+		AUDIOMIXER_API TWeakPtr<FMixerSubmix, ESPMode::ThreadSafe>  GetParent() const { return ParentSubmix; }
+		AUDIOMIXER_API const TMap<uint32, FChildSubmixInfo>& GetChildren() const { return ChildSubmixes; }
+
 	protected:
 		// Initialize the submix internal
 		AUDIOMIXER_API void InitInternal();
diff --git a/Engine/Source/Runtime/Engine/Classes/Sound/SoundSubmix.h b/Engine/Source/Runtime/Engine/Classes/Sound/SoundSubmix.h
--- a/Engine/Source/Runtime/Engine/Classes/Sound/SoundSubmix.h
+++ b/Engine/Source/Runtime/Engine/Classes/Sound/SoundSubmix.h
@@ -159,6 +159,15 @@ public:
 };
 #endif
 
+USTRUCT()
+struct FDynamicChildSubmix  
+{
+	GENERATED_USTRUCT_BODY()
+	
+	UPROPERTY(Transient)
+	TArray<TObjectPtr<USoundSubmixBase>> ChildSubmixes;
+};
+
 UCLASS(config = Engine, abstract, hidecategories = Object, editinlinenew, BlueprintType, MinimalAPI)
 class USoundSubmixBase : public UObject
 {
@@ -182,6 +191,10 @@ public:
 	UPROPERTY(VisibleAnywhere, BlueprintReadOnly, Category = SoundSubmix)
 	TArray<TObjectPtr<USoundSubmixBase>> ChildSubmixes;
 
+	// Dynamic Child submixes
+	UPROPERTY(Transient)
+	TMap<uint32, FDynamicChildSubmix> DynamicChildSubmixes;
+	
 protected:
 	//~ Begin UObject Interface.
 	ENGINE_API virtual FString GetDesc() override;
@@ -228,9 +241,26 @@ class USoundSubmixWithParentBase : public USoundSubmixBase
 	GENERATED_UCLASS_BODY()
 public:
 
+
 	UPROPERTY(VisibleAnywhere, BlueprintReadOnly, Category = SoundSubmix)
 	TObjectPtr<USoundSubmixBase> ParentSubmix;
 
+	/**
+	* Holds the dynamic (not serialized) parent for this Submix.
+	* Overrides the serialized one.
+	*/
+	
+	ENGINE_API TObjectPtr<USoundSubmixBase> GetParent(Audio::FDeviceId InDeviceId) const;
+	
+	UPROPERTY(Transient)
+	TMap<uint32, TObjectPtr<USoundSubmixBase>> DynamicParentSubmix;
+
+	UFUNCTION(BlueprintCallable, Category = "Audio", meta = (WorldContext = "WorldContextObject", DisplayName = "Connect"))
+	ENGINE_API bool DynamicConnect(const UObject* WorldContextObject, USoundSubmixBase* Parent);
+	
+	UFUNCTION(BlueprintCallable, Category = "Audio", meta = (WorldContext = "WorldContextObject", DisplayName = "Disconnect"))
+	ENGINE_API bool DynamicDisconnect(const UObject* WorldContextObject);
+	
 	/**
 	* Set the parent submix of this SoundSubmix, removing it as a child from its previous owner
 	*
diff --git a/Engine/Source/Runtime/Engine/Private/SoundSubmix.cpp b/Engine/Source/Runtime/Engine/Private/SoundSubmix.cpp
--- a/Engine/Source/Runtime/Engine/Private/SoundSubmix.cpp
+++ b/Engine/Source/Runtime/Engine/Private/SoundSubmix.cpp
@@ -621,6 +621,113 @@ void USoundSubmixBase::PostLoad()
 	}
 }
 
+
+TObjectPtr<USoundSubmixBase> USoundSubmixWithParentBase::GetParent(Audio::FDeviceId InDeviceId) const
+{
+	// Dynamic parent?
+	if (const TObjectPtr<USoundSubmixBase>* pFound = DynamicParentSubmix.Find(InDeviceId))
+	{
+		return *pFound;
+	}
+	return ParentSubmix;
+}
+
+bool USoundSubmixWithParentBase::DynamicConnect(const UObject* WorldContextObject, USoundSubmixBase* InParent)
+{	
+	if (!WorldContextObject)
+	{
+		UE_LOG(LogAudio, Warning, TEXT("Submix (DynamicConnect): World Context is null for [%s]"), *GetName());
+		return false;
+	}
+		
+	const UWorld* World = WorldContextObject->GetWorld();
+	if (!World)
+	{
+		UE_LOG(LogAudio, Warning, TEXT("Submix (DynamicConnect): World is null for [%s]"), *GetName());		
+		return false;
+	}
+
+	FAudioDeviceHandle Handle = World->GetAudioDevice();
+	if (!Handle.IsValid())
+	{
+		UE_LOG(LogAudio, Warning, TEXT("Submix (DynamicConnect): No valid audio device in this world for [%s]"), *GetName());				
+		return false;
+	}
+
+	const Audio::DeviceID Id = Handle.GetDeviceID();
+	
+	TObjectPtr<USoundSubmixBase>& CurrentParent = DynamicParentSubmix.FindOrAdd(Id);
+
+	if (CurrentParent != InParent)
+	{
+		if (CurrentParent)
+		{
+			CurrentParent->DynamicChildSubmixes.FindOrAdd(Id).ChildSubmixes.Remove(this);
+			UE_LOG(LogAudio, Verbose, TEXT("Submix (DynamicConnect): Reparenting [%s] and removing from it's parent [%s]"), *GetName(), *CurrentParent->GetName());
+		}
+
+		CurrentParent = InParent;
+
+		if (CurrentParent)
+		{
+			CurrentParent->DynamicChildSubmixes.FindOrAdd(Id).ChildSubmixes.AddUnique(this);
+			UE_LOG(LogAudio, Verbose, TEXT("Submix (DynamicConnect): Reparenting [%s] and adding to it's new parent [%s]"), *GetName(), *CurrentParent->GetName());
+		}
+
+		Handle->RegisterSoundSubmix(this, true);
+		UE_LOG(LogAudio, Verbose, TEXT("Submix (DynamicConnect): Registering [%s] with World's [%s] AudioDevice [%u]"), *GetName(), *World->GetName(), Handle.GetDeviceID());
+
+		return CurrentParent != nullptr;
+	}
+
+	UE_LOG(LogAudio, Warning, TEXT("Submix (DynamicConnect): Submix [%s] was already connected to [%s]"), *GetName(), *GetNameSafe(CurrentParent));
+	return false;
+}
+bool USoundSubmixWithParentBase::DynamicDisconnect(const UObject* WorldContextObject)
+{
+	if (!WorldContextObject)
+	{
+		UE_LOG(LogAudio, Warning, TEXT("Submix (DynamicDisconnect): World Context is null for [%s]"), *GetName());
+		return false;
+	}
+		
+	const UWorld* World = WorldContextObject->GetWorld();
+	if (!World)
+	{
+		UE_LOG(LogAudio, Warning, TEXT("Submix (DynamicDisconnect): World is null for [%s]"), *GetName());		
+		return false;
+	}
+
+	FAudioDeviceHandle Handle = World->GetAudioDevice();
+	if (!Handle.IsValid())
+	{
+		UE_LOG(LogAudio, Warning, TEXT("Submix (DynamicDisconnect): No valid audio device in this world for [%s]"), *GetName());				
+		return false;
+	}
+
+	const Audio::DeviceID Id = Handle.GetDeviceID();
+
+	TObjectPtr<USoundSubmixBase>& CurrentParent = DynamicParentSubmix.FindOrAdd(Id);
+
+	if (CurrentParent)
+	{
+		UE_LOG(LogAudio, Verbose, TEXT("Submix (DynamicDisconnect): Removing [%s] from it's parent [%s]"), *GetName(), *CurrentParent->GetName());
+		
+		CurrentParent->DynamicChildSubmixes.FindOrAdd(Id).ChildSubmixes.Remove(this);
+		CurrentParent = nullptr;
+
+		Handle->UnregisterSoundSubmix(this);
+
+		UE_LOG(LogAudio, Verbose, TEXT("Submix (DynamicDisconnect): Unregistering [%s] with World's [%s] AudioDevice [%u]"), 
+		       *GetName(), *World->GetName(), Handle.GetDeviceID());
+
+		return true;
+	}
+
+	UE_LOG(LogAudio, Warning, TEXT("Submix (DynamicDisconnect): Submix was not connected to any dynamic parent [%s]"), *GetName());
+	return false;
+}
+
 #if WITH_EDITOR
 
 void USoundSubmixBase::PostDuplicate(EDuplicateMode::Type DuplicateMode)
@@ -744,6 +851,7 @@ void USoundSubmixWithParentBase::SetParentSubmix(USoundSubmixBase* InParentSubmi
 	}
 }
 
+
 #if WITH_EDITOR
 void USoundSubmixWithParentBase::PostEditChangeProperty(struct FPropertyChangedEvent& PropertyChangedEvent)
 {
