[Backout] - CL27700227 - CIS Compile Error
Original CL Desc
-----------------------------------------------------------------
Dynamic Submixes Reimagined

diff --git a/Engine/Source/Runtime/AudioMixer/Private/AudioMixerDevice.cpp b/Engine/Source/Runtime/AudioMixer/Private/AudioMixerDevice.cpp
--- a/Engine/Source/Runtime/AudioMixer/Private/AudioMixerDevice.cpp
+++ b/Engine/Source/Runtime/AudioMixer/Private/AudioMixerDevice.cpp
@@ -28,29 +28,12 @@
 #include "ISubmixBufferListener.h"
 #include "Misc/App.h"
 #include "ProfilingDebugging/CsvProfiler.h"
-#include "AssetRegistry/IAssetRegistry.h"
-#include "Async/Async.h"
-#include "AudioDeviceNotificationSubsystem.h"
 #include "Sound/AudioFormatSettings.h"
-#include "HAL/PlatformMisc.h"
 
 #if WITH_EDITOR
 #include "AudioEditorModule.h"
 #endif // WITH_EDITOR
 
-#ifndef CASE_ENUM_TO_TEXT_RET
-#define CASE_ENUM_TO_TEXT_RET(OP) case OP: return TEXT(#OP);
-#endif
-
-const TCHAR* LexToString(EMasterSubmixType InType)
-{
-	switch (InType)
-	{
-		FOREACH_ENUM_EMASTERSUBMIXTYPE(CASE_ENUM_TO_TEXT_RET)
-	}
-	return TEXT("Unknown");
-}
-
 
 static int32 DisableSubmixEffectEQCvar = 1;
 FAutoConsoleVariableRef CVarDisableSubmixEQ(
@@ -144,7 +127,7 @@ namespace Audio
 		}
 	}
 
-	FMixerSubmixPtr FSubmixMap::FindRef(const FSubmixMap::FObjectId InObjectId) const
+	FMixerSubmixPtr FSubmixMap::FindRef(const FSubmixMap::FObjectId InObjectId)
 	{
 		if (DisableSubmixMutationLockCVar)
 		{
@@ -183,67 +166,6 @@ namespace Audio
 		}
 	}
 
-	static FAutoConsoleCommand DumpSubmixCmd(
-		TEXT("au.dumpsubmixes"),
-		TEXT("Draws the submix heirarchy for this world to the debug output"),
-		FConsoleCommandWithWorldArgsAndOutputDeviceDelegate::CreateLambda([](const TArray<FString>& InArgs, UWorld* InWorld, FOutputDevice& OutLog)
-		{
-			if (InWorld)
-			{
-				if (const FMixerDevice* MixerDevice = static_cast<FMixerDevice*>(InWorld->GetAudioDeviceRaw()))
-				{
-					MixerDevice->DrawSubmixes(OutLog);
-				}
-			}
-		})
-	);
-
-
-	static void DrawSubmixHeirarchy(USoundSubmixBase* InSubmix, FMixerSubmix* InInstance, const FMixerDevice* InDevice, int32 InIdent, FOutputDevice& Ar);
-	static void DrawSubmixHeirarchy(USoundSubmixBase* InSubmix, FMixerSubmix* InInstance, const FMixerDevice* InDevice, int32 InIdent, FOutputDevice& Ar)
-	{
-		if (!InSubmix)
-		{
-			return;
-		}
-
-		auto SubmixInfo = [](USoundSubmixBase* InUObj, FMixerSubmix* InInst, USoundSubmixBase* InParent, FMixerSubmix* InParentInst) -> FString
-		{
-			FMixerSubmixPtr ParentInstance = InInst->GetParent().Pin();
-			const bool bParentLinkSeemsSane = ParentInstance.IsValid() ? ParentInstance.Get() == InParentInst : false;
-			const int32 NumFxRunning = InInst->GetNumEffects();
-			const int32 NumChildren = InInst->GetChildren().Num();
-			return FString::Printf(TEXT("ParentLinkGood=%s,NumFxRunning=%d,NumChildren=%d"), ToCStr(LexToString(bParentLinkSeemsSane)),NumFxRunning, NumChildren);
-		};
-
-		for (TObjectPtr<USoundSubmixBase> i : InSubmix->ChildSubmixes)
-		{
-			FMixerSubmixPtr ChildInstance = InDevice->GetSubmixInstance(i).Pin();
-			FString Info = SubmixInfo(i.Get(), ChildInstance.Get(), InSubmix, InInstance);
-			Ar.Logf(TEXT("%s [%s](Static)(Instance=0x%p)%s"), FCString::Spc(InIdent * 2), *InSubmix->GetName(), ChildInstance.Get(), *Info);
-			DrawSubmixHeirarchy(i.Get(), ChildInstance.Get(), InDevice, InIdent + 1, Ar);
-		}
-		const auto& DynamicChildren = InSubmix->DynamicChildSubmixes.FindOrAdd(InDevice->DeviceID).ChildSubmixes;
-		for (TObjectPtr<USoundSubmixBase> i : DynamicChildren)
-		{
-			FMixerSubmixPtr ChildInstance = InDevice->GetSubmixInstance(i).Pin();
-			FString Info = SubmixInfo(i.Get(), ChildInstance.Get(), InSubmix, InInstance);
-			Ar.Logf(TEXT("%s [%s](Dynamic)(Instance=0x%p)%s"), FCString::Spc(InIdent * 2), *InSubmix->GetName(), ChildInstance.Get(), *Info);
-			DrawSubmixHeirarchy(i.Get(), ChildInstance.Get(), InDevice,  InIdent + 1, Ar);
-		}
-	}
-
-	void FMixerDevice::DrawSubmixes(FOutputDevice& Output) const
-	{
-		Output.Logf(TEXT("Submix Graph for AudioMixerDevice=%u/(0x%p) "), DeviceID, this);
-		for (int32 i = 0; i < RequiredSubmixes.Num(); i++)
-		{
-			FMixerSubmixPtr Instance = GetRequiredSubmixInstance(RequiredSubmixes[i]);
-			Output.Logf(TEXT("Slot=[%s] (Instance=0x%p)"), ToCStr(LexToString((EMasterSubmixType)i)), Instance.Get());
-			DrawSubmixHeirarchy(RequiredSubmixes[i], Instance.Get(), this, 1, Output);
-		}
-	}
-
 	FMixerDevice::FMixerDevice(IAudioMixerPlatformInterface* InAudioMixerPlatform)
 		: QuantizedEventClockManager(this)
 		, AudioMixerPlatform(InAudioMixerPlatform)
@@ -952,7 +874,7 @@ namespace Audio
 		SourceManager->UpdatePendingReleaseData(true);
 	}
 
-	void FMixerDevice::LoadRequiredSubmix(ERequiredSubmixes InType, const FString& InDefaultName, bool bInDefaultMuteWhenBackgrounded, FSoftObjectPath& InObjectPath) 
+	void FMixerDevice::LoadRequiredSubmix(ERequiredSubmixes::Type InType, const FString& InDefaultName, bool bInDefaultMuteWhenBackgrounded, FSoftObjectPath& InObjectPath)
 	{
 		check(IsInGameThread());
 
@@ -1194,9 +1116,9 @@ namespace Audio
 		FMixerSubmixPtr ParentSubmixInstance;
 		if (const USoundSubmixWithParentBase* SubmixWithParent = Cast<const USoundSubmixWithParentBase>(&SoundSubmix))
 		{
-			if (TObjectPtr<USoundSubmixBase> Parent = SubmixWithParent->GetParent(DeviceID); Parent)
+			if (SubmixWithParent->ParentSubmix)
 			{
-				ParentSubmixInstance = GetSubmixInstance(Parent).Pin();
+				ParentSubmixInstance = GetSubmixInstance(SubmixWithParent->ParentSubmix).Pin();
 			}
 			else
 			{
@@ -1242,43 +1164,43 @@ namespace Audio
 
 	FMixerSubmixWeakPtr FMixerDevice::GetMainSubmix()
 	{
-		return RequiredSubmixInstances[(int32)ERequiredSubmixes::Main];
+		return RequiredSubmixInstances[ERequiredSubmixes::Main];
 	}
 
 	FMixerSubmixWeakPtr FMixerDevice::GetBaseDefaultSubmix()
 	{
-		if (RequiredSubmixInstances[(int32)ERequiredSubmixes::BaseDefault].IsValid())
+		if (RequiredSubmixInstances[ERequiredSubmixes::BaseDefault].IsValid())
 		{
-			return RequiredSubmixInstances[(int32)ERequiredSubmixes::BaseDefault];
+			return RequiredSubmixInstances[ERequiredSubmixes::BaseDefault];
 		}
 		return GetMasterSubmix();
 	}
 
 	FMixerSubmixWeakPtr FMixerDevice::GetReverbSubmix()
 	{
-		return RequiredSubmixInstances[(int32)ERequiredSubmixes::Reverb];
+		return RequiredSubmixInstances[ERequiredSubmixes::Reverb];
 	}
 
 	FMixerSubmixWeakPtr FMixerDevice::GetEQSubmix()
 	{
-		return RequiredSubmixInstances[(int32)ERequiredSubmixes::EQ];
+		return RequiredSubmixInstances[ERequiredSubmixes::EQ];
 	}
 
 	FMixerSubmixWeakPtr FMixerDevice::GetMasterReverbSubmix()
 	{
-		return RequiredSubmixInstances[(int32)ERequiredSubmixes::Reverb];
+		return RequiredSubmixInstances[ERequiredSubmixes::Reverb];
 	}
 
 	FMixerSubmixWeakPtr FMixerDevice::GetMasterEQSubmix()
 	{
-		return RequiredSubmixInstances[(int32)ERequiredSubmixes::EQ];
+		return RequiredSubmixInstances[ERequiredSubmixes::EQ];
 	}
 
 	void FMixerDevice::AddMainSubmixEffect(FSoundEffectSubmixPtr SoundEffectSubmix)
 	{
 		AudioRenderThreadCommand([this, SoundEffectSubmix]()
 		{
-			RequiredSubmixInstances[(int32)ERequiredSubmixes::Main]->AddSoundEffectSubmix(SoundEffectSubmix);
+			RequiredSubmixInstances[ERequiredSubmixes::Main]->AddSoundEffectSubmix(SoundEffectSubmix);
 		});
 	}
 
@@ -1291,7 +1213,7 @@ namespace Audio
 	{
 		AudioRenderThreadCommand([this, SubmixEffectId]()
 		{
-			RequiredSubmixInstances[(int32)ERequiredSubmixes::Main]->RemoveSoundEffectSubmix(SubmixEffectId);
+			RequiredSubmixInstances[ERequiredSubmixes::Main]->RemoveSoundEffectSubmix(SubmixEffectId);
 		});
 	}
 
@@ -1309,7 +1231,7 @@ namespace Audio
 	{
 		AudioRenderThreadCommand([this]()
 		{
-			RequiredSubmixInstances[(int32)ERequiredSubmixes::Main]->ClearSoundEffectSubmixes();
+			RequiredSubmixInstances[ERequiredSubmixes::Main]->ClearSoundEffectSubmixes();
 		});
 	}
 
@@ -1705,7 +1627,7 @@ namespace Audio
 
 	bool FMixerDevice::IsRequiredSubmixType(const USoundSubmixBase* InSubmix) const
 	{
-		for (int32 i = 0; i < (int32)EMasterSubmixType::Count; ++i)
+		for (int32 i = 0; i < ERequiredSubmixes::Count; ++i)
 		{
 			if (InSubmix == RequiredSubmixes[i])
 			{
@@ -1715,12 +1637,12 @@ namespace Audio
 		return false;
 	}
 
-	FMixerSubmixPtr FMixerDevice::GetRequiredSubmixInstance(uint32 InObjectId) const
+	FMixerSubmixPtr FMixerDevice::GetRequiredSubmixInstance(uint32 InObjectId)
 	{
-		check(RequiredSubmixes.Num() == (int32)ERequiredSubmixes::Count);
-		for (int32 i = 0; i < (int32)ERequiredSubmixes::Count; ++i)
+		check(RequiredSubmixes.Num() == ERequiredSubmixes::Count);
+		for (int32 i = 0; i < ERequiredSubmixes::Count; ++i)
 		{
-			if (RequiredSubmixes[i] && InObjectId == RequiredSubmixes[i]->GetUniqueID())
+			if (InObjectId == RequiredSubmixes[i]->GetUniqueID())
 			{
 				return RequiredSubmixInstances[i];
 			}
@@ -1728,10 +1650,10 @@ namespace Audio
 		return nullptr;
 	}
 
-	FMixerSubmixPtr FMixerDevice::GetRequiredSubmixInstance(const USoundSubmixBase* InSubmix) const
+	FMixerSubmixPtr FMixerDevice::GetRequiredSubmixInstance(const USoundSubmixBase* InSubmix)
 	{
-		check(RequiredSubmixes.Num() == (int32)ERequiredSubmixes::Count);
-		for (int32 i = 0; i < (int32)ERequiredSubmixes::Count; ++i)
+		check(RequiredSubmixes.Num() == ERequiredSubmixes::Count);
+		for (int32 i = 0; i < ERequiredSubmixes::Count; ++i)
 		{
 			if (InSubmix == RequiredSubmixes[i])
 			{
@@ -1775,15 +1697,12 @@ namespace Audio
 		{
 			// Ensure parent structure is registered prior to current submix if missing
 			const USoundSubmixWithParentBase* SubmixWithParent = Cast<const USoundSubmixWithParentBase>(InSoundSubmix);
-
-			TObjectPtr<USoundSubmixBase> Parent = SubmixWithParent->GetParent(DeviceID);
-
-			if (SubmixWithParent && Parent)
+			if (SubmixWithParent && SubmixWithParent->ParentSubmix)
 			{
-				FMixerSubmixPtr ParentSubmix = GetSubmixInstance(Parent).Pin();
+				FMixerSubmixPtr ParentSubmix = GetSubmixInstance(SubmixWithParent->ParentSubmix).Pin();
 				if (!ParentSubmix.IsValid())
 				{
-					RegisterSoundSubmix(Parent, bInit);
+					RegisterSoundSubmix(SubmixWithParent->ParentSubmix, bInit);
 				}
 			}
 
@@ -1894,8 +1813,8 @@ namespace Audio
 		FMixerSubmixPtr ParentSubmixInstance;
 		if (const USoundSubmixWithParentBase* InSoundSubmixWithParent = Cast<const USoundSubmixWithParentBase>(&InSoundSubmix))
 		{
-			ParentSubmixInstance = InSoundSubmixWithParent->GetParent(DeviceID)
-				? GetSubmixInstance(InSoundSubmixWithParent->GetParent(DeviceID)).Pin()
+			ParentSubmixInstance = InSoundSubmixWithParent->ParentSubmix
+				? GetSubmixInstance(InSoundSubmixWithParent->ParentSubmix).Pin()
 				: MainSubmix.Pin();
 		}
 
@@ -1953,9 +1872,9 @@ namespace Audio
 			}
 			else
 			{
-				const ERequiredSubmixes SubmixType = static_cast<ERequiredSubmixes>(i);
+				const ERequiredSubmixes::Type SubmixType = static_cast<ERequiredSubmixes::Type>(i);
 				ensureAlwaysMsgf(ERequiredSubmixes::Main != SubmixType,
-					TEXT("Top-level main submix has to be registered before anything else, and is required for the lifetime of the application.")
+					TEXT("Main submix has to be registered before anything else, and is required for the lifetime of the application.")
 				);
 
 				if (!DisableSubmixEffectEQCvar && ERequiredSubmixes::EQ == SubmixType)
@@ -1976,7 +1895,7 @@ namespace Audio
 	{
 	}
 
-	FMixerSubmixWeakPtr FMixerDevice::GetSubmixInstance(const USoundSubmixBase* SoundSubmix) const
+	FMixerSubmixWeakPtr FMixerDevice::GetSubmixInstance(const USoundSubmixBase* SoundSubmix)
 	{
 		LLM_SCOPE(ELLMTag::AudioMixer);
 
diff --git a/Engine/Source/Runtime/AudioMixer/Public/AudioMixerDevice.h b/Engine/Source/Runtime/AudioMixer/Public/AudioMixerDevice.h
--- a/Engine/Source/Runtime/AudioMixer/Public/AudioMixerDevice.h
+++ b/Engine/Source/Runtime/AudioMixer/Public/AudioMixerDevice.h
@@ -20,29 +20,6 @@ class IAudioMixerPlatformInterface;
 class USoundModulatorBase;
 class IAudioLinkFactory;
 
-#include "AudioMixerDevice.generated.h"
-
-// Master submixes
-UENUM()
-enum class ERequiredSubmixes
-{
-	Main = 0,
-	BaseDefault = 1,
-	Reverb = 2,
-	EQ = 3,
-	Count = 4 UMETA(Hidden)
-};
-
-UENUM()
-enum class EMasterSubmixType : uint8
-{
-	Master = ERequiredSubmixes::Main,
-	BaseDefault = ERequiredSubmixes::BaseDefault,
-	Reverb = ERequiredSubmixes::Reverb,
-	EQ = ERequiredSubmixes::EQ,
-	Count = ERequiredSubmixes::Count
-};
-
 namespace Audio
 {
 	// Audio Namespace Forward Declarations
@@ -75,9 +52,34 @@ namespace Audio
 			, AudioRenderThreadTime(0.0)
 			, AudioThreadTimeJitterDelta(0.05)
 		{}
-	};	
+	};
+
+	// Required submixes
+	namespace ERequiredSubmixes
+	{
+		enum Type
+		{
+			Main = 0,
+			BaseDefault = 1,
+			Reverb = 2,
+			EQ = 3,
+			Count = 4
+		};
+	}
+
+	// Deprecated, use ERequiredSubmixes above
+	namespace EMasterSubmixType
+	{
+		enum Type
+		{
+			Master = ERequiredSubmixes::Main,
+			BaseDefault = ERequiredSubmixes::BaseDefault,
+			Reverb = ERequiredSubmixes::Reverb,
+			EQ = ERequiredSubmixes::EQ,
+			Count = ERequiredSubmixes::Count
+		};
+	}
 
-	
 	struct FSubmixMap
 	{
 	public:
@@ -87,14 +89,14 @@ namespace Audio
 
 		void Add(const FObjectId InObjectId, FMixerSubmixPtr InMixerSubmix);
 		void Iterate(FIterFunc InFunction);
-		FMixerSubmixPtr FindRef(FObjectId InObjectId) const;
+		FMixerSubmixPtr FindRef(FObjectId InObjectId);
 		int32 Remove(const FObjectId InObjectId);
 		void Reset();
 
 	private:
 		TMap<FObjectId, FMixerSubmixPtr> SubmixMap;
 
-		mutable FCriticalSection MutationLock;
+		FCriticalSection MutationLock;
 	};
 
 
@@ -223,7 +225,7 @@ namespace Audio
 
 		AUDIOMIXER_API FMixerSubmixPtr FindSubmixInstanceByObjectId(uint32 InObjectId);
 
-		AUDIOMIXER_API FMixerSubmixWeakPtr GetSubmixInstance(const USoundSubmixBase* SoundSubmix) const;
+		AUDIOMIXER_API FMixerSubmixWeakPtr GetSubmixInstance(const USoundSubmixBase* SoundSubmix);
 
 		// If SoundSubmix is a soundfield submix, this will return the factory used to encode 
 		// source audio to it's soundfield format.
@@ -366,8 +368,8 @@ namespace Audio
 		// Pushes the command to a MPSC queue to be executed on the game thread
 		AUDIOMIXER_API void GameThreadMPSCCommand(TFunction<void()> InCommand);
 
-		AUDIOMIXER_API void DrawSubmixes(FOutputDevice& Output) const;
 	protected:
+
 		AUDIOMIXER_API virtual void InitSoundSubmixes() override;
 
 		AUDIOMIXER_API virtual void OnListenerUpdated(const TArray<FListener>& InListeners) override;
@@ -390,7 +392,7 @@ namespace Audio
 
 		bool IsMainAudioDevice() const;
 
-		void LoadRequiredSubmix(ERequiredSubmixes InType, const FString& InDefaultName, bool bInDefaultMuteWhenBackgrounded, FSoftObjectPath& InOutObjectPath);
+		void LoadRequiredSubmix(ERequiredSubmixes::Type InType, const FString& InDefaultName, bool bInDefaultMuteWhenBackgrounded, FSoftObjectPath& InOutObjectPath);
 		void LoadPluginSoundSubmixes();
 		void LoadSoundSubmix(USoundSubmixBase& SoundSubmix);
 
@@ -401,8 +403,8 @@ namespace Audio
 		ICompressedAudioInfo* CreateAudioInfo(FName InFormat) const;
 
 		bool IsRequiredSubmixType(const USoundSubmixBase* InSubmix) const;
-		FMixerSubmixPtr GetRequiredSubmixInstance(uint32 InSubmixId) const;
-		FMixerSubmixPtr GetRequiredSubmixInstance(const USoundSubmixBase* InSubmix) const;
+		FMixerSubmixPtr GetRequiredSubmixInstance(uint32 InSubmixId);
+		FMixerSubmixPtr GetRequiredSubmixInstance(const USoundSubmixBase* InSubmix);
 		
 		// Pumps the audio render thread command queue
 		void PumpCommandQueue();
