Fix nullptr deref after bad cast in submix register

diff --git a/Engine/Source/Runtime/AudioMixer/Private/AudioMixerDevice.cpp b/Engine/Source/Runtime/AudioMixer/Private/AudioMixerDevice.cpp
--- a/Engine/Source/Runtime/AudioMixer/Private/AudioMixerDevice.cpp
+++ b/Engine/Source/Runtime/AudioMixer/Private/AudioMixerDevice.cpp
@@ -1773,16 +1773,15 @@ namespace Audio
 		if (!bIsMainSubmix)
 		{
 			// Ensure parent structure is registered prior to current submix if missing
-			const USoundSubmixWithParentBase* SubmixWithParent = Cast<const USoundSubmixWithParentBase>(InSoundSubmix);
-
-			TObjectPtr<USoundSubmixBase> Parent = SubmixWithParent->GetParent(DeviceID);
-
-			if (SubmixWithParent && Parent)
+			if (const USoundSubmixWithParentBase* SubmixWithParent = Cast<const USoundSubmixWithParentBase>(InSoundSubmix))
 			{
-				FMixerSubmixPtr ParentSubmix = GetSubmixInstance(Parent).Pin();
-				if (!ParentSubmix.IsValid())
+				if (TObjectPtr<USoundSubmixBase> Parent = SubmixWithParent->GetParent(DeviceID))
 				{
-					RegisterSoundSubmix(Parent, bInit);
+					FMixerSubmixPtr ParentSubmix = GetSubmixInstance(Parent).Pin();
+					if (!ParentSubmix.IsValid())
+					{
+						RegisterSoundSubmix(Parent, bInit);
+					}
 				}
 			}
 
