Reduced the amount of memory read\writes in FMixerSourceManager::ComputeSourceBuffersForIdRange

diff --git a/Engine/Source/Runtime/AudioMixer/Private/AudioMixerSourceManager.cpp b/Engine/Source/Runtime/AudioMixer/Private/AudioMixerSourceManager.cpp
--- a/Engine/Source/Runtime/AudioMixer/Private/AudioMixerSourceManager.cpp
+++ b/Engine/Source/Runtime/AudioMixer/Private/AudioMixerSourceManager.cpp
@@ -2233,6 +2233,8 @@ namespace Audio
 				}
 #endif // UE_AUDIO_PROFILERTRACE_ENABLED
 
+				float CurrentAlpha = SourceInfo.CurrentFrameAlpha;
+
 				for (int32 Frame = StartFrame; Frame < NumOutputFrames; ++Frame)
 				{
 					// If we've read our last buffer, we're done
@@ -2249,12 +2251,12 @@ namespace Audio
 					SourceInfo.bHasStarted = true;
 
 					// Update the PrevFrameIndex value for the source based on alpha value
-					if (SourceInfo.CurrentFrameAlpha >= 1.0f)
+					if (CurrentAlpha >= 1.0f)
 					{
 						// Our inter-frame alpha lerping value is causing us to read new source frames
 						bReadNextSample = true;
 						
-						const float Delta = FMath::FloorToFloat(SourceInfo.CurrentFrameAlpha);
+						const float Delta = FMath::FloorToFloat(CurrentAlpha);
 						const int DeltaInt = (int)Delta;
 
 						// Bump up the current frame index
@@ -2264,7 +2266,7 @@ namespace Audio
 						// CurrentFrameIndex can wrap for looping sounds so won't be accurate in that case
 						SourceInfo.NumFramesPlayed += DeltaInt;
 
-						SourceInfo.CurrentFrameAlpha -= Delta;
+						CurrentAlpha -= Delta;
 					}
 
 					// If our alpha parameter caused us to jump to a new source frame, we need
@@ -2281,7 +2283,6 @@ namespace Audio
 						{
 							const float CurrFrameValue = SourceInfo.CurrentFrameValues[Channel];
 							const float NextFrameValue = SourceInfo.NextFrameValues[Channel];
-							const float CurrentAlpha = SourceInfo.CurrentFrameAlpha;
 							PreDistanceAttenBufferPtr[SampleIndex++] = FMath::Lerp(CurrFrameValue, NextFrameValue, CurrentAlpha);
 						}
 					}
@@ -2291,7 +2292,6 @@ namespace Audio
 						{
 							const float CurrFrameValue = SourceInfo.CurrentFrameValues[Channel];
 							const float NextFrameValue = SourceInfo.NextFrameValues[Channel];
-							const float CurrentAlpha = SourceInfo.CurrentFrameAlpha;
 
 							const float CurrentSample = FMath::Lerp(CurrFrameValue, NextFrameValue, CurrentAlpha);
 
@@ -2301,9 +2301,11 @@ namespace Audio
 					}
 
 					const float CurrentPitchScale = SourceInfo.PitchSourceParam.Update();
-					SourceInfo.CurrentFrameAlpha += CurrentPitchScale;
+					CurrentAlpha += CurrentPitchScale;
 				}
 
+				SourceInfo.CurrentFrameAlpha = CurrentAlpha;
+
 				// After processing the frames, reset the pitch param
 				SourceInfo.PitchSourceParam.Reset();
 
