Interpolate audio clock in between buffer callbacks.

diff --git a/Engine/Source/Runtime/AudioMixer/Private/AudioMixerDevice.cpp b/Engine/Source/Runtime/AudioMixer/Private/AudioMixerDevice.cpp
--- a/Engine/Source/Runtime/AudioMixer/Private/AudioMixerDevice.cpp
+++ b/Engine/Source/Runtime/AudioMixer/Private/AudioMixerDevice.cpp
@@ -522,7 +522,7 @@ namespace Audio
 
 				AudioClock = 0.0;
 				AudioClockDelta = (double)OpenStreamParams.NumFrames / OpenStreamParams.SampleRate;
-
+				AudioClockTimingData.UpdateTime = 0.0;
 
 				PluginInitializationParams.NumSources = SourceManagerInitParams.NumSources;
 				PluginInitializationParams.SampleRate = SampleRate;
@@ -754,6 +754,11 @@ namespace Audio
 		return AudioClock;
 	}
 
+	double FMixerDevice::GetInterpolatedAudioClock() const
+	{
+		return AudioClockTimingData.GetInterpolatedAudioClock(AudioClock, AudioClockDelta);
+	}
+
 	FAudioEffectsManager* FMixerDevice::CreateEffectsManager()
 	{
 		return new FAudioMixerEffectsManager(this);
@@ -931,7 +936,7 @@ namespace Audio
 		}
 
 		// Update the audio clock
-		AudioClock += AudioClockDelta;
+		UpdateAudioClock();
 
 		// notify interested parties
 		NotifyAudioDevicePostRender(RenderInfo);
@@ -940,6 +945,12 @@ namespace Audio
 		return true;
 	}
 
+	void FMixerDevice::UpdateAudioClock()
+	{
+		AudioClock += AudioClockDelta;
+		AudioClockTimingData.UpdateTime = FPlatformTime::Seconds();
+	}
+
 	void FMixerDevice::OnAudioStreamShutdown()
 	{
 		// Make sure the source manager pumps any final commands on shutdown. These allow for cleaning up sources, interfacing with plugins, etc.
@@ -2922,4 +2933,20 @@ namespace Audio
 		return Audio::KickQueuedTasks(QueueId);
 	}
 
+	double FAudioClockTimingData::GetInterpolatedAudioClock(const double InAudioClock, const double InAudioClockDelta) const
+	{
+		if (UpdateTime > 0.0)
+		{
+			const double TargetClock = InAudioClock + InAudioClockDelta;
+			const double CurrentDeltaSeconds = FPlatformTime::Seconds() - UpdateTime;
+			const double Alpha = FMath::Clamp(CurrentDeltaSeconds / InAudioClockDelta, 0.0f, 1.0f);
+
+			const double InterpolatedClock = FMath::Lerp(InAudioClock, TargetClock, Alpha);
+
+			return InterpolatedClock;
+		}
+
+		// Fall back to quantized clock if no timing data is available
+		return InAudioClock;
+	}
 }
diff --git a/Engine/Source/Runtime/AudioMixer/Public/AudioMixerDevice.h b/Engine/Source/Runtime/AudioMixer/Public/AudioMixerDevice.h
--- a/Engine/Source/Runtime/AudioMixer/Public/AudioMixerDevice.h
+++ b/Engine/Source/Runtime/AudioMixer/Public/AudioMixerDevice.h
@@ -67,6 +67,15 @@ namespace Audio
 		{}
 	};	
 
+	/** Data used to interpolate the audio clock in between buffer callbacks */
+	struct FAudioClockTimingData
+	{
+		/** Time in secods of previous audio clock update */
+		double UpdateTime = 0.0;
+
+		/** Interpolates the given clock based on the amount of platform time that has passed since last update */
+		double GetInterpolatedAudioClock(const double InAudioClock, const double InAudioClockDelta) const;
+	};
 
 	// Deprecated, use ERequiredSubmixes above
 	namespace EMasterSubmixType
@@ -120,6 +129,7 @@ namespace Audio
 		AUDIOMIXER_API virtual void UpdateGameThread() override;
 		AUDIOMIXER_API virtual void UpdateHardware() override;
 		AUDIOMIXER_API virtual double GetAudioTime() const override;
+		AUDIOMIXER_API virtual double GetInterpolatedAudioClock() const override;
 		AUDIOMIXER_API virtual FAudioEffectsManager* CreateEffectsManager() override;
 		AUDIOMIXER_API virtual FSoundSource* CreateSoundSource() override;
 		AUDIOMIXER_API virtual FName GetRuntimeFormat(const USoundWave* SoundWave) const override;
@@ -410,6 +420,9 @@ namespace Audio
 		// Pumps the audio render thread command queue
 		void PumpCommandQueue();
 		void PumpGameThreadCommandQueue();
+
+		/** Updates the audio clock and the associated timing data */
+		void UpdateAudioClock();
 		
 		TArray<USoundSubmix*> RequiredSubmixes;
 		TArray<FMixerSubmixPtr> RequiredSubmixInstances;
@@ -438,6 +451,9 @@ namespace Audio
 		/** The time delta for each callback block. */
 		double AudioClockDelta;
 
+		/** The timing data used to interpolate the audio clock */
+		FAudioClockTimingData AudioClockTimingData;
+
 		/** What the previous master volume was. */
 		float PreviousPrimaryVolume;
 
diff --git a/Engine/Source/Runtime/Engine/Public/AudioDevice.h b/Engine/Source/Runtime/Engine/Public/AudioDevice.h
--- a/Engine/Source/Runtime/Engine/Public/AudioDevice.h
+++ b/Engine/Source/Runtime/Engine/Public/AudioDevice.h
@@ -1896,6 +1896,11 @@ public:
 	/** Returns the audio clock of the audio device. Not supported on all platforms. */
 	double GetAudioClock() const { return AudioClock; }
 
+	/** Returns the audio clock interploated between audio device callbacks to provide a smoothed value.
+	 *  Default implementation does not interpolate.
+	 */
+	virtual double GetInterpolatedAudioClock() const { return GetAudioClock(); }
+
 	ENGINE_API void AddVirtualLoop(const FAudioVirtualLoop& InVirtualLoop);
 
 	bool AreStartupSoundsPreCached() const { return bStartupSoundsPreCached; }
diff --git a/Engine/Source/Runtime/MovieScene/Private/MovieSceneTimeController.cpp b/Engine/Source/Runtime/MovieScene/Private/MovieSceneTimeController.cpp
--- a/Engine/Source/Runtime/MovieScene/Private/MovieSceneTimeController.cpp
+++ b/Engine/Source/Runtime/MovieScene/Private/MovieSceneTimeController.cpp
@@ -118,7 +118,7 @@ double FMovieSceneTimeController_PlatformClock::GetCurrentTime() const
 double FMovieSceneTimeController_AudioClock::GetCurrentTime() const
 {
 	FAudioDevice* AudioDevice = GEngine ? GEngine->GetMainAudioDevice().GetAudioDevice() : nullptr;
-	return AudioDevice ? AudioDevice->GetAudioClock() : FPlatformTime::Seconds();
+	return AudioDevice ? AudioDevice->GetInterpolatedAudioClock() : FPlatformTime::Seconds();
 }
 
 void FMovieSceneTimeController_RelativeTimecodeClock::OnStopPlaying(const FQualifiedFrameTime& InStopTime)
