Const correctness for sound effect submix and avoiding extra memcpy if there were no offects processed

diff --git a/Engine/Plugins/Runtime/Synthesis/Source/Synthesis/Private/SubmixEffectFilter.cpp b/Engine/Plugins/Runtime/Synthesis/Source/Synthesis/Private/SubmixEffectFilter.cpp
--- a/Engine/Plugins/Runtime/Synthesis/Source/Synthesis/Private/SubmixEffectFilter.cpp
+++ b/Engine/Plugins/Runtime/Synthesis/Source/Synthesis/Private/SubmixEffectFilter.cpp
@@ -52,7 +52,7 @@ void FSubmixEffectFilter::OnProcessAudio(const FSoundEffectSubmixInputData& InDa
 		InitFilter();
 	}
 
-	float* InAudioBuffer = InData.AudioBuffer->GetData();
+	const float* InAudioBuffer = InData.AudioBuffer->GetData();
 	float* OutAudioBuffer = OutData.AudioBuffer->GetData();
 	const int32 NumSamples = InData.AudioBuffer->Num();
 
diff --git a/Engine/Plugins/Runtime/Synthesis/Source/Synthesis/Private/SubmixEffectStereoDelay.cpp b/Engine/Plugins/Runtime/Synthesis/Source/Synthesis/Private/SubmixEffectStereoDelay.cpp
--- a/Engine/Plugins/Runtime/Synthesis/Source/Synthesis/Private/SubmixEffectStereoDelay.cpp
+++ b/Engine/Plugins/Runtime/Synthesis/Source/Synthesis/Private/SubmixEffectStereoDelay.cpp
@@ -30,7 +30,7 @@ void FSubmixEffectStereoDelay::OnProcessAudio(const FSoundEffectSubmixInputData&
 		ScratchStereoBuffer.Reset();
 		ScratchStereoBuffer.AddUninitialized(InData.NumFrames * 2);
 
-		float* InputBufferPtr = InData.AudioBuffer->GetData();
+		const float* InputBufferPtr = InData.AudioBuffer->GetData();
 
 		// Get the stereo data from the input buffer
 		for (int32 Frame = 0; Frame < InData.NumFrames; ++Frame)
@@ -60,7 +60,7 @@ void FSubmixEffectStereoDelay::OnProcessAudio(const FSoundEffectSubmixInputData&
 	}
 	else
 	{
-		float* InputBufferPtr = InData.AudioBuffer->GetData();
+		const float* InputBufferPtr = InData.AudioBuffer->GetData();
 		float* OutputBufferPtr = OutData.AudioBuffer->GetData();
 		int32 NumSamples = InData.NumChannels * InData.NumFrames;
 		DelayStereo.ProcessAudio(InputBufferPtr, NumSamples, OutputBufferPtr);
diff --git a/Engine/Plugins/Runtime/Synthesis/Source/Synthesis/Private/SubmixEffectStereoToQuad.cpp b/Engine/Plugins/Runtime/Synthesis/Source/Synthesis/Private/SubmixEffectStereoToQuad.cpp
--- a/Engine/Plugins/Runtime/Synthesis/Source/Synthesis/Private/SubmixEffectStereoToQuad.cpp
+++ b/Engine/Plugins/Runtime/Synthesis/Source/Synthesis/Private/SubmixEffectStereoToQuad.cpp
@@ -25,7 +25,7 @@ void FSubmixEffectStereoToQuad::OnProcessAudio(const FSoundEffectSubmixInputData
 
 	if (InData.NumChannels > 2)
 	{
-		float* InputBufferPtr = InData.AudioBuffer->GetData();
+		const float* InputBufferPtr = InData.AudioBuffer->GetData();
 		float* OutputBufferPtr = OutData.AudioBuffer->GetData();
 
 		// Channel offset in SMPTE channel order for side left and side right
diff --git a/Engine/Source/Runtime/AudioMixer/Private/AudioMixerSubmix.cpp b/Engine/Source/Runtime/AudioMixer/Private/AudioMixerSubmix.cpp
--- a/Engine/Source/Runtime/AudioMixer/Private/AudioMixerSubmix.cpp
+++ b/Engine/Source/Runtime/AudioMixer/Private/AudioMixerSubmix.cpp
@@ -1399,7 +1399,9 @@ namespace Audio
 					EffectChainOutputBuffer.Reset();
 					EffectChainOutputBuffer.AddZeroed(NumSamples);
 
-					bProcessedAnEffect |= GenerateEffectChainAudio(InputData, InputBuffer, FadeInfo.EffectChain, EffectChainOutputBuffer);
+					const bool bResult = GenerateEffectChainAudio(InputData, InputBuffer, FadeInfo.EffectChain, EffectChainOutputBuffer);
+					bProcessedAnEffect |= bResult;
+					const FAlignedFloatBuffer* OutBuffer = bResult ? &EffectChainOutputBuffer : &InputBuffer;
 
 					// Mix effect chain output into SubmixChainMixBuffer
 					float StartFadeVolume = FadeInfo.FadeVolume.GetValue();
@@ -1407,7 +1409,7 @@ namespace Audio
 					float EndFadeVolume = FadeInfo.FadeVolume.GetValue();
 
 					// Mix this effect chain with other effect chains.
-					ArrayMixIn(EffectChainOutputBuffer, SubmixChainMixBuffer, StartFadeVolume, EndFadeVolume);
+					ArrayMixIn(*OutBuffer, SubmixChainMixBuffer, StartFadeVolume, EndFadeVolume);
 				}
 
 				// If we processed any effects, write over the old input buffer vs mixing into it. This is basically the "wet channel" audio in a submix.
@@ -1640,8 +1642,6 @@ namespace Audio
 
 			// Check to see if we need to down-mix our audio before sending to the submix effect
 			const uint32 ChannelCountOverride = SubmixEffect->GetDesiredInputChannelCountOverride();
-			bool bCurrentEffectProcessedAudio = false;
-
 			if (ChannelCountOverride != INDEX_NONE && ChannelCountOverride != NumChannels)
 			{
 				// Perform the down-mix operation with the down-mixed scratch buffer
@@ -1650,17 +1650,15 @@ namespace Audio
 
 				InputData.NumChannels = ChannelCountOverride;
 				InputData.AudioBuffer = &DownmixedBuffer;
-				bCurrentEffectProcessedAudio = SubmixEffect->ProcessAudio(InputData, OutputData);
 			}
 			else
 			{
 				// If we're not down-mixing, then just pass in the current wet buffer and our channel count is the same as the output channel count
 				InputData.NumChannels = NumChannels;
 				InputData.AudioBuffer = &ScratchBuffer;
-				bCurrentEffectProcessedAudio = SubmixEffect->ProcessAudio(InputData, OutputData);
 			}
 
-			if (bCurrentEffectProcessedAudio)
+			if (SubmixEffect->ProcessAudio(InputData, OutputData))
 			{
 				// Mix in the dry signal directly
 				const float DryLevel = SubmixEffect->GetDryLevel();
@@ -1675,12 +1673,6 @@ namespace Audio
 			}
 		}
 
-		if (!bProcessedAnEffect)
-		{
-			// If no effects were processed pass through the input audio to the output buffer.
-			OutBuffer = InAudioBuffer;
-		}
-
 		return bProcessedAnEffect;
 	}
 
diff --git a/Engine/Source/Runtime/AudioMixer/Private/Effects/AudioMixerSubmixEffectEQ.cpp b/Engine/Source/Runtime/AudioMixer/Private/Effects/AudioMixerSubmixEffectEQ.cpp
--- a/Engine/Source/Runtime/AudioMixer/Private/Effects/AudioMixerSubmixEffectEQ.cpp
+++ b/Engine/Source/Runtime/AudioMixer/Private/Effects/AudioMixerSubmixEffectEQ.cpp
@@ -102,7 +102,7 @@ void FSubmixEffectSubmixEQ::OnProcessAudio(const FSoundEffectSubmixInputData& In
 	// Update parameters that may have been set from game thread
 	UpdateParameters(InData.NumChannels);
 
-	Audio::FAlignedFloatBuffer& InAudioBuffer = *InData.AudioBuffer;
+	const Audio::FAlignedFloatBuffer& InAudioBuffer = *InData.AudioBuffer;
 	Audio::FAlignedFloatBuffer& OutAudioBuffer = *OutData.AudioBuffer;
 
 	if (bEQSettingsSet && RenderThreadEQSettings.EQBands.Num() > 0)
diff --git a/Engine/Source/Runtime/Engine/Classes/Sound/SoundEffectSubmix.h b/Engine/Source/Runtime/Engine/Classes/Sound/SoundEffectSubmix.h
--- a/Engine/Source/Runtime/Engine/Classes/Sound/SoundEffectSubmix.h
+++ b/Engine/Source/Runtime/Engine/Classes/Sound/SoundEffectSubmix.h
@@ -58,7 +58,7 @@ struct FSoundEffectSubmixInputData
 	const TArray<FTransform>* ListenerTransforms;
 
 	/** The raw input audio buffer. Size is NumFrames * NumChannels */
-	Audio::FAlignedFloatBuffer* AudioBuffer;
+	const Audio::FAlignedFloatBuffer* AudioBuffer;
 
 	/** Sample accurate audio clock. */
 	double AudioClock;
