Check Submix Parents for Modulation Destinations

diff --git a/Engine/Source/Runtime/AudioMixer/Private/AudioMixerSource.cpp b/Engine/Source/Runtime/AudioMixer/Private/AudioMixerSource.cpp
--- a/Engine/Source/Runtime/AudioMixer/Private/AudioMixerSource.cpp
+++ b/Engine/Source/Runtime/AudioMixer/Private/AudioMixerSource.cpp
@@ -488,8 +488,7 @@ namespace Audio
 		, MixerDevice(static_cast<FMixerDevice*>(InAudioDevice))
 		, MixerBuffer(nullptr)
 		, MixerSourceVoice(nullptr)
-		, BinauralVolModulators(nullptr)
-		, BinauralWetModulators(nullptr)
+		, bBypassingSubmixModulation(false)
 		, bPreviousBusEnablement(false)
 		, bPreviousBaseSubmixEnablement(false)
 		, PreviousAzimuth(-1.0f)
@@ -752,10 +751,7 @@ namespace Audio
 							, *InWaveInstance->GetName());
 					}
 					
-					// Get the modulation info associated with the SubmixPtr here
-					// The audio never goes to this submix, but we can still use this to modulate the volume ourselves later
-					BinauralVolModulators = SubmixPtr->GetOutputVolumeDestination();
-					BinauralWetModulators = SubmixPtr->GetWetVolumeDestination();
+					bBypassingSubmixModulation = true;
 				}
 			}
 
@@ -1605,6 +1601,33 @@ namespace Audio
 		MixerSourceVoice->SetModPitch(ModPitchBase);
 	}
 
+	float FMixerSource::GetInheritedSubmixVolumeModulation() const
+	{
+		float SubmixModVolume = 1.0f;
+
+		FMixerSubmixWeakPtr CurrSubmixWeakPtr = MixerDevice->GetSubmixInstance(WaveInstance->SoundSubmix);
+		FMixerSubmixPtr CurrSubmixPtr = CurrSubmixWeakPtr.Pin();
+		// Check the submix and all its parents in the graph for active modulation
+		while (CurrSubmixPtr && CurrSubmixPtr->IsValid())
+		{
+			FModulationDestination* SubmixOutVolDest = CurrSubmixPtr->GetOutputVolumeDestination();
+			FModulationDestination* SubmixWetVolDest = CurrSubmixPtr->GetWetVolumeDestination();
+			if (SubmixOutVolDest && SubmixOutVolDest->IsActive())
+			{
+				SubmixModVolume *= SubmixOutVolDest->GetValue();
+			}
+			if (SubmixWetVolDest && SubmixWetVolDest->IsActive())
+			{
+				SubmixModVolume *= SubmixWetVolDest->GetValue();
+			}
+
+			CurrSubmixWeakPtr = CurrSubmixPtr->GetParent();
+			CurrSubmixPtr = CurrSubmixWeakPtr.Pin();
+		}
+
+		return SubmixModVolume;
+	}
+
 	void FMixerSource::UpdateVolume()
 	{
 		// TODO: investigate if occlusion should be split from raw distance attenuation
@@ -1622,13 +1645,9 @@ namespace Audio
 			CurrentVolume *= WaveInstance->GetDynamicVolume();
 
 			// 3. Submix Volume Modulation (this only happens if the asset is binaural and we're sending to an external submix)
-			if (BinauralVolModulators)
-			{
-				CurrentVolume *= BinauralVolModulators->GetValue();
-			}
-			if (BinauralWetModulators)
+			if (bBypassingSubmixModulation)
 			{
-				CurrentVolume *= BinauralWetModulators->GetValue();
+				CurrentVolume *= GetInheritedSubmixVolumeModulation();
 			}
 
 			// 4. Apply editor gain stage(s)
diff --git a/Engine/Source/Runtime/AudioMixer/Private/AudioMixerSource.h b/Engine/Source/Runtime/AudioMixer/Private/AudioMixerSource.h
--- a/Engine/Source/Runtime/AudioMixer/Private/AudioMixerSource.h
+++ b/Engine/Source/Runtime/AudioMixer/Private/AudioMixerSource.h
@@ -119,6 +119,9 @@ namespace Audio
 		/** Whether or not to use the source data override plugin */
 		bool UseSourceDataOverridePlugin() const;
 
+		/** Gets an accumulated volume value based on the Modulation Destination data of the WaveInstance's submix and all of the submix's ancestors */
+		float GetInheritedSubmixVolumeModulation() const;
+
 	private:
 
 		FMixerDevice* MixerDevice;
@@ -128,8 +131,7 @@ namespace Audio
 		IAudioLinkFactory::FAudioLinkSourcePushedSharedPtr AudioLink;
 
 		// These modulators are obtained from the submix and used only on binaural assets
-		FModulationDestination* BinauralVolModulators;
-		FModulationDestination* BinauralWetModulators;
+		bool bBypassingSubmixModulation;
 
 		uint32 bPreviousBusEnablement;
 		uint32 bPreviousBaseSubmixEnablement;
