Mixer patch fast connect

New patch connections to the mixer are established before the next mix.
Additional adjustments to bus reader and writer buffering to ensure enough data is in their patches.

Differences between the MetaSound and mixer sample rates are still not accounted for in the bus reader/writer nodes or elsewhere.

diff --git a/Engine/Plugins/Runtime/Metasound/Source/MetasoundEngine/Private/MetasoundAudioBusReaderNode.cpp b/Engine/Plugins/Runtime/Metasound/Source/MetasoundEngine/Private/MetasoundAudioBusReaderNode.cpp
--- a/Engine/Plugins/Runtime/Metasound/Source/MetasoundEngine/Private/MetasoundAudioBusReaderNode.cpp
+++ b/Engine/Plugins/Runtime/Metasound/Source/MetasoundEngine/Private/MetasoundAudioBusReaderNode.cpp
@@ -278,8 +278,9 @@ namespace Metasound
 				Buffer->Zero();
 			}
 			
-			const int32 BlockSizeFrames = InParams.OperatorSettings.GetNumFramesPerBlock();
-			CreatePatchOutput(BlockSizeFrames);
+			// Wait until Execute to connect the AudioBusPatchOutput.
+			// Otherwise the mixer will start filling the patch before the MetaSound starts.
+			AudioBusPatchOutput.Reset();
 		}
 
 	private:
diff --git a/Engine/Plugins/Runtime/Metasound/Source/MetasoundEngine/Private/MetasoundAudioBusWriterNode.cpp b/Engine/Plugins/Runtime/Metasound/Source/MetasoundEngine/Private/MetasoundAudioBusWriterNode.cpp
--- a/Engine/Plugins/Runtime/Metasound/Source/MetasoundEngine/Private/MetasoundAudioBusWriterNode.cpp
+++ b/Engine/Plugins/Runtime/Metasound/Source/MetasoundEngine/Private/MetasoundAudioBusWriterNode.cpp
@@ -95,7 +95,17 @@ namespace Metasound
 					AudioInputs.Add(InputData.GetOrConstructDataReadReference<FAudioBuffer>(METASOUND_GET_PARAM_NAME_WITH_INDEX(InParamAudio, ChannelIndex), InParams.OperatorSettings));
 				}
 
-				return MakeUnique<TAudioBusWriterOperator<NumChannels>>(InParams, MoveTemp(AudioBusIn), MoveTemp(AudioInputs));
+				FString GraphName;
+				if (InParams.Environment.Contains<FString>(SourceInterface::Environment::GraphName))
+				{
+					GraphName = InParams.Environment.GetValue<FString>(SourceInterface::Environment::GraphName);
+				}
+				else
+				{
+					GraphName = TEXT("<Unknown>");
+				}
+
+				return MakeUnique<TAudioBusWriterOperator<NumChannels>>(InParams, MoveTemp(AudioBusIn), MoveTemp(AudioInputs), MoveTemp(GraphName));
 			}
 			else
 			{
@@ -105,9 +115,10 @@ namespace Metasound
 			}
 		}
 
-		TAudioBusWriterOperator(const FBuildOperatorParams& InParams, FAudioBusAssetReadRef InAudioBusAsset, TArray<FAudioBufferReadRef> InAudioInputs)
+		TAudioBusWriterOperator(const FBuildOperatorParams& InParams, FAudioBusAssetReadRef InAudioBusAsset, TArray<FAudioBufferReadRef> InAudioInputs, FString InGraphName)
 			: AudioBusAsset(MoveTemp(InAudioBusAsset))
 			, AudioInputs(MoveTemp(InAudioInputs))
+			, GraphName(MoveTemp(InGraphName))
 		{
 			Reset(InParams);
 		}
@@ -132,7 +143,7 @@ namespace Metasound
 						AudioBusSubsystem->StartAudioBus(AudioBusKey, AudioBusChannels, false);
 
 						InterleavedBuffer.Reset();
-						InterleavedBuffer.AddZeroed(NumBlocksToNumSamples(1));
+						InterleavedBuffer.Reserve(NumBlocksToNumSamples(InitialNumBlocks()));
 
 						// Create a bus patch input with enough room for the number of samples we expect and some buffering
 						AudioBusPatchInput = AudioBusSubsystem->AddPatchInputForAudioBus(AudioBusKey, BlockSizeFrames, AudioBusChannels);
@@ -165,6 +176,8 @@ namespace Metasound
 			
 			BlockSizeFrames = InParams.OperatorSettings.GetNumFramesPerBlock();
 
+			bWasUnderrunReported = false;
+
 			CreatePatchInput();
 		}
 
@@ -233,7 +246,9 @@ namespace Metasound
 			{
 				AudioInputBufferPtrs[ChannelIndex] = AudioInputs[ChannelIndex]->GetData();
 			}
-			float* InterleavedBufferPtr = InterleavedBuffer.GetData();
+			int32 InterleavedBufferOffset = InterleavedBuffer.Num();
+			InterleavedBuffer.SetNum(InterleavedBufferOffset + BlockSizeFrames * NumChannels);
+			float* InterleavedBufferPtr = InterleavedBuffer.GetData() + InterleavedBufferOffset;
 
 			if (AudioBusChannels == 1)
 			{
@@ -257,34 +272,42 @@ namespace Metasound
 				}
 			}
 
-			if (ConnectionState != EConnectionState::Connected)
+			switch (ConnectionState)
 			{
-				int32 InitialNumBlocks = 1;
-				if (AudioMixerOutputFrames != BlockSizeFrames)
-				{
-					InitialNumBlocks = FMath::DivideAndRoundUp(FMath::Max(AudioMixerOutputFrames, BlockSizeFrames), FMath::Min(AudioMixerOutputFrames, BlockSizeFrames));
-				}
-
-				if (ConnectionState == EConnectionState::Disconnected)
+			case EConnectionState::Disconnected:
 				{
-					if (InitialNumBlocks > 1)
+					// Wait until InterleavedBuffer contains enough samples to satisfy any mix eventualities!
+					// A single mix requires enough MetaSound executions to satisfy its buffer size.
+					// Mixing can occur concurrently with the first MetaSound execution intended for the next mix,
+					// and can steal that execution's patch output.
+					if (InterleavedBuffer.Num() < NumBlocksToNumSamples(InitialNumBlocks()))
 					{
-						// Ensure there will be enough samples in the patch input to support the maximum metasound executions the mixer requires to fill its output frames after the next push.
-						AudioBusPatchInput.PushAudio(nullptr, NumBlocksToNumSamples(InitialNumBlocks - 1));
+						return;
 					}
 
 					ConnectionState = EConnectionState::ConnectionPending;
+					break;
 				}
-				else
+
+			case EConnectionState::ConnectionPending:
 				{
+					int32 InitialNumSamples = NumBlocksToNumSamples(InitialNumBlocks());
+
 					// Determine if the pending connection has been established, by detecting if samples have been consumed.
-					if (AudioBusPatchInput.GetNumSamplesAvailable() == NumBlocksToNumSamples(InitialNumBlocks))
+					if (AudioBusPatchInput.GetNumSamplesAvailable() == InitialNumSamples)
 					{
-						UE_LOG(LogMetaSound, Verbose, TEXT("Writer node executed before mixer patch connection established."));
+						// If the connection hasn't been established by the time as many executions have occurred again,
+						// something has probably gone wrong.
+						if (InterleavedBuffer.Num() == InitialNumSamples)
+						{
+							UE_LOG(LogMetaSound, Warning, TEXT("Graph %s: Writer node executed before mixer patch connection established, with buffer size %d."), *GraphName, InterleavedBuffer.Num());
+							return;
+						}
 						return;
 					}
 
 					ConnectionState = EConnectionState::Connected;
+					break;
 				}
 			}
 
@@ -295,9 +318,24 @@ namespace Metasound
 				UE_LOG(LogMetaSound, Warning, TEXT("Underrun detected in audio bus writer node."));
 				bWasUnderrunReported = true;
 			}
+			InterleavedBuffer.Reset();
 		}
 
 	private:
+		int32 InitialNumBlocks() const
+		{
+			if (AudioMixerOutputFrames == BlockSizeFrames)
+			{
+				return 1;
+			}
+
+			// We need enough blocks for as many executions as the mixer can consume at once, plus one more,
+			// because the last execution can be concurrent with the mixer, and could contribute to either the
+			// current mix or the next.  That could leave us with one block too few if we didn't have an extra.
+			int32 MaxSizeFrames = FMath::Max(AudioMixerOutputFrames, BlockSizeFrames), MinSizeFrames = FMath::Min(AudioMixerOutputFrames, BlockSizeFrames);
+			return 1 + FMath::DivideAndRoundUp(MaxSizeFrames, MinSizeFrames);
+		}
+
 		int32 NumBlocksToNumSamples(int32 NumBlocks) const
 		{
 			return NumBlocks * BlockSizeFrames * AudioBusChannels;
@@ -305,6 +343,7 @@ namespace Metasound
 
 		FAudioBusAssetReadRef AudioBusAsset;
 		TArray<FAudioBufferReadRef> AudioInputs;
+		FString GraphName;
 
 		TArray<float> InterleavedBuffer;
 		int32 AudioMixerOutputFrames = INDEX_NONE;
diff --git a/Engine/Source/Runtime/AudioMixer/Private/AudioBusSubsystem.cpp b/Engine/Source/Runtime/AudioMixer/Private/AudioBusSubsystem.cpp
--- a/Engine/Source/Runtime/AudioMixer/Private/AudioBusSubsystem.cpp
+++ b/Engine/Source/Runtime/AudioMixer/Private/AudioBusSubsystem.cpp
@@ -115,89 +115,22 @@ bool UAudioBusSubsystem::IsAudioBusActive(Audio::FAudioBusKey InAudioBusKey) con
 
 Audio::FPatchInput UAudioBusSubsystem::AddPatchInputForAudioBus(Audio::FAudioBusKey InAudioBusKey, int32 InFrames, int32 InChannels, float InGain)
 {
-	using namespace Audio;
-	// This function is supporting adding audio bus patches from multiple threads (AT, ART, GT, and tasks) and is currently
-	// depending on a number of places where data lives, which accounts for the complexity here.
-	// This code needs a clean up to refactor everything into a true MPSC model, along with an MPSC refactor of the source manager
-	// and our command queues. Once we do that we can remove the code which branches based on the thread the request is coming from. 
-
 	Audio::FMixerSourceManager* SourceManager = GetMutableSourceManager();
 	check(SourceManager);
 
 	Audio::FMixerDevice* MixerDevice = GetMutableMixerDevice();
 	if (!MixerDevice)
 	{
-		return FPatchInput();
+		return Audio::FPatchInput();
 	}
 
-	if (IsInGameThread())
-	{
-		if (ActiveAudioBuses_GameThread.Find(InAudioBusKey))
-		{
-			FPatchInput PatchInput = MixerDevice->MakePatch(InFrames, InChannels, InGain);
-			FAudioThread::RunCommandOnAudioThread([this, PatchInput, InAudioBusKey]()
-			{
-				if (Audio::FMixerSourceManager* SourceManager = GetMutableSourceManager())
-				{
-					SourceManager->AddPatchInputForAudioBus_AudioThread(InAudioBusKey, PatchInput);
-				}
-			});
-			return PatchInput;
-		}
-		UE_LOG(LogAudioMixer, Warning, TEXT("Unable to add a patch output for audio bus because audio bus with object id '%u' and instance id '%u' is not active."), InAudioBusKey.ObjectId, InAudioBusKey.InstanceId);
-		return FPatchInput();
-	}
-	else if (IsInAudioThread())
-	{
-		FPatchInput PatchInput = MixerDevice->MakePatch(InFrames, InChannels, InGain);
-		SourceManager->AddPatchInputForAudioBus_AudioThread(InAudioBusKey, PatchInput);
-		return PatchInput;
-	}
-	else if (GetMixerDevice()->IsAudioRenderingThread())
-	{
-		check(SourceManager);
-
-		const int32 NumChannels = SourceManager->GetAudioBusNumChannels(InAudioBusKey);
-		if (NumChannels > 0)
-		{
-			FPatchInput PatchInput = MixerDevice->MakePatch(InFrames, InChannels, InGain);
-			SourceManager->AddPatchInputForAudioBus(InAudioBusKey, PatchInput);
-			return PatchInput;
-		}
-
-		return FPatchInput();
-	}
-	else
-	{
-		FPatchInput PatchInput = MixerDevice->MakePatch(InFrames, InChannels, InGain);
-		MixerDevice->GameThreadMPSCCommand([this, PatchInput, InAudioBusKey]() mutable
-		{
-			if (ActiveAudioBuses_GameThread.Find(InAudioBusKey))
-			{
-				FAudioThread::RunCommandOnAudioThread([this, InAudioBusKey, PatchInput = MoveTemp(PatchInput)]()
-				{
-					if (Audio::FMixerSourceManager* SourceManager = GetMutableSourceManager())
-					{
-						SourceManager->AddPatchInputForAudioBus_AudioThread(InAudioBusKey, PatchInput);
-					}
-				});
-			}
-			else
-			{
-				UE_LOG(LogAudioMixer, Warning, TEXT("Unable to add a patch output for audio bus because audio bus with object id '%u' and instance id '%u' is not active."), InAudioBusKey.ObjectId, InAudioBusKey.InstanceId);
-			}
-		});
-		return PatchInput;
-	}
+	Audio::FPatchInput PatchInput = MixerDevice->MakePatch(InFrames, InChannels, InGain);
+	SourceManager->AddPendingAudioBusConnection(InAudioBusKey, InChannels, false, PatchInput);
+	return PatchInput;
 }
 
 Audio::FPatchOutputStrongPtr UAudioBusSubsystem::AddPatchOutputForAudioBus(Audio::FAudioBusKey InAudioBusKey, int32 InFrames, int32 InChannels, float InGain)
 {
-	// This function is supporting adding audio bus patches from multiple threads (AT, ART, GT, and tasks) and is currently
-	// depending on a number of places where data lives, which accounts for the complexity here.
-	// This code needs a clean up to refactor everything into a true MPSC model, along with an MPSC refactor of the source manager
-	// and our command queues. Once we do that we can remove the code which branches based on the thread the request is coming from. 
-
 	Audio::FMixerSourceManager* SourceManager = GetMutableSourceManager();
 	check(SourceManager);
 
@@ -207,65 +140,9 @@ Audio::FPatchOutputStrongPtr UAudioBusSubsystem::AddPatchOutputForAudioBus(Audio
 		return nullptr;
 	}
 
-	if (IsInGameThread())
-	{
-		if (ActiveAudioBuses_GameThread.Find(InAudioBusKey))
-		{
-			Audio::FPatchOutputStrongPtr PatchOutput = MixerDevice->MakePatch(InFrames, InChannels, InGain);
-			FAudioThread::RunCommandOnAudioThread([this, PatchOutput, InAudioBusKey]()
-			{
-				if (Audio::FMixerSourceManager* SourceManager = GetMutableSourceManager())
-				{
-					SourceManager->AddPatchOutputForAudioBus_AudioThread(InAudioBusKey, PatchOutput);
-				}
-			});
-			return PatchOutput;
-		}
-		UE_LOG(LogAudioMixer, Warning, TEXT("Unable to add a patch output for audio bus because audio bus with object id '%u' and instance id '%u' is not active."), InAudioBusKey.ObjectId, InAudioBusKey.InstanceId);
-		return nullptr;
-	}
-	else if (IsInAudioThread())
-	{
-		Audio::FPatchOutputStrongPtr PatchOutput = MixerDevice->MakePatch(InFrames, InChannels, InGain);
-		SourceManager->AddPatchOutputForAudioBus_AudioThread(InAudioBusKey, PatchOutput);
-		return PatchOutput;
-	}
-	else if (MixerDevice->IsAudioRenderingThread())
-	{
-		check(SourceManager);
-
-		const int32 NumChannels = SourceManager->GetAudioBusNumChannels(InAudioBusKey);
-		if (NumChannels > 0)
-		{
-			Audio::FPatchOutputStrongPtr PatchOutput = MixerDevice->MakePatch(InFrames, InChannels, InGain);
-			SourceManager->AddPatchOutputForAudioBus(InAudioBusKey, PatchOutput);
-			return PatchOutput;
-		}
-
-		return nullptr;
-	}
-	else
-	{
-		Audio::FPatchOutputStrongPtr PatchOutput = MixerDevice->MakePatch(InFrames, InChannels, InGain);
-		MixerDevice->GameThreadMPSCCommand([this, PatchOutput, InAudioBusKey]() mutable
-		{
-			if (ActiveAudioBuses_GameThread.Find(InAudioBusKey))
-			{
-				FAudioThread::RunCommandOnAudioThread([this, InAudioBusKey, PatchOutput = MoveTemp(PatchOutput)]()
-				{
-					if (Audio::FMixerSourceManager* SourceManager = GetMutableSourceManager())
-					{
-						SourceManager->AddPatchOutputForAudioBus_AudioThread(InAudioBusKey, PatchOutput);
-					}
-				});
-			}
-			else
-			{
-				UE_LOG(LogAudioMixer, Warning, TEXT("Unable to add a patch output for audio bus because audio bus with object id '%u' and instance id '%u' is not active."), InAudioBusKey.ObjectId, InAudioBusKey.InstanceId);
-			}
-		});
-		return PatchOutput;
-	}
+	Audio::FPatchOutputStrongPtr PatchOutput = MixerDevice->MakePatch(InFrames, InChannels, InGain);
+	SourceManager->AddPendingAudioBusConnection(InAudioBusKey, InChannels, false, PatchOutput);
+	return PatchOutput;
 }
 
 void UAudioBusSubsystem::InitDefaultAudioBuses()
diff --git a/Engine/Source/Runtime/AudioMixer/Private/AudioMixerDevice.cpp b/Engine/Source/Runtime/AudioMixer/Private/AudioMixerDevice.cpp
--- a/Engine/Source/Runtime/AudioMixer/Private/AudioMixerDevice.cpp
+++ b/Engine/Source/Runtime/AudioMixer/Private/AudioMixerDevice.cpp
@@ -99,7 +99,7 @@ FAutoConsoleVariableRef CVarDebugGeneratorFreq(
 	TEXT("0: Not Disabled, 1: SinTone, 2: WhiteNoise"),
 	ECVF_Default);
 
-static int32 AudioMixerPatchBufferBlocks = 2;
+static int32 AudioMixerPatchBufferBlocks = 3;
 FAutoConsoleVariableRef CVarAudioMixerPatchBufferBlocks(
 	TEXT("au.PatchBufferBlocks"),
 	AudioMixerPatchBufferBlocks,
diff --git a/Engine/Source/Runtime/AudioMixer/Private/AudioMixerSourceManager.cpp b/Engine/Source/Runtime/AudioMixer/Private/AudioMixerSourceManager.cpp
--- a/Engine/Source/Runtime/AudioMixer/Private/AudioMixerSourceManager.cpp
+++ b/Engine/Source/Runtime/AudioMixer/Private/AudioMixerSourceManager.cpp
@@ -2317,10 +2317,47 @@ namespace Audio
 		}
 	}
 
+	void FMixerSourceManager::ConnectBusPatches()
+	{
+		while (TOptional<FPendingAudioBusConnection> PendingAudioBusConnection = PendingAudioBusConnections.Dequeue())
+		{
+			FAudioBusKey& AudioBusKey = PendingAudioBusConnection->AudioBusKey;
+			int32 NumChannels = PendingAudioBusConnection->NumChannels;
+			bool bIsAutomatic = PendingAudioBusConnection->bIsAutomatic;
+
+			// If this audio bus id already exists, set it to not be automatic and return it
+			TSharedPtr<FMixerAudioBus> AudioBusPtr = AudioBuses.FindRef(AudioBusKey);
+			if (AudioBusPtr.IsValid())
+			{
+				// If this audio bus already existed, make sure the num channels lines up
+				ensure(AudioBusPtr->GetNumChannels() == NumChannels);
+				AudioBusPtr->SetAutomatic(bIsAutomatic);
+			}
+			else
+			{
+				// If the bus is not registered, make a new entry.
+				AudioBusPtr = TSharedPtr<FMixerAudioBus>(new FMixerAudioBus(this, bIsAutomatic, NumChannels));
+				AudioBuses.Add(AudioBusKey, AudioBusPtr);
+			}
+
+			switch (PendingAudioBusConnection->PatchVariant.GetIndex())
+			{
+			case FPendingAudioBusConnection::FPatchVariant::IndexOfType<FPatchInput>():
+				AudioBusPtr->AddNewPatchInput(PendingAudioBusConnection->PatchVariant.Get<FPatchInput>());
+				break;
+			case FPendingAudioBusConnection::FPatchVariant::IndexOfType<FPatchOutputStrongPtr>():
+				AudioBusPtr->AddNewPatchOutput(PendingAudioBusConnection->PatchVariant.Get<FPatchOutputStrongPtr>());
+				break;
+			}
+		}
+	}
+
 	void FMixerSourceManager::ComputeBuses()
 	{
 		RenderThreadPhase = ESourceManagerRenderThreadPhase::ComputeBusses;
-		
+
+		ConnectBusPatches();
+
 		// Loop through the bus registry and mix source audio
 		for (auto& Entry : AudioBuses)
 		{
diff --git a/Engine/Source/Runtime/AudioMixer/Private/AudioMixerSourceManager.h b/Engine/Source/Runtime/AudioMixer/Private/AudioMixerSourceManager.h
--- a/Engine/Source/Runtime/AudioMixer/Private/AudioMixerSourceManager.h
+++ b/Engine/Source/Runtime/AudioMixer/Private/AudioMixerSourceManager.h
@@ -310,6 +310,16 @@ namespace Audio
 		// Pushes a TFUnction command into an MPSC queue from an arbitrary thread to the audio render thread
 		void AudioMixerThreadMPSCCommand(TFunction<void()>&& InCommand, const char* InDebugString=nullptr);
 		
+		void AddPendingAudioBusConnection(FAudioBusKey AudioBusKey, int32 NumChannels, bool bIsAutomatic, FPatchInput PatchInput)
+		{
+			PendingAudioBusConnections.Enqueue(FPendingAudioBusConnection{ FPendingAudioBusConnection::FPatchVariant(TInPlaceType<FPatchInput>(), MoveTemp(PatchInput)), MoveTemp(AudioBusKey), NumChannels, bIsAutomatic });
+		}
+
+		void AddPendingAudioBusConnection(FAudioBusKey AudioBusKey, int32 NumChannels, bool bIsAutomatic, FPatchOutputStrongPtr PatchOutputStrongPtr)
+		{
+			PendingAudioBusConnections.Enqueue(FPendingAudioBusConnection{ FPendingAudioBusConnection::FPatchVariant(TInPlaceType<FPatchOutputStrongPtr>(), MoveTemp(PatchOutputStrongPtr)), MoveTemp(AudioBusKey), NumChannels, bIsAutomatic });
+		}
+
 	private:
 #define INVALID_AUDIO_RENDER_THREAD_ID static_cast<uint32>(-1)
 		uint32 AudioRenderThreadId = INVALID_AUDIO_RENDER_THREAD_ID;
@@ -325,6 +335,7 @@ namespace Audio
 		void ComputePostSourceEffectBufferForIdRange(const bool bGenerateBuses, const int32 SourceIdStart, const int32 SourceIdEnd);
 		void ComputeOutputBuffersForIdRange(const bool bGenerateBuses, const int32 SourceIdStart, const int32 SourceIdEnd);
 
+		void ConnectBusPatches();
 		void ComputeBuses();
 		void UpdateBuses();
 
@@ -645,6 +656,17 @@ namespace Audio
 		FRWLock CurrentlyExecutingCmdLock;						// R/W slim lock for the currently executing cmd, so we can safely query it.
 		FAudioMixerThreadCommand CurrentlyExecuteingCmd;		// Keep this as a member so we can't always peek the executing cmd.
 
+		struct FPendingAudioBusConnection
+		{
+			using FPatchVariant = TVariant<FPatchInput, FPatchOutputStrongPtr>;
+			FPatchVariant PatchVariant;
+			FAudioBusKey AudioBusKey;
+			int32 NumChannels = 0;
+			bool bIsAutomatic = false;
+		};
+
+		TMpscQueue<FPendingAudioBusConnection> PendingAudioBusConnections;
+
 		friend class FMixerSourceVoice;
 	};
 }
