Protect against duplicate calls to register submix buffer listener.
#rnx

diff --git a/Engine/Source/Runtime/AudioMixer/Private/AudioMixerSubmix.cpp b/Engine/Source/Runtime/AudioMixer/Private/AudioMixerSubmix.cpp
--- a/Engine/Source/Runtime/AudioMixer/Private/AudioMixerSubmix.cpp
+++ b/Engine/Source/Runtime/AudioMixer/Private/AudioMixerSubmix.cpp
@@ -2122,13 +2122,13 @@ namespace Audio
 	{
 		FScopeLock Lock(&BufferListenerCriticalSection);
 		check(BufferListener);
-		BufferListeners.Add(BufferListener->AsShared());
+		BufferListeners.AddUnique(BufferListener->AsShared());
 	}
 
 	void FMixerSubmix::RegisterBufferListener(TSharedRef<ISubmixBufferListener, ESPMode::ThreadSafe> BufferListener)
 	{
 		FScopeLock Lock(&BufferListenerCriticalSection);
-		BufferListeners.Add(BufferListener);
+		BufferListeners.AddUnique(BufferListener);
 	}
 
 	void FMixerSubmix::UnregisterBufferListener(ISubmixBufferListener* BufferListener)
