Fix for DynamicSubmix's attaching to Master by accident and also submixes Auto registering AudioDevice Creation

diff --git a/Engine/Source/Runtime/AudioMixer/Private/AudioMixerDevice.cpp b/Engine/Source/Runtime/AudioMixer/Private/AudioMixerDevice.cpp
--- a/Engine/Source/Runtime/AudioMixer/Private/AudioMixerDevice.cpp
+++ b/Engine/Source/Runtime/AudioMixer/Private/AudioMixerDevice.cpp
@@ -182,8 +182,17 @@ namespace Audio
 		}
 	}
 
+	TSet<FSubmixMap::FObjectId> FSubmixMap::GetKeys() const
+	{
+		FScopeLock ScopeLock(&MutationLock);
+		TArray<FObjectId> Keys;
+		SubmixMap.GenerateKeyArray(Keys);
+		return TSet<FObjectId>(Keys);
+	}
+
+
 	static FAutoConsoleCommand DumpSubmixCmd(
-		TEXT("au.dumpsubmixes"),
+		TEXT("au.submix.drawgraph"),
 		TEXT("Draws the submix heirarchy for this world to the debug output"),
 		FConsoleCommandWithWorldArgsAndOutputDeviceDelegate::CreateLambda([](const TArray<FString>& InArgs, UWorld* InWorld, FOutputDevice& OutLog)
 		{
@@ -191,7 +200,7 @@ namespace Audio
 			{
 				if (const FMixerDevice* MixerDevice = static_cast<FMixerDevice*>(InWorld->GetAudioDeviceRaw()))
 				{
-					MixerDevice->DrawSubmixes(OutLog);
+					MixerDevice->DrawSubmixes(OutLog, InArgs);
 				}
 			}
 		})
@@ -228,14 +237,65 @@ namespace Audio
 		}
 	}
 
-	void FMixerDevice::DrawSubmixes(FOutputDevice& Output) const
+	static void DrawSubmixInstances(const FMixerSubmix* InRoot, const int32 InIdent, FOutputDevice& InOutput)
+	{
+		if (!InRoot)
+		{
+			return;
+		}
+
+		const FString Indent = FCString::Spc(InIdent*3);
+		InOutput.Logf(TEXT("%sName=%s,Instance=0x%p,Id=%u"),
+			*Indent, *InRoot->GetName(), InRoot, InRoot->GetId());
+
+		// Go downwards.
+		const TMap<uint32, FChildSubmixInfo>& Children = InRoot->GetChildren();
+		for (const auto& i : Children)
+		{
+			if (FMixerSubmixPtr Child =  i.Value.SubmixPtr.Pin())
+			{
+				DrawSubmixInstances(Child.Get(), InIdent+1, InOutput);
+			}
+		}
+	}
+	
+	void FMixerDevice::DrawSubmixes(FOutputDevice& InOutput, const TArray<FString>& InArgs) const
 	{
-		Output.Logf(TEXT("AudioDevice=%d, Device Instance=0x%p"), DeviceID, this);		
+		InOutput.Logf(TEXT("AudioDevice=%d, Device Instance=0x%p"), DeviceID, this);
 
+		// Params.
+		if (Algo::FindByPredicate(InArgs, [](const FString& InStr){ return FParse::Param(*InStr, TEXT("Instances")); }) != nullptr)
+		{
+			InOutput.Logf(TEXT("[Instance Hierarchy]"));
+			for (int32 i = 0; i < RequiredSubmixes.Num(); i++)
+			{
+				InOutput.Logf(TEXT("SlotName=[%s]"), ToCStr(LexToString(static_cast<ERequiredSubmixes>(i))));
+				DrawSubmixInstances(RequiredSubmixInstances[i].Get(), 1, InOutput);
+			}
+		}
+		if (Algo::FindByPredicate(InArgs, [](const FString& InStr){ return FParse::Param(*InStr, TEXT("Map")); }) != nullptr)
+		{
+			InOutput.Logf(TEXT("[Map of UObject -> SubmixPtrs]"));
+
+			// Map loop of unique ids from UObjects.
+			TMap<uint32, USoundSubmixBase*> AllSubmixes;
+			for (TObjectIterator<USoundSubmixBase> It; It; ++It)
+			{
+				AllSubmixes.Add(It->GetUniqueID(),*It);
+			}
+			for (const auto i : Submixes.GetKeys())
+			{
+				const auto pFound = AllSubmixes.Find(i);
+				InOutput.Logf(TEXT("%u -> %s"), i, pFound ? *GetNameSafe(*pFound) : TEXT("Not found"));
+			}
+		}
+
+		// USubmixMap hierarchy from slots downwards.
+		InOutput.Logf(TEXT("[Map of UObject Hierarchy]"));
 		for (int32 i = 0; i < RequiredSubmixes.Num(); i++)
 		{
-			Output.Logf(TEXT("SlotName=[%s]"), ToCStr(LexToString((ERequiredSubmixes)i)));
-			DrawSubmixHeirarchy(RequiredSubmixes[i], RequiredSubmixInstances[i], this, 1, Output, TEXT("In Slot"));
+			InOutput.Logf(TEXT("SlotName=[%s]"), ToCStr(LexToString(static_cast<ERequiredSubmixes>(i))));
+			DrawSubmixHeirarchy(RequiredSubmixes[i], RequiredSubmixInstances[i], this, 1, InOutput, TEXT("In Slot"));
 		}
 	}
 
@@ -607,7 +667,7 @@ namespace Audio
 		{
 			for (TObjectIterator<USoundSubmix> It; It; ++It)
 			{
-				UnregisterSoundSubmix(*It);
+				UnregisterSoundSubmix(*It, true);
 			}
 			
 			// Destroy the synchronized Audio Task Queue for this device
@@ -1139,7 +1199,7 @@ namespace Audio
 				USoundSubmixBase* SubmixToLoad = *It;
 				check(SubmixToLoad);
 
-				if (!IsRequiredSubmixType(SubmixToLoad))
+				if (!IsRequiredSubmixType(SubmixToLoad) && SubmixToLoad->bAutoRegister )
 				{
 					LoadSoundSubmix(*SubmixToLoad);
 					InitSoundfieldAndEndpointDataForSubmix(*SubmixToLoad, GetSubmixInstance(SubmixToLoad).Pin(), false);
@@ -1863,7 +1923,7 @@ namespace Audio
 		}
 	}
 
-	void FMixerDevice::UnregisterSoundSubmix(const USoundSubmixBase* InSoundSubmix)
+	void FMixerDevice::UnregisterSoundSubmix(const USoundSubmixBase* InSoundSubmix, const bool bReparentChildren)
 	{
 		if (!InSoundSubmix || bSubmixRegistrationDisabled || IsRequiredSubmixType(InSoundSubmix))
 		{
@@ -1875,21 +1935,21 @@ namespace Audio
 			DECLARE_CYCLE_STAT(TEXT("FAudioThreadTask.UnregisterSoundSubmix"), STAT_AudioUnregisterSoundSubmix, STATGROUP_AudioThreadCommands);
 
 			const TWeakObjectPtr<const USoundSubmixBase> SubmixToUnload = InSoundSubmix;
-			FAudioThread::RunCommandOnAudioThread([this, SubmixToUnload]()
+			FAudioThread::RunCommandOnAudioThread([this, SubmixToUnload, bReparentChildren]()
 			{
 				CSV_SCOPED_TIMING_STAT(Audio, UnregisterSubmix);
 				if (SubmixToUnload.IsValid())
 				{
-					UnloadSoundSubmix(*SubmixToUnload.Get());
+					UnloadSoundSubmix(*SubmixToUnload.Get(), bReparentChildren);
 				}
 			}, GET_STATID(STAT_AudioUnregisterSoundSubmix));
 			return;
 		}
 
-		UnloadSoundSubmix(*InSoundSubmix);
+		UnloadSoundSubmix(*InSoundSubmix, bReparentChildren);
 	}
 
-	void FMixerDevice::UnloadSoundSubmix(const USoundSubmixBase& InSoundSubmix)
+	void FMixerDevice::UnloadSoundSubmix(const USoundSubmixBase& InSoundSubmix, const bool bReparentChildren)
 	{
 		check(IsInAudioThread());
 
@@ -1909,14 +1969,17 @@ namespace Audio
 			ParentSubmixInstance->RemoveChildSubmix(GetSubmixInstance(&InSoundSubmix));
 		}
 
-		for (USoundSubmixBase* ChildSubmix : InSoundSubmix.ChildSubmixes)
+		if (bReparentChildren)
 		{
-			FMixerSubmixPtr ChildSubmixPtr = GetSubmixInstance(ChildSubmix).Pin();
-			if (ChildSubmixPtr.IsValid())
+			for (USoundSubmixBase* ChildSubmix : InSoundSubmix.ChildSubmixes)
 			{
-				ChildSubmixPtr->SetParentSubmix(ParentSubmixInstance.IsValid()
-					? ParentSubmixInstance
-					: MainSubmix);
+				FMixerSubmixPtr ChildSubmixPtr = GetSubmixInstance(ChildSubmix).Pin();
+				if (ChildSubmixPtr.IsValid())
+				{
+					ChildSubmixPtr->SetParentSubmix(ParentSubmixInstance.IsValid()
+						? ParentSubmixInstance
+						: MainSubmix);
+				}
 			}
 		}
 
diff --git a/Engine/Source/Runtime/AudioMixer/Private/AudioMixerSource.cpp b/Engine/Source/Runtime/AudioMixer/Private/AudioMixerSource.cpp
--- a/Engine/Source/Runtime/AudioMixer/Private/AudioMixerSource.cpp
+++ b/Engine/Source/Runtime/AudioMixer/Private/AudioMixerSource.cpp
@@ -1724,8 +1724,24 @@ namespace Audio
 			MixerSourceVoice->SetSubmixSendInfo(MixerDevice->GetMasterReverbSubmix(), WaveInstance->ReverbSendLevel);
 		}
 
+		// Safely track if the submix has changed between updates.
+		bool bSubmixHasChanged = false;
+		TObjectKey<USoundSubmixBase> SubmixKey(WaveInstance->SoundSubmix);
+		if (SubmixKey != PrevousSubmix )
+		{
+			bSubmixHasChanged = true;
+		}
+
+		// This will reattempt to resolve a submix each update if there's a valid input
+		if ((!WaveInstance->SoundSubmix && PreviousSubmixResolved.IsValid()) || 
+		     (WaveInstance->SoundSubmix && !PreviousSubmixResolved.IsValid()) )
+		{
+			bSubmixHasChanged = true;
+		}
+
 		//Check whether the base submix send has been enabled or disabled since the last update
-		if (WaveInstance->bEnableBaseSubmix != bPreviousBaseSubmixEnablement)
+		//Or if the submix has now been registered with the world.
+		if (WaveInstance->bEnableBaseSubmix != bPreviousBaseSubmixEnablement || bSubmixHasChanged)
 		{
 			// set the level for this send
 			FMixerSubmixWeakPtr SubmixPtr;
@@ -1733,13 +1749,16 @@ namespace Audio
 			{
 				SubmixPtr = MixerDevice->GetSubmixInstance(WaveInstance->SoundSubmix);
 			}
-			else
+			else if(WaveInstance->SoundSubmix && WaveInstance->SoundSubmix->bAutoRouteToMasterSubmixWhenOrphaned)
 			{
 				SubmixPtr = MixerDevice->GetMasterSubmix();
 			}
 
+
 			MixerSourceVoice->SetSubmixSendInfo(SubmixPtr, WaveInstance->bEnableBaseSubmix);
 			bPreviousBaseSubmixEnablement = WaveInstance->bEnableBaseSubmix;
+			PreviousSubmixResolved = SubmixPtr;
+			PrevousSubmix = SubmixKey;
 		}
 
 		if (WaveInstance->SubmixSendSettings.Num() > 0)
diff --git a/Engine/Source/Runtime/AudioMixer/Private/AudioMixerSource.h b/Engine/Source/Runtime/AudioMixer/Private/AudioMixerSource.h
--- a/Engine/Source/Runtime/AudioMixer/Private/AudioMixerSource.h
+++ b/Engine/Source/Runtime/AudioMixer/Private/AudioMixerSource.h
@@ -129,6 +129,8 @@ namespace Audio
 		TSharedPtr<FMixerSourceBuffer, ESPMode::ThreadSafe> MixerSourceBuffer;
 		FMixerSourceVoice* MixerSourceVoice;
 		IAudioLinkFactory::FAudioLinkSourcePushedSharedPtr AudioLink;
+		FMixerSubmixWeakPtr PreviousSubmixResolved;
+		TObjectKey<USoundSubmixBase> PrevousSubmix;
 
 		// These modulators are obtained from the submix and used only on binaural assets
 		bool bBypassingSubmixModulation;
diff --git a/Engine/Source/Runtime/AudioMixer/Public/AudioMixerDevice.h b/Engine/Source/Runtime/AudioMixer/Public/AudioMixerDevice.h
--- a/Engine/Source/Runtime/AudioMixer/Public/AudioMixerDevice.h
+++ b/Engine/Source/Runtime/AudioMixer/Public/AudioMixerDevice.h
@@ -102,7 +102,7 @@ namespace Audio
 		FMixerSubmixPtr FindRef(FObjectId InObjectId) const;
 		int32 Remove(const FObjectId InObjectId);
 		void Reset();
-
+		TSet<FSubmixMap::FObjectId> GetKeys() const;
 	private:
 		TMap<FObjectId, FMixerSubmixPtr> SubmixMap;
 
@@ -149,7 +149,7 @@ namespace Audio
 		AUDIOMIXER_API virtual void EnableDebugAudioOutput() override;
 		AUDIOMIXER_API virtual FAudioPlatformSettings GetPlatformSettings() const override;
 		AUDIOMIXER_API virtual void RegisterSoundSubmix(USoundSubmixBase* SoundSubmix, bool bInit = true) override;
-		AUDIOMIXER_API virtual void UnregisterSoundSubmix(const USoundSubmixBase* SoundSubmix) override;
+		AUDIOMIXER_API virtual void UnregisterSoundSubmix(const USoundSubmixBase* SoundSubmix, const bool bReparentChildren) override;
 
 		AUDIOMIXER_API virtual void InitSoundEffectPresets() override;
 		UE_DEPRECATED(5.2, "The functionality for this has been moved to UAudioBusSubsystem::InitDefaultAudioBuses, which is now automatically called on subsystem creation.")
@@ -379,7 +379,9 @@ namespace Audio
 		// Pushes the command to a MPSC queue to be executed on the game thread
 		AUDIOMIXER_API void GameThreadMPSCCommand(TFunction<void()> InCommand);
 
-		AUDIOMIXER_API void DrawSubmixes(FOutputDevice& Output) const;
+		// Debug Commands
+		AUDIOMIXER_API void DrawSubmixes(FOutputDevice& InOutput, const TArray<FString>& InArgs) const;
+
 	protected:
 		AUDIOMIXER_API virtual void InitSoundSubmixes() override;
 
@@ -409,7 +411,7 @@ namespace Audio
 
 		void InitSoundfieldAndEndpointDataForSubmix(const USoundSubmixBase& InSoundSubmix, FMixerSubmixPtr MixerSubmix, bool bAllowReInit);
 
-		void UnloadSoundSubmix(const USoundSubmixBase& SoundSubmix);
+		void UnloadSoundSubmix(const USoundSubmixBase& SoundSubmix, const bool bReparentChildren);
 
 		ICompressedAudioInfo* CreateAudioInfo(FName InFormat) const;
 
diff --git a/Engine/Source/Runtime/AudioMixer/Public/AudioMixerSubmix.h b/Engine/Source/Runtime/AudioMixer/Public/AudioMixerSubmix.h
--- a/Engine/Source/Runtime/AudioMixer/Public/AudioMixerSubmix.h
+++ b/Engine/Source/Runtime/AudioMixer/Public/AudioMixerSubmix.h
@@ -83,7 +83,7 @@ namespace Audio
 
 	class FMixerSubmix;
 
-	struct FChildSubmixInfo
+	struct FChildSubmixInfo : FNoncopyable
 	{
 		TWeakPtr<FMixerSubmix, ESPMode::ThreadSafe> SubmixPtr;
 
@@ -99,13 +99,6 @@ namespace Audio
 		FChildSubmixInfo()
 		{}
 
-		// TMap doesn't compile using non-copyable types as value types, though you can build and use a TMap without using the copy constructor.
-		FChildSubmixInfo(const FChildSubmixInfo& InInfo)
-			: FChildSubmixInfo()
-		{
-			checkf(false, TEXT("FChildSubmixInfo is not copyable. If you are using FChildSubmixInfo, consider using Emplace or Add(FChildSubmixInfo&&)."));
-		}
-
 		FChildSubmixInfo(TWeakPtr<FMixerSubmix, ESPMode::ThreadSafe> SubmixWeakPtr)
 			: SubmixPtr(SubmixWeakPtr)
 		{
@@ -124,6 +117,9 @@ namespace Audio
 		// Returns the mixer submix Id
 		uint32 GetId() const { return Id; }
 
+		// Return the owners name 
+		AUDIOMIXER_API const FString& GetName() const { return SubmixName; }
+
 		// Sets the parent submix to the given submix
 		AUDIOMIXER_API void SetParentSubmix(TWeakPtr<FMixerSubmix, ESPMode::ThreadSafe> Submix);
 
@@ -386,6 +382,9 @@ namespace Audio
 		// Generates audio from the given effect chain into the given buffer
 		AUDIOMIXER_API bool GenerateEffectChainAudio(FSoundEffectSubmixInputData& InputData, const FAlignedFloatBuffer& InAudioBuffer, TArray<FSoundEffectSubmixPtr>& InEffectChain, FAlignedFloatBuffer& OutBuffer);
 
+		// The name of this submix (the owning USoundSubmix) (at top so we can see in debugger it's name)
+		FString SubmixName;
+		
 		// This mixer submix's Id
 		uint32 Id;
 
@@ -619,9 +618,6 @@ namespace Audio
 		// The time that the first full silent buffer was detected in the submix. Submix will auto-disable if the timeout is reached and the submix has bAutoDisable set to true.
 		double SilenceTimeStartSeconds;
 
-		// The name of this submix (the owning USoundSubmix)
-		FString SubmixName;
-
 		// Bool set to true when envelope following is enabled
 		FThreadSafeBool bIsEnvelopeFollowing;
 
diff --git a/Engine/Source/Runtime/Engine/Classes/Sound/SoundSubmix.h b/Engine/Source/Runtime/Engine/Classes/Sound/SoundSubmix.h
--- a/Engine/Source/Runtime/Engine/Classes/Sound/SoundSubmix.h
+++ b/Engine/Source/Runtime/Engine/Classes/Sound/SoundSubmix.h
@@ -181,18 +181,26 @@ public:
 #endif
 
 	// Auto-manage enabling and disabling the submix as a CPU optimization. It will be disabled if the submix and all child submixes are silent. It will re-enable if a sound is sent to the submix or a child submix is audible.
-	UPROPERTY(EditAnywhere, BlueprintReadOnly, Category = AutoDisablement)
+	UPROPERTY(config, EditAnywhere, BlueprintReadOnly, Category = AutoDisablement)
 	bool bAutoDisable = true;
 
 	// The minimum amount of time to wait before automatically disabling a submix if it is silent. Will immediately re-enable if source audio is sent to it. 
-	UPROPERTY(EditAnywhere, BlueprintReadWrite, Category = AutoDisablement, meta = (EditCondition = "bAutoDisable"))
+	UPROPERTY(config, EditAnywhere, BlueprintReadWrite, Category = AutoDisablement, meta = (EditCondition = "bAutoDisable"))
 	float AutoDisableTime = 0.01f;
 
+	// Whether to automatically register this Submix with an AudioDevice. 
+	UPROPERTY(config, EditAnywhere, BlueprintReadOnly, Category = Automatic)
+	bool bAutoRegister = true;
+
+	// Whether to automatically route to the 'MasterSubmix' aka 'MainSubmix' when the parent is null
+	UPROPERTY(config, EditAnywhere, BlueprintReadOnly, Category = Automatic)
+	bool bAutoRouteToMasterSubmixWhenOrphaned = true;
+
 	// Child submixes to this sound mix
 	UPROPERTY(VisibleAnywhere, BlueprintReadOnly, Category = SoundSubmix)
 	TArray<TObjectPtr<USoundSubmixBase>> ChildSubmixes;
 
-	// Dynamic Child submixes
+	// Dynamic Child submixes (Map of AudioDevice -> [Submix] )
 	UPROPERTY(Transient)
 	TMap<uint32, FDynamicChildSubmix> DynamicChildSubmixes;
 		
diff --git a/Engine/Source/Runtime/Engine/Private/AudioDeviceManager.cpp b/Engine/Source/Runtime/Engine/Private/AudioDeviceManager.cpp
--- a/Engine/Source/Runtime/Engine/Private/AudioDeviceManager.cpp
+++ b/Engine/Source/Runtime/Engine/Private/AudioDeviceManager.cpp
@@ -825,7 +825,7 @@ void FAudioDeviceManager::UnregisterSoundSubmix(const USoundSubmixBase* SoundSub
 	IterateOverAllDevices(
 		[&SoundSubmix](Audio::FDeviceId, FAudioDevice* InDevice)
 		{
-			InDevice->UnregisterSoundSubmix(SoundSubmix);
+			InDevice->UnregisterSoundSubmix(SoundSubmix, true);
 		}
 	);
 }
diff --git a/Engine/Source/Runtime/Engine/Private/SoundSubmix.cpp b/Engine/Source/Runtime/Engine/Private/SoundSubmix.cpp
--- a/Engine/Source/Runtime/Engine/Private/SoundSubmix.cpp
+++ b/Engine/Source/Runtime/Engine/Private/SoundSubmix.cpp
@@ -677,6 +677,7 @@ bool USoundSubmixWithParentBase::DynamicConnect(FAudioDeviceHandle Handle, USoun
 
 	if (CurrentParent != InParent)
 	{
+
 		if (CurrentParent)
 		{
 			CurrentParent->DynamicChildSubmixes.FindOrAdd(Id).ChildSubmixes.Remove(this);
@@ -684,7 +685,6 @@ bool USoundSubmixWithParentBase::DynamicConnect(FAudioDeviceHandle Handle, USoun
 		}
 
 		CurrentParent = InParent;
-
 		if (CurrentParent)
 		{
 			CurrentParent->DynamicChildSubmixes.FindOrAdd(Id).ChildSubmixes.AddUnique(this);
@@ -694,8 +694,9 @@ bool USoundSubmixWithParentBase::DynamicConnect(FAudioDeviceHandle Handle, USoun
 			Handle->SetSubmixAutoDisable(Cast<USoundSubmix>(CurrentParent.Get()), false);
 		}
 
-		Handle->RegisterSoundSubmix(this, false);
-		Handle->SetSubmixAutoDisable(Cast<USoundSubmix>(this), CurrentParent != nullptr);
+		// Register us, and disable parents auto disable feature.
+		Handle->RegisterSoundSubmix(this, /*bInit*/ true);
+		Handle->SetSubmixAutoDisable(Cast<USoundSubmix>(this), CurrentParent == nullptr);
 
 		UE_LOG(LogAudio, Verbose, TEXT("Submix (DynamicConnect): Registering [%s] with AudioDevice [%u]"), *GetName(), Id);
 
@@ -745,7 +746,7 @@ bool USoundSubmixWithParentBase::DynamicDisconnect(FAudioDeviceHandle Handle)
 		CurrentParent->DynamicChildSubmixes.FindOrAdd(Id).ChildSubmixes.Remove(this);
 		CurrentParent = nullptr;
 
-		Handle->UnregisterSoundSubmix(this);
+		Handle->UnregisterSoundSubmix(this, false);
 				
 		// If we still have a valid parent static submix? Make sure that's still live and registered.
 		if (ParentSubmix)
@@ -813,6 +814,11 @@ void USoundSubmixBase::PostEditChangeProperty(FPropertyChangedEvent& PropertyCha
 						// Update parentage
 						SubmixWithParent->SetParentSubmix(this);
 					}
+					
+					// Make the children follow our auto register setting.
+					ChildSubmixes[ChildIndex]->bAutoRegister = bAutoRegister;
+					ChildSubmixes[ChildIndex]->PostEditChangeProperty(PropertyChangedEvent);
+
 					break;
 				}
 			}
@@ -833,7 +839,25 @@ void USoundSubmixBase::PostEditChangeProperty(FPropertyChangedEvent& PropertyCha
 			// Force the properties to be initialized for this SoundSubmix on all active audio devices
 			if (FAudioDeviceManager* AudioDeviceManager = GEngine->GetAudioDeviceManager())
 			{
-				AudioDeviceManager->RegisterSoundSubmix(this);
+				if (bAutoRegister)
+				{
+					AudioDeviceManager->RegisterSoundSubmix(this);
+				}
+			}
+		}
+
+		// Propagate auto register to children of this submix. 
+		// We don't want them to auto register 
+		if (PropertyChangedEvent.Property->GetFName() == GET_MEMBER_NAME_CHECKED(USoundSubmixBase, bAutoRegister))
+		{
+			// All children so should follow the 
+			for (int32 ChildIndex = 0; ChildIndex < ChildSubmixes.Num(); ChildIndex++)
+			{
+				if (USoundSubmixBase* Child = ChildSubmixes[ChildIndex] )
+				{
+					Child->bAutoRegister = bAutoRegister;
+					Child->PostEditChangeProperty(PropertyChangedEvent);
+				}
 			}
 		}
 	}
diff --git a/Engine/Source/Runtime/Engine/Public/AudioDevice.h b/Engine/Source/Runtime/Engine/Public/AudioDevice.h
--- a/Engine/Source/Runtime/Engine/Public/AudioDevice.h
+++ b/Engine/Source/Runtime/Engine/Public/AudioDevice.h
@@ -956,7 +956,7 @@ public:
 	virtual void RegisterSoundSubmix(USoundSubmixBase* SoundSubmix, bool bInit) {}
 
 	/** Unregisters the sound submix */
-	virtual void UnregisterSoundSubmix(const USoundSubmixBase* SoundSubmix) {}
+	virtual void UnregisterSoundSubmix(const USoundSubmixBase* SoundSubmix, const bool bReparentChildren) {}
 
 	ENGINE_API virtual USoundSubmix& GetMainSubmixObject() const;
 
