Make a couple modulation calls thread safe.
#rnx

diff --git a/Engine/Source/Runtime/AudioMixer/Private/AudioMixerDevice.cpp b/Engine/Source/Runtime/AudioMixer/Private/AudioMixerDevice.cpp
--- a/Engine/Source/Runtime/AudioMixer/Private/AudioMixerDevice.cpp
+++ b/Engine/Source/Runtime/AudioMixer/Private/AudioMixerDevice.cpp
@@ -1650,20 +1650,25 @@ namespace Audio
 
 	void FMixerDevice::UpdateSubmixModulationSettings(USoundSubmix* InSoundSubmix, const TSet<TObjectPtr<USoundModulatorBase>>& InOutputModulation, const TSet<TObjectPtr<USoundModulatorBase>>& InWetLevelModulation, const TSet<TObjectPtr<USoundModulatorBase>>& InDryLevelModulation)
 	{
+		TWeakObjectPtr<USoundSubmix> SubmixWeakPtr = InSoundSubmix;
 		if (!IsInAudioThread())
 		{
-			FMixerDevice* MixerDevice = this;
-
-			FAudioThread::RunCommandOnAudioThread([MixerDevice, InSoundSubmix, OutMod = InOutputModulation, WetMod = InWetLevelModulation, DryMod = InDryLevelModulation]()
+			FAudioThread::RunCommandOnAudioThread([ThisDeviceID = DeviceID, SubmixWeakPtr, OutMod = InOutputModulation, WetMod = InWetLevelModulation, DryMod = InDryLevelModulation]()
 			{
-				MixerDevice->UpdateSubmixModulationSettings(InSoundSubmix, OutMod, WetMod, DryMod);
+				if (USoundSubmix* Submix = SubmixWeakPtr.Get())
+				{
+					if (FAudioDevice* Device = FAudioDeviceManager::Get()->GetAudioDeviceRaw(ThisDeviceID))
+					{
+						FMixerDevice* ThisMixerDevice = static_cast<FMixerDevice*>(Device);
+						ThisMixerDevice->UpdateSubmixModulationSettings(Submix, OutMod, WetMod, DryMod);
+					}
+				}
 			});
 			return;
 		}
 
 		if (IsModulationPluginEnabled() && ModulationInterface.IsValid())
 		{
-			TWeakObjectPtr<USoundSubmix> SubmixWeakPtr = InSoundSubmix;
 			FMixerSubmixWeakPtr MixerSubmixWeakPtr = GetSubmixInstance(InSoundSubmix);
 
 			if (SubmixWeakPtr.IsValid())
@@ -1684,11 +1689,17 @@ namespace Audio
 	{
 		if (!IsInAudioThread())
 		{
-			FMixerDevice* MixerDevice = this;
-
-			FAudioThread::RunCommandOnAudioThread([MixerDevice, InSoundSubmix, InVolumeModBase, InWetModBase, InDryModBase]() 
+			TWeakObjectPtr<USoundSubmix> SubmixWeakPtr = InSoundSubmix;
+			FAudioThread::RunCommandOnAudioThread([ThisDeviceID = DeviceID, SubmixWeakPtr, InVolumeModBase, InWetModBase, InDryModBase]()
 			{
-				MixerDevice->SetSubmixModulationBaseLevels(InSoundSubmix, InVolumeModBase, InWetModBase, InDryModBase);
+				if (USoundSubmix* Submix = SubmixWeakPtr.Get())
+				{
+					if (FAudioDevice* Device = FAudioDeviceManager::Get()->GetAudioDeviceRaw(ThisDeviceID))
+					{
+						FMixerDevice* ThisMixerDevice = static_cast<FMixerDevice*>(Device);
+						ThisMixerDevice->SetSubmixModulationBaseLevels(Submix, InVolumeModBase, InWetModBase, InDryModBase);
+					}
+				}
 			});
 			return;
 		}
