Vectorized non silent audio check in FMixerSubmix::ProcessAudio

diff --git a/Engine/Source/Runtime/AudioMixer/Private/AudioMixerSubmix.cpp b/Engine/Source/Runtime/AudioMixer/Private/AudioMixerSubmix.cpp
--- a/Engine/Source/Runtime/AudioMixer/Private/AudioMixerSubmix.cpp
+++ b/Engine/Source/Runtime/AudioMixer/Private/AudioMixerSubmix.cpp
@@ -1563,10 +1563,25 @@ namespace Audio
 		{
 			// Extremely cheap silent buffer detection: as soon as we hit a sample which isn't silent, we flag we're not silent
 			bool bIsNowSilent = true;
-			for (float& SampleValue : OutAudioBuffer)
+			int i = 0;
+#if PLATFORM_ENABLE_VECTORINTRINSICS
+			for (; i < OutAudioBuffer.Num(); i += 4)
+			{
+				VectorRegister4Float Samples = VectorLoad(&OutAudioBuffer[i]);
+				if (VectorAnyGreaterThan(Samples, GlobalVectorConstants::SmallNumber))
+				{
+					bIsNowSilent = false;
+					i = INT_MAX;
+					break;
+				}
+			}
+#endif
+
+			// Finish to the end of the buffer or check each sample if vector intrinsics are disabled
+			for (; i < OutAudioBuffer.Num(); ++i)
 			{
 				// As soon as we hit a non-silent sample, we're not silent
-				if (SampleValue > SMALL_NUMBER)
+				if (OutAudioBuffer[i] > SMALL_NUMBER)
 				{
 					bIsNowSilent = false;
 					break;
