Prevent excessive log spam in AudioMixerSourceManager

Gate the error message on whether it has been detected already. Flip the gate when it has had the chance to recover.

#tests PIE

diff --git a/Engine/Source/Runtime/AudioMixer/Private/AudioMixerSourceManager.cpp b/Engine/Source/Runtime/AudioMixer/Private/AudioMixerSourceManager.cpp
--- a/Engine/Source/Runtime/AudioMixer/Private/AudioMixerSourceManager.cpp
+++ b/Engine/Source/Runtime/AudioMixer/Private/AudioMixerSourceManager.cpp
@@ -3384,16 +3384,23 @@ namespace Audio
 		}
 
 		// Before adding further commands, ensure we're not growing outside any sensible size for these buffers.
-		// On shipping builds, this will just stop us crashing from growing out of control and OOMing the machine.		
+		// On shipping builds, this will just stop us crashing from growing out of control and OOMing the machine.
+		static bool bCommandBufferOverflowDetected = false;
 		const SIZE_T MaxBufferSizeInBytes = ((SIZE_T)CommandBufferMaxSizeInMbCvar) << 20;
 		if (CurrentBufferSizeInBytes < MaxBufferSizeInBytes)
 		{
 			CommandBuffers[AudioThreadCommandIndex].SourceCommandQueue.Add(AudioCommand);
 			NumCommands.Increment();
+			if (bCommandBufferOverflowDetected)
+			{
+				UE_LOG(LogAudioMixer, Log, TEXT("Command buffer shrunk to %umb, allowing adds again."), CurrentBufferSizeInBytes >> 20);
+				bCommandBufferOverflowDetected = false;
+			}
 		}
-		else
+		else if (!bCommandBufferOverflowDetected)
 		{
-			UE_LOG(LogAudioMixer, Error, TEXT("Command buffer grown to %umb, preventing any more adds! Likely cause the AudioRenderer has hung"), CurrentBufferSizeInBytes >>20);
+			UE_LOG(LogAudioMixer, Error, TEXT("Command buffer grown to %umb, preventing any more adds! Likely cause the AudioRenderer has hung"), CurrentBufferSizeInBytes >> 20);
+			bCommandBufferOverflowDetected = true;
 		}
 	}
 
