Custom growth factor for Command Buffers so they don't grow at factors of 2x, resulting in 50% memory savings and better overall performance.

"Reserve" Command Buffers with reasonable sizes to start with (2000)
Start with a growth factor of about 20% each time we need to reallocate
Controllable via CVars.

With this growth strategy, even if the command buffer gets filled with 65k commands (via local Command Buffer spamming), the command buffer only reached 8mb (as opposed to 15mb), which means we're still below the "crisis" threshold.

Have one overflow detection counter per command buffer. Arguably, should just assert at this point. But technically, a recoverable state assuming we don't OOM before hand.

#tests PIE

diff --git a/Engine/Source/Runtime/AudioMixer/Private/AudioMixerSourceManager.cpp b/Engine/Source/Runtime/AudioMixer/Private/AudioMixerSourceManager.cpp
--- a/Engine/Source/Runtime/AudioMixer/Private/AudioMixerSourceManager.cpp
+++ b/Engine/Source/Runtime/AudioMixer/Private/AudioMixerSourceManager.cpp
@@ -107,6 +107,20 @@ FAutoConsoleVariableRef CVarCommandBufferMaxSizeMb(
 	TEXT("How big to allow the command buffer to grow before ignoring more commands"),
 	ECVF_Default);
 
+static int32 CommandBufferInitialCapacityCvar = 500;
+FAutoConsoleVariableRef CVarCommandBufferInitialCapacity(
+	TEXT("au.CommandBufferInitialCapacity"),
+	CommandBufferInitialCapacityCvar,
+	TEXT("How many elements to initialize the command buffer capacity with"),
+	ECVF_Default);
+
+static float CommandBufferGrowthFactorCvar = 0.20f;
+FAutoConsoleVariableRef CVarCommandBufferGrowthFactor(
+	TEXT("au.CommandBufferGrowthFactor"),
+	CommandBufferGrowthFactorCvar,
+	TEXT("How much to grow the Command Buffer when adding new commands that will cause re-allocation"),
+	ECVF_Default);
+
 static float AudioCommandExecTimeMsWarningThresholdCvar = 500.f;
 FAutoConsoleVariableRef CVarAudioCommandExecTimeMsWarningThreshold(
 	TEXT("au.AudioThreadCommand.ExecutionTimeWarningThresholdInMs"),
@@ -228,6 +242,10 @@ const TCHAR* LexToString(ESourceManagerRenderThreadPhase InPhase)
 
 namespace Audio
 {
+	int32 GetCommandBufferInitialCapacity()
+	{
+		return FMath::Clamp(CommandBufferInitialCapacityCvar, 0, 10000);
+	}
 	/*************************************************************************
 	* FMixerSourceManager
 	**************************************************************************/
@@ -250,6 +268,9 @@ namespace Audio
 
 		// Immediately trigger the command processed in case a flush happens before the audio thread swaps command buffers
 		CommandsProcessedEvent->Trigger();
+
+		// reserve the first buffer with the initial capacity
+		CommandBuffers[0].SourceCommandQueue.Reserve(GetCommandBufferInitialCapacity());
 	}
 
 	FMixerSourceManager::~FMixerSourceManager()
@@ -3366,42 +3387,51 @@ namespace Audio
 
 		// Add the function to the command queue:
 		int32 AudioThreadCommandIndex = !RenderThreadCommandBufferIndex.GetValue();
-		SIZE_T CurrentBufferSizeInBytes = CommandBuffers[AudioThreadCommandIndex].SourceCommandQueue.GetAllocatedSize();
-
-		TRACE_INT_VALUE(TEXT("AudioMixerThreadCommands::CurrentBufferSizeInKb"), CurrentBufferSizeInBytes >> 10);
+		FCommands& Commands = CommandBuffers[AudioThreadCommandIndex];
+		SIZE_T CurrentBufferSizeInBytes = Commands.SourceCommandQueue.GetAllocatedSize();
 
 		static SIZE_T WarnSize = 1024 * 1024;
 		if (CurrentBufferSizeInBytes > WarnSize )
 		{
-			SIZE_T Num = CommandBuffers[AudioThreadCommandIndex].SourceCommandQueue.Num();
+			SIZE_T Num = Commands.SourceCommandQueue.Num();
 			float TimeSinceLastComplete = FPlatformTime::ToSeconds64(FPlatformTime::Cycles64() - LastPumpCompleteTimeInCycles);
 
-			UE_LOG(LogAudioMixer, Error, TEXT("Command Queue has grown to %ukb, containing %d cmds, last complete pump was %2.5f seconds ago."),
-				CurrentBufferSizeInBytes >> 10, Num, TimeSinceLastComplete);
+			UE_LOG(LogAudioMixer, Error, TEXT("Command Queue %d has grown to %ukb, containing %d cmds, last complete pump was %2.5f seconds ago."),
+				AudioThreadCommandIndex, CurrentBufferSizeInBytes >> 10, Num, TimeSinceLastComplete);
 			WarnSize *= 2;
 
 			DoStallDiagnostics();
 		}
 
-		// Before adding further commands, ensure we're not growing outside any sensible size for these buffers.
-		// On shipping builds, this will just stop us crashing from growing out of control and OOMing the machine.
-		static bool bCommandBufferOverflowDetected = false;
-		const SIZE_T MaxBufferSizeInBytes = ((SIZE_T)CommandBufferMaxSizeInMbCvar) << 20;
-		if (CurrentBufferSizeInBytes < MaxBufferSizeInBytes)
+		SIZE_T Num = Commands.SourceCommandQueue.Num();
+		SIZE_T Max = Commands.SourceCommandQueue.Max();
+		if (Num == Max)
 		{
-			CommandBuffers[AudioThreadCommandIndex].SourceCommandQueue.Add(AudioCommand);
-			NumCommands.Increment();
-			if (bCommandBufferOverflowDetected)
+			// do our own re-allocation based on our own growth factor rather than letting the TArray grow itself
+			// (TArray grows by about 2x each time)
+			float GrowthFactor = FMath::Clamp(CommandBufferGrowthFactorCvar, 0.01f, 1.0f);
+			int32 NewSize = Num + (float)Num * GrowthFactor;
+			Commands.SourceCommandQueue.Reserve(NewSize);
+
+			// get new size
+			CurrentBufferSizeInBytes = Commands.SourceCommandQueue.GetAllocatedSize();
+
+			// check that we haven't gone over the max size
+			const SIZE_T MaxBufferSizeInBytes = ((SIZE_T)CommandBufferMaxSizeInMbCvar) << 20;
+			if (CurrentBufferSizeInBytes >= MaxBufferSizeInBytes)
 			{
-				UE_LOG(LogAudioMixer, Log, TEXT("Command buffer shrunk to %umb, allowing adds again."), CurrentBufferSizeInBytes >> 20);
-				bCommandBufferOverflowDetected = false;
+				// this will only throw an error every time we have to reallocate, which will be less often then every single time we add
+				Commands.NumTimesOvergrown++;
+				UE_LOG(LogAudioMixer, Error, TEXT("%d: Command buffer %d allocated size has grown to %umb! Likely cause the AudioRenderer has hung"), Commands.NumTimesOvergrown, AudioThreadCommandIndex, CurrentBufferSizeInBytes >> 20);
 			}
 		}
-		else if (!bCommandBufferOverflowDetected)
-		{
-			UE_LOG(LogAudioMixer, Error, TEXT("Command buffer grown to %umb, preventing any more adds! Likely cause the AudioRenderer has hung"), CurrentBufferSizeInBytes >> 20);
-			bCommandBufferOverflowDetected = true;
-		}
+
+		// always add commands to the buffer. If we're not going to assert, might as well chug along and hope we can recover!
+		Commands.SourceCommandQueue.Add(AudioCommand);
+		NumCommands.Increment();
+
+		TRACE_INT_VALUE(TEXT("AudioMixerThreadCommands::NumCommands"), Commands.SourceCommandQueue.Num());
+		TRACE_INT_VALUE(TEXT("AudioMixerThreadCommands::CurrentBufferSizeInKb"), CurrentBufferSizeInBytes >> 10);
 	}
 
 
@@ -3447,8 +3477,7 @@ namespace Audio
 
 		// Pop and execute all the commands that came since last update tick
 		TArray<FAudioMixerThreadCommand> DelayedCommands;
-
-		RenderThreadPhase = ESourceManagerRenderThreadPhase::PumpCmds;		
+		RenderThreadPhase = ESourceManagerRenderThreadPhase::PumpCmds;
 		for (int32 Id = 0; Id < NumCommandsToExecute; ++Id)
 		{
 			// First copy/move out the command and keep a copy of it.
@@ -3475,7 +3504,9 @@ namespace Audio
 		}
 
 		LastPumpCompleteTimeInCycles = FPlatformTime::Cycles64();
-		Commands.SourceCommandQueue = DelayedCommands;
+		// This is intentionally re-assigning the Command Queue and clearing the buffer in the process
+		Commands.SourceCommandQueue = MoveTemp(DelayedCommands);
+		Commands.SourceCommandQueue.Reserve(GetCommandBufferInitialCapacity());
 
 		if (FPlatformProcess::SupportsMultithreading())
 		{
diff --git a/Engine/Source/Runtime/AudioMixer/Private/AudioMixerSourceManager.h b/Engine/Source/Runtime/AudioMixer/Private/AudioMixerSourceManager.h
--- a/Engine/Source/Runtime/AudioMixer/Private/AudioMixerSourceManager.h
+++ b/Engine/Source/Runtime/AudioMixer/Private/AudioMixerSourceManager.h
@@ -421,6 +421,7 @@ namespace Audio
 		// A command queue to execute commands from audio thread (or game thread) to audio mixer device thread.
 		struct FCommands
 		{
+			int32 NumTimesOvergrown = 0;
 			TArray<FAudioMixerThreadCommand> SourceCommandQueue;
 		};
 		
