Fix for dynamic submixes not always obeying "bAutoRouteToMasterSubmixWhenOrphaned"
Fix for crash in debug logging of submix graph

diff --git a/Engine/Source/Runtime/AudioMixer/Private/AudioMixerDevice.cpp b/Engine/Source/Runtime/AudioMixer/Private/AudioMixerDevice.cpp
--- a/Engine/Source/Runtime/AudioMixer/Private/AudioMixerDevice.cpp
+++ b/Engine/Source/Runtime/AudioMixer/Private/AudioMixerDevice.cpp
@@ -221,7 +221,7 @@ namespace Audio
 		{		
 			for (const TObjectPtr<USoundEffectSubmixPreset>& i: Submix->SubmixEffectChain)
 			{
-				FxChain += FString::Printf(TEXT("[%s]"), *i->GetName());
+				FxChain += FString::Printf(TEXT("[%s]"), *GetNameSafe(i));
 			}
 		}
 
@@ -1264,7 +1264,7 @@ namespace Audio
 			{
 				ParentSubmixInstance = GetSubmixInstance(Parent).Pin();
 			}
-			else
+			else if (SubmixWithParent->bAutoRouteToMasterSubmixWhenOrphaned)
 			{
 				// If this submix is itself the broadcast submix, set its parent to the master submix
 				if (SubmixInstance == RequiredSubmixInstances[static_cast<int32>(ERequiredSubmixes::BaseDefault)])
