Get rid of custom "Growth Factor" of AudioMixerSourceManager command buffer

UE's default growth factor is 1.35 for each re-allocation of the array: (so `Max_N = 1.35 * Max_N-1`), so _not_ doubling every time. 1.35 is a reasonable factor, and using a smaller value doesn't seem to gain that much benefit and adds code complexity.

The actual performance boost comes from "reserving" the array at creation. And preventing log spam avoids the frame hitches.

Also, avoid holding onto the RenderThreadIndex Lock for longer than we need to. Use a ThreadSafeCounter for "NumTimesOvergrown" as well so we avoid needing a lock for that as well.

#tests PIE

diff --git a/Engine/Source/Runtime/AudioMixer/Private/AudioMixerSourceManager.cpp b/Engine/Source/Runtime/AudioMixer/Private/AudioMixerSourceManager.cpp
--- a/Engine/Source/Runtime/AudioMixer/Private/AudioMixerSourceManager.cpp
+++ b/Engine/Source/Runtime/AudioMixer/Private/AudioMixerSourceManager.cpp
@@ -114,13 +114,6 @@ FAutoConsoleVariableRef CVarCommandBufferInitialCapacity(
 	TEXT("How many elements to initialize the command buffer capacity with"),
 	ECVF_Default);
 
-static float CommandBufferGrowthFactorCvar = 0.20f;
-FAutoConsoleVariableRef CVarCommandBufferGrowthFactor(
-	TEXT("au.CommandBufferGrowthFactor"),
-	CommandBufferGrowthFactorCvar,
-	TEXT("How much to grow the Command Buffer when adding new commands that will cause re-allocation"),
-	ECVF_Default);
-
 static float AudioCommandExecTimeMsWarningThresholdCvar = 500.f;
 FAutoConsoleVariableRef CVarAudioCommandExecTimeMsWarningThreshold(
 	TEXT("au.AudioThreadCommand.ExecutionTimeWarningThresholdInMs"),
@@ -3381,56 +3374,61 @@ namespace Audio
 	{
 		FAudioMixerThreadCommand AudioCommand(MoveTemp(InFunction), InDebugString, bInDeferExecution);
 
-		// Here, we make sure that we don't flip our command double buffer while we are executing this function.
-		FScopeLock ScopeLock(&CommandBufferIndexCriticalSection);
-		AUDIO_MIXER_CHECK_GAME_THREAD(MixerDevice);
-
-		// Add the function to the command queue:
-		int32 AudioThreadCommandIndex = !RenderThreadCommandBufferIndex.GetValue();
-		FCommands& Commands = CommandBuffers[AudioThreadCommandIndex];
-		SIZE_T CurrentBufferSizeInBytes = Commands.SourceCommandQueue.GetAllocatedSize();
-
-		static SIZE_T WarnSize = 1024 * 1024;
-		if (CurrentBufferSizeInBytes > WarnSize )
+		// collect values for debugging
+		// outside of the ScopeLock so we can avoid doing a bunch of work that doesn't require the lock
+		SIZE_T OldMax = 0;
+		SIZE_T NewMax = 0;
+		SIZE_T NewNum = 0;
+		SIZE_T CurrentBufferSizeInBytes = 0;
+		int32 AudioThreadCommandIndex = -1;
 		{
-			SIZE_T Num = Commands.SourceCommandQueue.Num();
-			float TimeSinceLastComplete = FPlatformTime::ToSeconds64(FPlatformTime::Cycles64() - LastPumpCompleteTimeInCycles);
+			// Here, we make sure that we don't flip our command double buffer while modifying the command buffer
+			FScopeLock ScopeLock(&CommandBufferIndexCriticalSection);
+			AUDIO_MIXER_CHECK_GAME_THREAD(MixerDevice);
 
-			UE_LOG(LogAudioMixer, Error, TEXT("Command Queue %d has grown to %ukb, containing %d cmds, last complete pump was %2.5f seconds ago."),
-				AudioThreadCommandIndex, CurrentBufferSizeInBytes >> 10, Num, TimeSinceLastComplete);
-			WarnSize *= 2;
+			// Add the function to the command queue:
+			AudioThreadCommandIndex = !RenderThreadCommandBufferIndex.GetValue();
+			FCommands& Commands = CommandBuffers[AudioThreadCommandIndex];
 
-			DoStallDiagnostics();
-		}
+			OldMax = Commands.SourceCommandQueue.Max();
+			
+			// always add commands to the buffer. If we're not going to assert, might as well chug along and hope we can recover!
+			Commands.SourceCommandQueue.Add(AudioCommand);
+			NumCommands.Increment();
 
-		SIZE_T Num = Commands.SourceCommandQueue.Num();
-		SIZE_T Max = Commands.SourceCommandQueue.Max();
-		if (Num == Max)
+			NewNum = Commands.SourceCommandQueue.Num();
+			NewMax = Commands.SourceCommandQueue.Max();
+			CurrentBufferSizeInBytes = Commands.SourceCommandQueue.GetAllocatedSize();
+		}
+		
+		// log warnings for command buffer growing too large
+		if (OldMax != NewMax)
 		{
-			// do our own re-allocation based on our own growth factor rather than letting the TArray grow itself
-			// (TArray grows by about 2x each time)
-			float GrowthFactor = FMath::Clamp(CommandBufferGrowthFactorCvar, 0.01f, 1.0f);
-			int32 NewSize = Num + (float)Num * GrowthFactor;
-			Commands.SourceCommandQueue.Reserve(NewSize);
+			// Only throw a warning every time we have to reallocate, which will be less often then every single time we add
+			static SIZE_T WarnSize = 1024 * 1024;
+			if (CurrentBufferSizeInBytes > WarnSize )
+			{
+				float TimeSinceLastComplete = FPlatformTime::ToSeconds64(FPlatformTime::Cycles64() - LastPumpCompleteTimeInCycles);
 
-			// get new size
-			CurrentBufferSizeInBytes = Commands.SourceCommandQueue.GetAllocatedSize();
+				UE_LOG(LogAudioMixer, Error, TEXT("Command Queue %d has grown to %ukb, containing %d cmds, last complete pump was %2.5f seconds ago."),
+					AudioThreadCommandIndex, CurrentBufferSizeInBytes >> 10, NewNum, TimeSinceLastComplete);
+				WarnSize *= 2;
 
+				DoStallDiagnostics();
+			}
+			
 			// check that we haven't gone over the max size
 			const SIZE_T MaxBufferSizeInBytes = ((SIZE_T)CommandBufferMaxSizeInMbCvar) << 20;
 			if (CurrentBufferSizeInBytes >= MaxBufferSizeInBytes)
 			{
-				// this will only throw an error every time we have to reallocate, which will be less often then every single time we add
-				Commands.NumTimesOvergrown++;
-				UE_LOG(LogAudioMixer, Error, TEXT("%d: Command buffer %d allocated size has grown to %umb! Likely cause the AudioRenderer has hung"), Commands.NumTimesOvergrown, AudioThreadCommandIndex, CurrentBufferSizeInBytes >> 20);
+				int32 NumTimesOvergrown = CommandBuffers[AudioThreadCommandIndex].NumTimesOvergrown.Increment();
+				UE_LOG(LogAudioMixer, Error, TEXT("%d: Command buffer %d allocated size has grown to %umb! Likely cause the AudioRenderer has hung"),
+					NumTimesOvergrown, AudioThreadCommandIndex, CurrentBufferSizeInBytes >> 20);
 			}
 		}
 
-		// always add commands to the buffer. If we're not going to assert, might as well chug along and hope we can recover!
-		Commands.SourceCommandQueue.Add(AudioCommand);
-		NumCommands.Increment();
-
-		TRACE_INT_VALUE(TEXT("AudioMixerThreadCommands::NumCommands"), Commands.SourceCommandQueue.Num());
+		// update trace values
+		TRACE_INT_VALUE(TEXT("AudioMixerThreadCommands::NumCommands"), NewNum);
 		TRACE_INT_VALUE(TEXT("AudioMixerThreadCommands::CurrentBufferSizeInKb"), CurrentBufferSizeInBytes >> 10);
 	}
 
diff --git a/Engine/Source/Runtime/AudioMixer/Private/AudioMixerSourceManager.h b/Engine/Source/Runtime/AudioMixer/Private/AudioMixerSourceManager.h
--- a/Engine/Source/Runtime/AudioMixer/Private/AudioMixerSourceManager.h
+++ b/Engine/Source/Runtime/AudioMixer/Private/AudioMixerSourceManager.h
@@ -421,7 +421,7 @@ namespace Audio
 		// A command queue to execute commands from audio thread (or game thread) to audio mixer device thread.
 		struct FCommands
 		{
-			int32 NumTimesOvergrown = 0;
+			FThreadSafeCounter NumTimesOvergrown = 0;
 			TArray<FAudioMixerThreadCommand> SourceCommandQueue;
 		};
 		
