Use background thread pool for procedural audio decode tasks

diff --git a/Engine/Source/Runtime/AudioMixer/Private/AudioMixerSourceDecode.cpp b/Engine/Source/Runtime/AudioMixer/Private/AudioMixerSourceDecode.cpp
--- a/Engine/Source/Runtime/AudioMixer/Private/AudioMixerSourceDecode.cpp
+++ b/Engine/Source/Runtime/AudioMixer/Private/AudioMixerSourceDecode.cpp
@@ -28,6 +28,14 @@ FAutoConsoleVariableRef CVarForceSynchronizedAudioTaskKick(
 	TEXT("0: Don't Force, 1: Force"),
 	ECVF_Default);
 
+static int32 UseBackgroundThreadPoolForProceduralAudioDecodeCvar = 1;
+FAutoConsoleVariableRef CVarUseBackgroundThreadPoolForProceduralAudioDecode(
+	TEXT("au.UseBackgroundThreadPoolForProceduralAudioDecode"),
+	UseBackgroundThreadPoolForProceduralAudioDecodeCvar,
+	TEXT("Use background thread pool for procedural audio decode tasks. This is affected by AudioThread.UseBackgroundThreadPool.\n")
+	TEXT("0: Don't use background thread pool, 1: Use background thread pool"),
+	ECVF_Default);
+
 namespace Audio
 {
 
@@ -295,8 +303,16 @@ public:
             Task->StartSynchronousTask();
             return;
         }
-        
-		Task->StartBackgroundTask();
+
+		if (UseBackgroundThreadPoolForProceduralAudioDecodeCvar)
+		{
+			const bool bUseBackground = ShouldUseBackgroundPoolFor_FAsyncRealtimeAudioTask();
+			Task->StartBackgroundTask(bUseBackground ? GBackgroundPriorityThreadPool : GThreadPool);
+		}
+		else
+		{
+			Task->StartBackgroundTask();
+		}
 	}
 
 	virtual EAudioTaskType GetType() const override
