CIS fix: ambigous call

diff --git a/Engine/Source/Runtime/AudioMixer/Private/AudioMixerSourceManager.cpp b/Engine/Source/Runtime/AudioMixer/Private/AudioMixerSourceManager.cpp
--- a/Engine/Source/Runtime/AudioMixer/Private/AudioMixerSourceManager.cpp
+++ b/Engine/Source/Runtime/AudioMixer/Private/AudioMixerSourceManager.cpp
@@ -3429,7 +3429,7 @@ namespace Audio
 		}
 
 		// update trace values
-		CSV_CUSTOM_STAT(Audio, AudioMixerThreadCommands, NewNum, ECsvCustomStatOp::Set);
+		CSV_CUSTOM_STAT(Audio, AudioMixerThreadCommands, static_cast<int32>(NewNum), ECsvCustomStatOp::Set);
 		TRACE_INT_VALUE(TEXT("AudioMixerThreadCommands::NumCommands"), NewNum);
 		TRACE_INT_VALUE(TEXT("AudioMixerThreadCommands::CurrentBufferSizeInKb"), CurrentBufferSizeInBytes >> 10);
 	}
