Make submix sends consistent across attenuation + soundbase

diff --git a/Engine/Config/BaseEngine.ini b/Engine/Config/BaseEngine.ini
--- a/Engine/Config/BaseEngine.ini
+++ b/Engine/Config/BaseEngine.ini
@@ -1220,6 +1220,16 @@ DefaultWorldPartitionHLODLayer=/Engine/Maps/Templates/HLODs/HLODLayer_Instanced.
 +FunctionRedirects=(OldName="EditorValidatorBase.CanValidate", NewName="EditorValidatorBase.K2_CanValidate")
 +FunctionRedirects=(OldName="EditorValidatorBase.CanValidateAsset", NewName="EditorValidatorBase.K2_CanValidateAsset")
 
++EnumRedirects=(OldName="/Script/Engine.ESubmixSendMethod",NewName="/Script/Engine.ESendLevelControlMethod")
++PropertyRedirects=(OldName="AttenuationSubmixSendSettings.Submix",NewName="AttenuationSubmixSendSettings.SoundSubmix")
++PropertyRedirects=(OldName="AttenuationSubmixSendSettings.SubmixSendMethod",NewName="AttenuationSubmixSendSettings.SendLevelControlMethod")
++PropertyRedirects=(OldName="AttenuationSubmixSendSettings.SubmixSendLevelMin",NewName="AttenuationSubmixSendSettings.MinSendLevel")
++PropertyRedirects=(OldName="AttenuationSubmixSendSettings.SubmixSendLevelMax",NewName="AttenuationSubmixSendSettings.MaxSendLevel")
++PropertyRedirects=(OldName="AttenuationSubmixSendSettings.SubmixSendDistanceMin",NewName="AttenuationSubmixSendSettings.MinSendDistance")
++PropertyRedirects=(OldName="AttenuationSubmixSendSettings.SubmixSendDistanceMax",NewName="AttenuationSubmixSendSettings.MaxSendDistance")
++PropertyRedirects=(OldName="AttenuationSubmixSendSettings.ManualSubmixSendLevel",NewName="AttenuationSubmixSendSettings.SendLevel")
++PropertyRedirects=(OldName="AttenuationSubmixSendSettings.CustomSubmixSendCurve",NewName="AttenuationSubmixSendSettings.CustomSendLevelCurve")
+
 +EnumRedirects=(OldName="ETargetingTraceType", ValueChanges=(("Sweep", "Sphere")) )
 
 ; HLOD actors now store their source actors using FWorldPartitionRuntimeCellObjectMapping rather than a custom FHLODSubActor class
diff --git a/Engine/Source/Runtime/AudioMixer/Private/AudioMixerSource.cpp b/Engine/Source/Runtime/AudioMixer/Private/AudioMixerSource.cpp
--- a/Engine/Source/Runtime/AudioMixer/Private/AudioMixerSource.cpp
+++ b/Engine/Source/Runtime/AudioMixer/Private/AudioMixerSource.cpp
@@ -91,6 +91,34 @@ namespace Audio
 
 			return SoundClass;
 		}
+
+		template<typename SendInfo>
+		void ClearPreviousSubmixSends(const TArray<SendInfo>& InPreviousSendInfos, const TArray<SendInfo>& InNewSendInfos, FMixerDevice* InMixerDevice, FMixerSourceVoice* InMixerSourceVoice)
+		{
+			// Loop through every previous send setting
+			for (const SendInfo& PreviousSendSetting : InPreviousSendInfos)
+			{
+				bool bFound = false;
+
+				// See if it's in the current send list
+				for (const SendInfo&  CurrentSendSettings : InNewSendInfos)
+				{
+					if (CurrentSendSettings.SoundSubmix == PreviousSendSetting.SoundSubmix)
+					{
+						bFound = true;
+						break;
+					}
+				}
+
+				// If it's not in the current send list, add to submixes to clear
+				if (!bFound)
+				{
+					FMixerSubmixPtr SubmixPtr = InMixerDevice->GetSubmixInstance(PreviousSendSetting.SoundSubmix).Pin();
+					InMixerSourceVoice->ClearSubmixSendInfo(SubmixPtr);
+				}
+			}
+		}
+		
 	} // namespace MixerSourcePrivate
 
 	namespace ModulationUtils
@@ -1675,6 +1703,52 @@ namespace Audio
 			MixerSourceVoice->SetSpatializationParams(SpatializationParams);
 		}
 	}
+	
+	void FMixerSource::UpdateSubmixSendLevels(const FSoundSubmixSendInfoBase& InSendInfo, const EMixerSourceSubmixSendStage InSendStage)
+	{
+		if (InSendInfo.SoundSubmix != nullptr)
+		{
+			const FMixerSubmixWeakPtr SubmixInstance = MixerDevice->GetSubmixInstance(InSendInfo.SoundSubmix);
+			float SendLevel = 1.0f;
+
+			// calculate send level based on distance if that method is enabled
+			if (!WaveInstance->bEnableSubmixSends)
+			{
+				SendLevel = 0.0f;
+			}
+			else if (InSendInfo.SendLevelControlMethod == ESendLevelControlMethod::Manual)
+			{
+				if (InSendInfo.DisableManualSendClamp)
+				{
+					SendLevel = InSendInfo.SendLevel;
+				}
+				else
+				{
+					SendLevel = FMath::Clamp(InSendInfo.SendLevel, 0.0f, 1.0f);
+				}
+			}
+			else
+			{
+				// The alpha value is determined identically between manual and custom curve methods
+				const FVector2D SendRadialRange = { InSendInfo.MinSendDistance, InSendInfo.MaxSendDistance};
+				const FVector2D SendLevelRange = { InSendInfo.MinSendLevel, InSendInfo.MaxSendLevel };
+				const float Denom = FMath::Max(SendRadialRange.Y - SendRadialRange.X, 1.0f);
+				const float Alpha = FMath::Clamp((WaveInstance->ListenerToSoundDistance - SendRadialRange.X) / Denom, 0.0f, 1.0f);
+
+				if (InSendInfo.SendLevelControlMethod == ESendLevelControlMethod::Linear)
+				{
+					SendLevel = FMath::Clamp(FMath::Lerp(SendLevelRange.X, SendLevelRange.Y, Alpha), 0.0f, 1.0f);
+				}
+				else // use curve
+				{
+					SendLevel = FMath::Clamp(InSendInfo.CustomSendLevelCurve.GetRichCurveConst()->Eval(Alpha), 0.0f, 1.0f);
+				}
+			}
+
+			// set the level and stage for this send
+			MixerSourceVoice->SetSubmixSendInfo(SubmixInstance, SendLevel, InSendStage);
+		}
+	}
 
 	void FMixerSource::UpdateEffects()
 	{
@@ -1749,7 +1823,7 @@ namespace Audio
 			{
 				SubmixPtr = MixerDevice->GetSubmixInstance(WaveInstance->SoundSubmix);
 			}
-			else if(WaveInstance->SoundSubmix && WaveInstance->SoundSubmix->bAutoRouteToMasterSubmixWhenOrphaned)
+			else if(!WaveInstance->SoundSubmix && WaveInstance->SoundSubmix->bAutoRouteToMasterSubmixWhenOrphaned)
 			{
 				SubmixPtr = MixerDevice->GetMasterSubmix();
 			}
@@ -1761,117 +1835,25 @@ namespace Audio
 			PrevousSubmix = SubmixKey;
 		}
 
-		if (WaveInstance->SubmixSendSettings.Num() > 0)
+		// Attenuation Submix Sends. (these come from Attenuation assets).
+		// These are largely identical to SoundSubmix Sends, but don't specify a send stage, so we pass one here.
+		for (const FAttenuationSubmixSendSettings& SendSettings : WaveInstance->AttenuationSubmixSends)
 		{
-			for (const FAttenuationSubmixSendSettings& SendSettings : WaveInstance->SubmixSendSettings)
-			{
-				if (SendSettings.Submix)
-				{
-					float SubmixSendLevel = 0.0f;
-
-					if (SendSettings.SubmixSendMethod == ESubmixSendMethod::Manual)
-					{
-						SubmixSendLevel = FMath::Clamp(SendSettings.ManualSubmixSendLevel, 0.0f, 1.0f);
-					}
-					else
-					{
-						// The alpha value is determined identically between manual and custom curve methods
-						const float Denom = FMath::Max(SendSettings.SubmixSendDistanceMax - SendSettings.SubmixSendDistanceMin, 1.0f);
-						const float Alpha = FMath::Clamp((WaveInstance->ListenerToSoundDistance - SendSettings.SubmixSendDistanceMin) / Denom, 0.0f, 1.0f);
-
-						if (SendSettings.SubmixSendMethod == ESubmixSendMethod::Linear)
-						{
-							SubmixSendLevel = FMath::Clamp(FMath::Lerp(SendSettings.SubmixSendLevelMin, SendSettings.SubmixSendLevelMax, Alpha), 0.0f, 1.0f);
-						}
-						else
-						{
-							SubmixSendLevel = FMath::Clamp(SendSettings.CustomSubmixSendCurve.GetRichCurveConst()->Eval(Alpha), 0.0f, 1.0f);
-						}
-					}
-
-
-					FMixerSubmixPtr SubmixPtr = MixerDevice->GetSubmixInstance(SendSettings.Submix).Pin();
-					MixerSourceVoice->SetSubmixSendInfo(SubmixPtr, SubmixSendLevel);
-				}
-			}
+			UpdateSubmixSendLevels(SendSettings, EMixerSourceSubmixSendStage::PostDistanceAttenuation);
 		}
-
-		// Clear submix sends if they need clearing.
-		if (PreviousSubmixSendSettings.Num() > 0)
+		// Clear any previous sends that may not exist now.
+		MixerSourcePrivate::ClearPreviousSubmixSends(PreviousAttenuationSendSettings, WaveInstance->AttenuationSubmixSends, MixerDevice, MixerSourceVoice);
+		PreviousAttenuationSendSettings = WaveInstance->AttenuationSubmixSends; 
+		
+		// Sound submix Sends. (these come from SoundBase derived assets).
+		for (FSoundSubmixSendInfo& SendInfo : WaveInstance->SoundSubmixSends)
 		{
-			// Loop through every previous send setting
-			for (FSoundSubmixSendInfo& PreviousSendSetting : PreviousSubmixSendSettings)
-			{
-				bool bFound = false;
-
-				// See if it's in the current send list
-				for (const FSoundSubmixSendInfo& CurrentSendSettings : WaveInstance->SoundSubmixSends)
-				{
-					if (CurrentSendSettings.SoundSubmix == PreviousSendSetting.SoundSubmix)
-					{
-						bFound = true;
-						break;
-					}
-				}
-
-				// If it's not in the current send list, add to submixes to clear
-				if (!bFound)
-				{
-					FMixerSubmixPtr SubmixPtr = MixerDevice->GetSubmixInstance(PreviousSendSetting.SoundSubmix).Pin();
-					MixerSourceVoice->ClearSubmixSendInfo(SubmixPtr);
-				}
-			}
+			UpdateSubmixSendLevels(SendInfo, MixerSourcePrivate::SubmixSendStageToMixerSourceSubmixSendStage(SendInfo.SendStage));
 		}
+		// Again, Clear any sends that maybe not exist now.
+		MixerSourcePrivate::ClearPreviousSubmixSends(PreviousSubmixSendSettings, WaveInstance->SoundSubmixSends, MixerDevice, MixerSourceVoice);
 		PreviousSubmixSendSettings = WaveInstance->SoundSubmixSends;
 
-		// Update submix send levels
-		for (FSoundSubmixSendInfo& SendInfo : WaveInstance->SoundSubmixSends)
-		{
-			if (SendInfo.SoundSubmix != nullptr)
-			{
-				FMixerSubmixWeakPtr SubmixInstance = MixerDevice->GetSubmixInstance(SendInfo.SoundSubmix);
-				float SendLevel = 1.0f;
-
-				// calculate send level based on distance if that method is enabled
-				if (!WaveInstance->bEnableSubmixSends)
-				{
-					SendLevel = 0.0f;
-				}
-				else if (SendInfo.SendLevelControlMethod == ESendLevelControlMethod::Manual)
-				{
-					if (SendInfo.DisableManualSendClamp)
-					{
-						SendLevel = SendInfo.SendLevel;
-					}
-					else
-					{
-						SendLevel = FMath::Clamp(SendInfo.SendLevel, 0.0f, 1.0f);
-					}
-				}
-				else
-				{
-					// The alpha value is determined identically between manual and custom curve methods
-					const FVector2D SendRadialRange = { SendInfo.MinSendDistance, SendInfo.MaxSendDistance};
-					const FVector2D SendLevelRange = { SendInfo.MinSendLevel, SendInfo.MaxSendLevel };
-					const float Denom = FMath::Max(SendRadialRange.Y - SendRadialRange.X, 1.0f);
-					const float Alpha = FMath::Clamp((WaveInstance->ListenerToSoundDistance - SendRadialRange.X) / Denom, 0.0f, 1.0f);
-
-					if (SendInfo.SendLevelControlMethod == ESendLevelControlMethod::Linear)
-					{
-						SendLevel = FMath::Clamp(FMath::Lerp(SendLevelRange.X, SendLevelRange.Y, Alpha), 0.0f, 1.0f);
-					}
-					else // use curve
-					{
-						SendLevel = FMath::Clamp(SendInfo.CustomSendLevelCurve.GetRichCurveConst()->Eval(Alpha), 0.0f, 1.0f);
-					}
-				}
-
-				// set the level and stage for this send
-				EMixerSourceSubmixSendStage SendStage = MixerSourcePrivate::SubmixSendStageToMixerSourceSubmixSendStage(SendInfo.SendStage);
-				MixerSourceVoice->SetSubmixSendInfo(SubmixInstance, SendLevel, SendStage);
-			}
-		}
- 		
 		MixerSourceVoice->SetEnablement(WaveInstance->bEnableBusSends, WaveInstance->bEnableBaseSubmix, WaveInstance->bEnableSubmixSends);
 
 		MixerSourceVoice->SetSourceBufferListener(WaveInstance->SourceBufferListener, WaveInstance->bShouldSourceBufferListenerZeroBuffer);
diff --git a/Engine/Source/Runtime/AudioMixer/Private/AudioMixerSource.h b/Engine/Source/Runtime/AudioMixer/Private/AudioMixerSource.h
--- a/Engine/Source/Runtime/AudioMixer/Private/AudioMixerSource.h
+++ b/Engine/Source/Runtime/AudioMixer/Private/AudioMixerSource.h
@@ -123,6 +123,7 @@ namespace Audio
 		float GetInheritedSubmixVolumeModulation() const;
 
 	private:
+		void UpdateSubmixSendLevels(const FSoundSubmixSendInfoBase& InSendInfo, EMixerSourceSubmixSendStage InSendStage);
 
 		FMixerDevice* MixerDevice;
 		FMixerBuffer* MixerBuffer;
@@ -180,6 +181,7 @@ namespace Audio
 
 		// An array of submix sends from previous update. Allows us to clear out submix sends if they are no longer being sent.
 		TArray<FSoundSubmixSendInfo> PreviousSubmixSendSettings;
+		TArray<FAttenuationSubmixSendSettings> PreviousAttenuationSendSettings;
 
 		// Whether or not we're currently releasing our resources. Prevents recycling the source until release is finished.
 		FThreadSafeBool bIsReleasing;
diff --git a/Engine/Source/Runtime/Engine/Classes/Sound/SoundAttenuation.h b/Engine/Source/Runtime/Engine/Classes/Sound/SoundAttenuation.h
--- a/Engine/Source/Runtime/Engine/Classes/Sound/SoundAttenuation.h
+++ b/Engine/Source/Runtime/Engine/Classes/Sound/SoundAttenuation.h
@@ -9,6 +9,7 @@
 #include "IAudioParameterInterfaceRegistry.h"
 #include "AudioLinkSettingsAbstract.h"
 #include "SoundAttenuationEditorSettings.h"
+#include "Sound/SoundSubmixSend.h"
 
 #include "SoundAttenuation.generated.h"
 
@@ -63,20 +64,6 @@ enum class EReverbSendMethod : uint8
 	Manual,
 };
 
-UENUM(BlueprintType)
-enum class ESubmixSendMethod : uint8
-{
-	// A submix send based on linear interpolation between a distance range and send-level range
-	Linear,
-
-	// A submix send based on a supplied curve
-	CustomCurve,
-
-	// A manual submix send level (Uses the specified constant send level value. Useful for 2D sounds.)
-	Manual,
-};
-
-
 UENUM(BlueprintType)
 enum class EPriorityAttenuationMethod : uint8
 {
@@ -113,45 +100,6 @@ struct FSoundAttenuationPluginSettings
 	TArray<TObjectPtr<USourceDataOverridePluginSourceSettingsBase>> SourceDataOverridePluginSettingsArray;
 };
 
-USTRUCT(BlueprintType)
-struct FAttenuationSubmixSendSettings
-{
-	GENERATED_USTRUCT_BODY()
-
-	/** Submix to send audio to based on distance. */
-	UPROPERTY(EditAnywhere, BlueprintReadWrite, Category = AttenuationSubmixSend)
-	TObjectPtr<USoundSubmixBase> Submix = nullptr;
-
-	/** What method to use to use for submix sends. */
-	UPROPERTY(EditAnywhere, BlueprintReadWrite, Category = AttenuationReverbSend)
-	ESubmixSendMethod SubmixSendMethod = ESubmixSendMethod::Linear;
-
-	/** The amount to send to the Submix when the sound is located at a distance equal to value specified in the submix send distance min. */
-	UPROPERTY(EditAnywhere, BlueprintReadWrite, Category = AttenuationSubmixSend, meta = (DisplayName = "Submix Min Send Level"))
-	float SubmixSendLevelMin = 0.0f;
-
-	/** The amount to send to the Submix when the sound is located at a distance equal to value specified in the reverb max send distance. */
-	UPROPERTY(EditAnywhere, BlueprintReadWrite, Category = AttenuationSubmixSend, meta = (DisplayName = "Submix Max Send Level"))
-	float SubmixSendLevelMax = 1.0f;
-
-	/** The min distance to send to the Submix. */
-	UPROPERTY(EditAnywhere, BlueprintReadWrite, Category = AttenuationSubmixSend, meta = (DisplayName = "Submix Min Send Distance"))
-	float SubmixSendDistanceMin = 400.0f;
-
-	/** The max distance to send to the Submix. */
-	UPROPERTY(EditAnywhere, BlueprintReadWrite, Category = AttenuationSubmixSend, meta = (DisplayName = "Submix Max Send Distance"))
-	float SubmixSendDistanceMax = 6000.0f;
-
-	/* The manual Submix send level to use. Doesn't change as a function of distance. */
-	UPROPERTY(EditAnywhere, BlueprintReadWrite, Category = AttenuationSubmixSend)
-	float ManualSubmixSendLevel = 0.2f;
-
-	/* The custom Submix send curve to use for distance-based send level. */
-	UPROPERTY(EditAnywhere, BlueprintReadWrite, Category = AttenuationSubmixSend)
-	FRuntimeFloatCurve CustomSubmixSendCurve;
-};
-
-
 // Defines how to speaker map the sound when using the non-spatialized radius feature
 UENUM(BlueprintType)
 enum class ENonSpatializedRadiusSpeakerMapMode : uint8
@@ -166,6 +114,14 @@ enum class ENonSpatializedRadiusSpeakerMapMode : uint8
 	Surround2D,
 };
 
+USTRUCT(BlueprintType)
+struct FAttenuationSubmixSendSettings : public FSoundSubmixSendInfoBase
+{
+	GENERATED_BODY();
+	
+	FAttenuationSubmixSendSettings();
+};
+
 /*
 The settings for attenuating.
 */
diff --git a/Engine/Source/Runtime/Engine/Classes/Sound/SoundSubmixSend.h b/Engine/Source/Runtime/Engine/Classes/Sound/SoundSubmixSend.h
--- a/Engine/Source/Runtime/Engine/Classes/Sound/SoundSubmixSend.h
+++ b/Engine/Source/Runtime/Engine/Classes/Sound/SoundSubmixSend.h
@@ -78,23 +78,13 @@ enum class ESendLevelControlMethod : uint8
 	Manual,
 };
 
-UENUM(BlueprintType)
-enum class ESubmixSendStage : uint8
-{
-	// Whether to do the send pre distance attenuation
-	PostDistanceAttenuation,
-
-	// Whether to do the send post distance attenuation
-	PreDistanceAttenuation,
-};
-
-// Class used to send audio to submixes from USoundBase
+// Common set of settings that are uses as submix sends.
 USTRUCT(BlueprintType)
-struct FSoundSubmixSendInfo
+struct FSoundSubmixSendInfoBase
 {
 	GENERATED_USTRUCT_BODY()
 
-	ENGINE_API FSoundSubmixSendInfo();
+	ENGINE_API FSoundSubmixSendInfoBase();
 
 	/*
 		Manual: Use Send Level only
@@ -103,11 +93,7 @@ struct FSoundSubmixSendInfo
 	*/
 	UPROPERTY(EditAnywhere, BlueprintReadWrite, Category = SubmixSend)
 	ESendLevelControlMethod SendLevelControlMethod;
-
-	/** Defines at what mix stage the send should happen.*/
-	UPROPERTY(EditAnywhere, BlueprintReadWrite, Category = SubmixSend)
-	ESubmixSendStage SendStage;
-
+	
 	// The submix to send the audio to
 	UPROPERTY(EditAnywhere, BlueprintReadWrite, Category = SubmixSend)
 	TObjectPtr<USoundSubmixBase> SoundSubmix;
@@ -140,3 +126,23 @@ struct FSoundSubmixSendInfo
 	UPROPERTY(EditAnywhere, BlueprintReadWrite, Category = SubmixSend)
 	FRuntimeFloatCurve CustomSendLevelCurve;
 };
+
+UENUM(BlueprintType)
+enum class ESubmixSendStage : uint8
+{
+	// Whether to do the send pre distance attenuation
+	PostDistanceAttenuation,
+
+	// Whether to do the send post distance attenuation
+	PreDistanceAttenuation,
+};
+
+USTRUCT(BlueprintType)
+struct FSoundSubmixSendInfo : public FSoundSubmixSendInfoBase
+{
+	GENERATED_BODY();
+
+	/** Defines at what mix stage the send should happen.*/
+	UPROPERTY(EditAnywhere, BlueprintReadWrite, Category = SubmixSend)
+	ESubmixSendStage SendStage = ESubmixSendStage::PostDistanceAttenuation;
+};
diff --git a/Engine/Source/Runtime/Engine/Private/ActiveSound.cpp b/Engine/Source/Runtime/Engine/Private/ActiveSound.cpp
--- a/Engine/Source/Runtime/Engine/Private/ActiveSound.cpp
+++ b/Engine/Source/Runtime/Engine/Private/ActiveSound.cpp
@@ -1708,12 +1708,12 @@ void FActiveSound::UpdateAttenuation(float DeltaTime, FSoundParseParameters& Par
 
 	if (Settings->bEnableSubmixSends)
 	{
-		ParseParams.SubmixSendSettings.Reset();
+		ParseParams.AttenuationSubmixSends.Reset();
 		for (const FAttenuationSubmixSendSettings& SendSettings : Settings->SubmixSendSettings)
 		{
-			if (SendSettings.Submix)
+			if (SendSettings.SoundSubmix)
 			{
-				ParseParams.SubmixSendSettings.Add(SendSettings);
+				ParseParams.AttenuationSubmixSends.Add(SendSettings);
 			}
 		}
 	}
diff --git a/Engine/Source/Runtime/Engine/Private/Audio.cpp b/Engine/Source/Runtime/Engine/Private/Audio.cpp
--- a/Engine/Source/Runtime/Engine/Private/Audio.cpp
+++ b/Engine/Source/Runtime/Engine/Private/Audio.cpp
@@ -20,6 +20,7 @@
 #include "UObject/UObjectIterator.h"
 #include "XmlFile.h"
 #include "XmlNode.h"
+#include "Algo/ForEach.h"
 
 DEFINE_LOG_CATEGORY(LogAudio);
 
@@ -976,14 +977,17 @@ void FWaveInstance::AddReferencedObjects( FReferenceCollector& Collector )
 		Collector.AddReferencedObject(SynthSound->GetOwningSynthComponentPtr());
 	}
 
-	for (FAttenuationSubmixSendSettings& SubmixSend : SubmixSendSettings)
+	auto AddSubmixSendRef = [&Collector](FSoundSubmixSendInfoBase& Info)
 	{
-		if (SubmixSend.Submix)
+		if (Info.SoundSubmix)
 		{
-			Collector.AddReferencedObject(SubmixSend.Submix);
-		}
-	}
-
+			Collector.AddReferencedObject(Info.SoundSubmix);
+		}	
+	};
+	
+	Algo::ForEach(SoundSubmixSends, AddSubmixSendRef);
+	Algo::ForEach(AttenuationSubmixSends, AddSubmixSendRef);
+	
 	Collector.AddReferencedObject( SoundClass );
 	NotifyBufferFinishedHooks.AddReferencedObjects( Collector );
 }
diff --git a/Engine/Source/Runtime/Engine/Private/SoundAttenuation.cpp b/Engine/Source/Runtime/Engine/Private/SoundAttenuation.cpp
--- a/Engine/Source/Runtime/Engine/Private/SoundAttenuation.cpp
+++ b/Engine/Source/Runtime/Engine/Private/SoundAttenuation.cpp
@@ -155,6 +155,16 @@ USoundAttenuation::USoundAttenuation(const FObjectInitializer& ObjectInitializer
 {
 }
 
+FAttenuationSubmixSendSettings::FAttenuationSubmixSendSettings()
+{
+	// These were the defaults in the previous attenuation settings.
+	MinSendLevel = 0.0f;
+	MaxSendLevel = 1.0f;
+	MinSendDistance = 400.0f;
+	MaxSendDistance = 6000.0f;
+	SendLevel = 0.2f;
+	SendLevelControlMethod = ESendLevelControlMethod::Linear;
+}
 
 #define LOCTEXT_NAMESPACE "AudioParameterInterface"
 #define AUDIO_PARAMETER_INTERFACE_NAMESPACE "UE.Attenuation"
diff --git a/Engine/Source/Runtime/Engine/Private/SoundSubmixSend.cpp b/Engine/Source/Runtime/Engine/Private/SoundSubmixSend.cpp
--- a/Engine/Source/Runtime/Engine/Private/SoundSubmixSend.cpp
+++ b/Engine/Source/Runtime/Engine/Private/SoundSubmixSend.cpp
@@ -4,9 +4,8 @@
 
 #include UE_INLINE_GENERATED_CPP_BY_NAME(SoundSubmixSend)
 
-FSoundSubmixSendInfo::FSoundSubmixSendInfo()
+FSoundSubmixSendInfoBase::FSoundSubmixSendInfoBase()
 	: SendLevelControlMethod(ESendLevelControlMethod::Manual)
-	, SendStage(ESubmixSendStage::PostDistanceAttenuation)
 	, SoundSubmix(nullptr)
 	, SendLevel(1.0f)
 	, DisableManualSendClamp(false)
@@ -14,6 +13,6 @@ FSoundSubmixSendInfo::FSoundSubmixSendInfo()
 	, MaxSendLevel(1.0f)
 	, MinSendDistance(100.0f)
 	, MaxSendDistance(1000.0f)
-	{
-	}
+{
+}
 
diff --git a/Engine/Source/Runtime/Engine/Private/SoundWave.cpp b/Engine/Source/Runtime/Engine/Private/SoundWave.cpp
--- a/Engine/Source/Runtime/Engine/Private/SoundWave.cpp
+++ b/Engine/Source/Runtime/Engine/Private/SoundWave.cpp
@@ -3443,17 +3443,14 @@ void USoundWave::Parse(FAudioDevice* AudioDevice, const UPTRINT NodeWaveInstance
 		}
 	}
 
-
- 	// Copy the submix send settings
- 	WaveInstance->SubmixSendSettings = ParseParams.SubmixSendSettings;
+ 	// Copy the submix sends. (both attenuation and soundbase versions).
+ 	WaveInstance->AttenuationSubmixSends = ParseParams.AttenuationSubmixSends;
+	WaveInstance->SoundSubmixSends = ParseParams.SoundSubmixSends;
 
 	// Get the envelope follower settings
 	WaveInstance->EnvelopeFollowerAttackTime = ParseParams.EnvelopeFollowerAttackTime;
 	WaveInstance->EnvelopeFollowerReleaseTime = ParseParams.EnvelopeFollowerReleaseTime;
 
-	// Copy over the submix sends.
-	WaveInstance->SoundSubmixSends = ParseParams.SoundSubmixSends;
-
 	// Copy over the source bus send and data
 	if (!WaveInstance->ActiveSound->bIsPreviewSound)
 	{
diff --git a/Engine/Source/Runtime/Engine/Public/ActiveSound.h b/Engine/Source/Runtime/Engine/Public/ActiveSound.h
--- a/Engine/Source/Runtime/Engine/Public/ActiveSound.h
+++ b/Engine/Source/Runtime/Engine/Public/ActiveSound.h
@@ -129,8 +129,9 @@ struct FSoundParseParameters
 	// The sound submix to use for the wave instance
 	USoundSubmixBase* SoundSubmix;
 
-	// The submix sends to use
+	// The submix sends. 
 	TArray<FSoundSubmixSendInfo> SoundSubmixSends;
+	TArray<FAttenuationSubmixSendSettings> AttenuationSubmixSends;
 
 	// The source bus sends to use
 	TArray<FSoundSourceBusSendInfo> BusSends[(int32)EBusSendType::Count];
@@ -142,9 +143,6 @@ struct FSoundParseParameters
 	float ManualReverbSendLevel;
 	FRuntimeFloatCurve CustomReverbSendCurve;
 
-	// Submix send params to use for this sound
-	TArray<FAttenuationSubmixSendSettings> SubmixSendSettings;
-
 	// The distance between left and right channels when spatializing stereo assets
 	float StereoSpread;
 
diff --git a/Engine/Source/Runtime/Engine/Public/Audio.h b/Engine/Source/Runtime/Engine/Public/Audio.h
--- a/Engine/Source/Runtime/Engine/Public/Audio.h
+++ b/Engine/Source/Runtime/Engine/Public/Audio.h
@@ -403,7 +403,7 @@ public:
 	float ManualReverbSendLevel;
 	
 	/** The submix send settings to use. */
-	TArray<FAttenuationSubmixSendSettings> SubmixSendSettings;
+	TArray<FAttenuationSubmixSendSettings> AttenuationSubmixSends;
 
 private:
 	/** Cached play order */
