Fix static analysis warning
#rnx

diff --git a/Engine/Source/Runtime/AudioMixer/Private/AudioMixerSource.cpp b/Engine/Source/Runtime/AudioMixer/Private/AudioMixerSource.cpp
--- a/Engine/Source/Runtime/AudioMixer/Private/AudioMixerSource.cpp
+++ b/Engine/Source/Runtime/AudioMixer/Private/AudioMixerSource.cpp
@@ -1823,7 +1823,7 @@ namespace Audio
 			{
 				SubmixPtr = MixerDevice->GetSubmixInstance(WaveInstance->SoundSubmix);
 			}
-			else if(!WaveInstance->SoundSubmix && WaveInstance->SoundSubmix->bAutoRouteToMasterSubmixWhenOrphaned)
+			else if(WaveInstance->SoundSubmix && WaveInstance->SoundSubmix->bAutoRouteToMasterSubmixWhenOrphaned)
 			{
 				SubmixPtr = MixerDevice->GetMasterSubmix();
 			}
