Modulation Submix Code Fixes

diff --git a/Engine/Source/Runtime/AudioMixer/Private/AudioMixerSource.cpp b/Engine/Source/Runtime/AudioMixer/Private/AudioMixerSource.cpp
--- a/Engine/Source/Runtime/AudioMixer/Private/AudioMixerSource.cpp
+++ b/Engine/Source/Runtime/AudioMixer/Private/AudioMixerSource.cpp
@@ -529,7 +529,6 @@ namespace Audio
 		, bIsEffectTailsDone(false)
 		, bIsPlayingEffectTails(false)
 		, bEditorWarnedChangedSpatialization(false)
-		, bUsingHRTFSpatialization(false)
 		, bIs3D(false)
 		, bDebugMode(false)
 		, bIsVorbis(false)
@@ -593,6 +592,8 @@ namespace Audio
 		bIsEffectTailsDone = true;
 		bIsDone = false;
 
+		bBypassingSubmixModulation = false;
+
 		FSoundBuffer* SoundBuffer = static_cast<FSoundBuffer*>(MixerBuffer);
 		if (SoundBuffer->NumChannels > 0)
 		{
@@ -1640,11 +1641,11 @@ namespace Audio
 		{
 			FModulationDestination* SubmixOutVolDest = CurrSubmixPtr->GetOutputVolumeDestination();
 			FModulationDestination* SubmixWetVolDest = CurrSubmixPtr->GetWetVolumeDestination();
-			if (SubmixOutVolDest && SubmixOutVolDest->IsActive())
+			if (SubmixOutVolDest)
 			{
 				SubmixModVolume *= SubmixOutVolDest->GetValue();
 			}
-			if (SubmixWetVolDest && SubmixWetVolDest->IsActive())
+			if (SubmixWetVolDest)
 			{
 				SubmixModVolume *= SubmixWetVolDest->GetValue();
 			}
diff --git a/Engine/Source/Runtime/AudioMixer/Private/AudioMixerSource.h b/Engine/Source/Runtime/AudioMixer/Private/AudioMixerSource.h
--- a/Engine/Source/Runtime/AudioMixer/Private/AudioMixerSource.h
+++ b/Engine/Source/Runtime/AudioMixer/Private/AudioMixerSource.h
@@ -187,7 +187,6 @@ namespace Audio
 		FThreadSafeBool bIsReleasing;
 
 		uint32 bEditorWarnedChangedSpatialization : 1;
-		uint32 bUsingHRTFSpatialization : 1;
 		uint32 bIs3D : 1;
 		uint32 bDebugMode : 1;
 		uint32 bIsVorbis : 1;
