PR #10346: Allow Quartz Subsystem to (optionally) tick when game is paused

diff --git a/Engine/Source/Runtime/AudioMixer/Private/Quartz/QuartzSubsystem.cpp b/Engine/Source/Runtime/AudioMixer/Private/Quartz/QuartzSubsystem.cpp
--- a/Engine/Source/Runtime/AudioMixer/Private/Quartz/QuartzSubsystem.cpp
+++ b/Engine/Source/Runtime/AudioMixer/Private/Quartz/QuartzSubsystem.cpp
@@ -208,6 +208,11 @@ void UQuartzSubsystem::Tick(float DeltaTime)
 	TickableObjectManagerPtr->Tick(DeltaTime);
 }
 
+bool UQuartzSubsystem::IsTickableWhenPaused() const
+{
+	return bTickEvenWhenPaused;
+}
+
 bool UQuartzSubsystem::IsTickable() const
 {
 	check(TickableObjectManagerPtr);
@@ -277,35 +282,6 @@ Audio::FQuartzQuantizedRequestData UQuartzSubsystem::CreateRequestDataForSchedul
 	return CommandInitInfo;
 }
 
-
-Audio::FQuartzQuantizedRequestData UQuartzSubsystem::CreateDataForTickRateChange(UQuartzClockHandle* InClockHandle, const FOnQuartzCommandEventBP& InDelegate, const Audio::FQuartzClockTickRate& InNewTickRate, const FQuartzQuantizationBoundary& InQuantizationBoundary)
-{
-	// new static function
-	return UQuartzSubsystem::CreateRequestDataForTickRateChange(InClockHandle, InDelegate, InNewTickRate, InQuantizationBoundary);
-}
-
-
-Audio::FQuartzQuantizedRequestData UQuartzSubsystem::CreateDataForTransportReset(UQuartzClockHandle* InClockHandle, const FQuartzQuantizationBoundary& InQuantizationBoundary, const FOnQuartzCommandEventBP& InDelegate)
-{
-	// new static function
-	return UQuartzSubsystem::CreateRequestDataForTransportReset(InClockHandle, InQuantizationBoundary, InDelegate);
-}
-
-
-Audio::FQuartzQuantizedRequestData CreateDataForStartOtherClock(UQuartzClockHandle* InClockHandle, FName InClockToStart, const FQuartzQuantizationBoundary& InQuantizationBoundary, const FOnQuartzCommandEventBP& InDelegate)
-{
-	// new static function
-	return UQuartzSubsystem::CreateRequestDataForStartOtherClock(InClockHandle, InClockToStart, InQuantizationBoundary, InDelegate);
-}
-
-
-Audio::FQuartzQuantizedRequestData CreateDataDataForSchedulePlaySound(UQuartzClockHandle* InClockHandle, const FOnQuartzCommandEventBP& InDelegate, const FQuartzQuantizationBoundary& InQuantizationBoundary)
-{
-	// new static function
-	return UQuartzSubsystem::CreateRequestDataForSchedulePlaySound(InClockHandle, InDelegate, InQuantizationBoundary);
-}
-
-
 bool UQuartzSubsystem::IsQuartzEnabled()
 {
 	return true;
@@ -660,14 +636,9 @@ float UQuartzSubsystem::GetRoundTripMaxLatency(const UObject* WorldContextObject
 	return GetAudioRenderThreadToGameThreadMinLatency() + GetGameThreadToAudioRenderThreadMinLatency(WorldContextObject);
 }
 
-
-void UQuartzSubsystem::AddCommandToClock(const UObject* WorldContextObject, Audio::FQuartzQuantizedCommandInitInfo& InQuantizationCommandInitInfo, FName ClockName)
+void UQuartzSubsystem::SetQuartzSubsystemTickableWhenPaused(const bool bInTickableWhenPaused)
 {
-	Audio::FQuartzClockProxy* Proxy = FindProxyByName(ClockName);
-	if(Proxy && Proxy->IsValid())
-	{
-		Proxy->SendCommandToClock([=](Audio::FQuartzClock* InClock) mutable{ InClock->AddQuantizedCommand(InQuantizationCommandInitInfo); });
-	}
+	bTickEvenWhenPaused = bInTickableWhenPaused;
 }
 
 TWeakPtr<FQuartzTickableObjectsManager> UQuartzSubsystem::GetTickableObjectManager() const
diff --git a/Engine/Source/Runtime/AudioMixer/Public/Quartz/QuartzSubsystem.h b/Engine/Source/Runtime/AudioMixer/Public/Quartz/QuartzSubsystem.h
--- a/Engine/Source/Runtime/AudioMixer/Public/Quartz/QuartzSubsystem.h
+++ b/Engine/Source/Runtime/AudioMixer/Public/Quartz/QuartzSubsystem.h
@@ -64,6 +64,7 @@ public:
 
 	//~ Begin FTickableGameObject Interface
 	AUDIOMIXER_API virtual void Tick(float DeltaTime) override;
+	AUDIOMIXER_API virtual bool IsTickableWhenPaused() const override;
 	AUDIOMIXER_API virtual bool IsTickable() const override;
 	AUDIOMIXER_API virtual TStatId GetStatId() const override;
 	//~ End FTickableGameObject Interface
@@ -88,19 +89,6 @@ public:
 	static AUDIOMIXER_API Audio::FQuartzQuantizedRequestData CreateRequestDataForSchedulePlaySound(UQuartzClockHandle* InClockHandle, const FOnQuartzCommandEventBP& InDelegate, const FQuartzQuantizationBoundary& InQuantizationBoundary);
 	static AUDIOMIXER_API Audio::FQuartzQuantizedRequestData CreateRequestDataForQuantizedNotify(UQuartzClockHandle* InClockHandle, const FQuartzQuantizationBoundary& InQuantizationBoundary, const FOnQuartzCommandEventBP& InDelegate, float InMsOffset = 0.f);
 
-	// DEPRECATED HELPERS: non-static versions of the above CreateDataFor...() functions
-	UE_DEPRECATED(5.1, "Use the static (CreateRequestDataFor) version of this function instead")
-	AUDIOMIXER_API Audio::FQuartzQuantizedRequestData CreateDataForTickRateChange(UQuartzClockHandle* InClockHandle, const FOnQuartzCommandEventBP& InDelegate, const Audio::FQuartzClockTickRate& InNewTickRate, const FQuartzQuantizationBoundary& InQuantizationBoundary);
-
-	UE_DEPRECATED(5.1, "Use the static (CreateRequestDataFor) version of this function instead")
-	AUDIOMIXER_API Audio::FQuartzQuantizedRequestData CreateDataForTransportReset(UQuartzClockHandle* InClockHandle, const FQuartzQuantizationBoundary& InQuantizationBoundary, const FOnQuartzCommandEventBP& InDelegate);
-	
-	UE_DEPRECATED(5.1, "Use the static (CreateRequestDataFor) version of this function instead")
-	AUDIOMIXER_API Audio::FQuartzQuantizedRequestData CreateDataForStartOtherClock(UQuartzClockHandle* InClockHandle, FName InClockToStart, const FQuartzQuantizationBoundary& InQuantizationBoundary, const FOnQuartzCommandEventBP& InDelegate);
-
-	UE_DEPRECATED(5.1, "Use the static (CreateRequestDataFor) version of this function instead")
-	AUDIOMIXER_API Audio::FQuartzQuantizedRequestData CreateDataDataForSchedulePlaySound(UQuartzClockHandle* InClockHandle, const FOnQuartzCommandEventBP& InDelegate, const FQuartzQuantizationBoundary& InQuantizationBoundary);
-
 	UFUNCTION(BlueprintCallable, Category = "Quartz Clock Handle")
 	AUDIOMIXER_API bool IsQuartzEnabled();
 
@@ -170,9 +158,9 @@ public:
 
 	UFUNCTION(BlueprintCallable, Category = "Quartz Subsystem", meta = (WorldContext = "WorldContextObject"))
 	AUDIOMIXER_API float GetRoundTripMaxLatency(const UObject* WorldContextObject);
-
-	UE_DEPRECATED(5.1, "Obtain and use a UQuartzClockHandle / FQuartzClockProxy instead")
-	AUDIOMIXER_API void AddCommandToClock(const UObject* WorldContextObject, Audio::FQuartzQuantizedCommandInitInfo& InQuantizationCommandInitInfo, FName ClockName);
+	
+	UFUNCTION(BlueprintCallable, Category = "Quartz Subsystem")
+	AUDIOMIXER_API void SetQuartzSubsystemTickableWhenPaused(const bool bInTickableWhenPaused);
 
 	// sharable to allow non-UObjects to un-subscribe if the Subsystem is going to outlive them
 	AUDIOMIXER_API TWeakPtr<FQuartzTickableObjectsManager> GetTickableObjectManager() const;
@@ -189,6 +177,8 @@ private:
 	// Clock manager/proxy-related data that lives on the AudioDevice for persistence.
 	TSharedPtr<Audio::FPersistentQuartzSubsystemData> ClockManagerDataPtr { nullptr };
 
+	bool bTickEvenWhenPaused = false;
+
 	// helpers
 	AUDIOMIXER_API Audio::FQuartzClockProxy* FindProxyByName(const FName& ClockName);
 	AUDIOMIXER_API Audio::FQuartzClockProxy const* FindProxyByName(const FName& ClockName) const;
