Fix predistanceattenuation submix send access bad data.

diff --git a/Engine/Source/Runtime/AudioMixer/Private/AudioMixerSourceManager.cpp b/Engine/Source/Runtime/AudioMixer/Private/AudioMixerSourceManager.cpp
--- a/Engine/Source/Runtime/AudioMixer/Private/AudioMixerSourceManager.cpp
+++ b/Engine/Source/Runtime/AudioMixer/Private/AudioMixerSourceManager.cpp
@@ -1770,21 +1770,20 @@ namespace Audio
 			FMixerSubmixPtr InSubmixPtr = InSubmixSend.Submix.Pin();
 			if (InSubmixPtr.IsValid())
 			{
+				// Determine whether submix send is new and whether any sends have 
+				// a pre-distance-attenuation send.
 				bool bIsNew = true;
-				
-				SourceInfo.bHasPreDistanceAttenuationSend = false;
+				SourceInfo.bHasPreDistanceAttenuationSend = InSubmixSend.SubmixSendStage == EMixerSourceSubmixSendStage::PreDistanceAttenuation;
+
 				for (FMixerSourceSubmixSend& SubmixSend : SourceInfo.SubmixSends)
 				{
 					FMixerSubmixPtr SubmixPtr = SubmixSend.Submix.Pin();
+
 					if (SubmixPtr.IsValid())
 					{
-						if (SubmixSend.SubmixSendStage == EMixerSourceSubmixSendStage::PreDistanceAttenuation)
-						{
-							SourceInfo.bHasPreDistanceAttenuationSend = true;
-						}
-					
 						if (SubmixPtr->GetId() == InSubmixPtr->GetId())
 						{
+							// Update existing submix send if it already exists
 							SubmixSend.SendLevel = InSubmixSend.SendLevel;
 							SubmixSend.SubmixSendStage = InSubmixSend.SubmixSendStage;
 							bIsNew = false;
@@ -1793,6 +1792,11 @@ namespace Audio
 								break;
 							}
 						}
+
+						if (SubmixSend.SubmixSendStage == EMixerSourceSubmixSendStage::PreDistanceAttenuation)
+						{
+							SourceInfo.bHasPreDistanceAttenuationSend = true;
+						}
 					}
 				}
 
