Remove deprecated submix code

diff --git a/Engine/Source/Runtime/AudioMixer/Private/AudioMixerDevice.cpp b/Engine/Source/Runtime/AudioMixer/Private/AudioMixerDevice.cpp
--- a/Engine/Source/Runtime/AudioMixer/Private/AudioMixerDevice.cpp
+++ b/Engine/Source/Runtime/AudioMixer/Private/AudioMixerDevice.cpp
@@ -1481,7 +1481,7 @@ namespace Audio
 		FMixerSubmixPtr MixerSubmix = GetSubmixInstance(InSoundSubmix).Pin();
 		if (MixerSubmix.IsValid())
 		{
-			const float NewVolume = CastedSubmix->OutputVolume;
+			const float NewVolume = CastedSubmix->OutputVolumeModulation.Value;
 			AudioRenderThreadCommand([MixerSubmix, NewVolume]()
 			{
 				MixerSubmix->SetOutputVolume(NewVolume);
@@ -1987,14 +1987,6 @@ namespace Audio
 		Submixes.Remove(InSoundSubmix.GetUniqueID());
 	}
 
-	void FMixerDevice::InitSoundEffectPresets()
-	{
-#if WITH_EDITOR
-		IAudioEditorModule* AudioEditorModule = &FModuleManager::LoadModuleChecked<IAudioEditorModule>("AudioEditor");
-		AudioEditorModule->RegisterEffectPresetAssetActions();
-#endif
-	}
-
 	FMixerSubmixPtr FMixerDevice::FindSubmixInstanceByObjectId(uint32 InObjectId)
 	{
 		for (int32 i = 0; i < RequiredSubmixes.Num(); i++)
@@ -2023,14 +2015,6 @@ namespace Audio
 		return Submixes.FindRef(InObjectId);
 	}
 
-	void FMixerDevice::InitDefaultAudioBuses()
-	{
-	}
-
-	void FMixerDevice::ShutdownDefaultAudioBuses()
-	{
-	}
-
 	FMixerSubmixWeakPtr FMixerDevice::GetSubmixInstance(const USoundSubmixBase* SoundSubmix) const
 	{
 		LLM_SCOPE(ELLMTag::AudioMixer);
@@ -2795,66 +2779,6 @@ namespace Audio
 		QuantizedEventClockManager.Flush();
 	}
 
-	void FMixerDevice::StartAudioBus(uint32 InAudioBusId, int32 InNumChannels, bool bInIsAutomatic)
-	{
-		UAudioBusSubsystem* AudioBusSubsystem = GetSubsystem<UAudioBusSubsystem>();
-		if (AudioBusSubsystem)
-		{
-			AudioBusSubsystem->StartAudioBus(FAudioBusKey(InAudioBusId), InNumChannels, bInIsAutomatic);
-		}
-	}
-
-	void FMixerDevice::StopAudioBus(uint32 InAudioBusId)
-	{
-		UAudioBusSubsystem* AudioBusSubsystem = GetSubsystem<UAudioBusSubsystem>();
-		if (AudioBusSubsystem)
-		{
-			AudioBusSubsystem->StopAudioBus(FAudioBusKey(InAudioBusId));
-		}
-	}
-
-	bool FMixerDevice::IsAudioBusActive(uint32 InAudioBusId) const
-	{
-		UAudioBusSubsystem* AudioBusSubsystem = GetSubsystem<UAudioBusSubsystem>();
-		if (AudioBusSubsystem)
-		{
-			return AudioBusSubsystem->IsAudioBusActive(FAudioBusKey(InAudioBusId));
-		}
-		return false;
-	}
-
-	//UE_DEPRECATED(5.2, "AddPatchForAudioBus is deprecated.  Use UAudioBusSubsystem::AddPatchOutputForAudioBus.")
-	FPatchOutputStrongPtr FMixerDevice::AddPatchForAudioBus(uint32 InAudioBusId, float InPatchGain)
-	{
-		UAudioBusSubsystem* AudioBusSubsystem = GetSubsystem<UAudioBusSubsystem>();
-		if (AudioBusSubsystem)
-		{
-			return AudioBusSubsystem->AddPatchOutputForAudioBus(FAudioBusKey(InAudioBusId), SourceManager->GetNumOutputFrames(), 8, InPatchGain);
-		}
-		return nullptr;
-	}
-
-	//UE_DEPRECATED(5.2, "AddPatchForAudioBus_GameThread is deprecated.  Use UAudioBusSubsystem::AddPatchOutputForAudioBus.")
-	FPatchOutputStrongPtr FMixerDevice::AddPatchForAudioBus_GameThread(uint32 InAudioBusId, float InPatchGain)
-	{
-		PRAGMA_DISABLE_DEPRECATION_WARNINGS
-		return FMixerDevice::AddPatchForAudioBus(InAudioBusId, InPatchGain);
-		PRAGMA_ENABLE_DEPRECATION_WARNINGS
-	}
-
-	//UE_DEPRECATED(5.2, "This overload of AddPatchInputForAudioBus is deprecated and non-functional.  Use the overload that takes the number of frames and channels as parameters.")
-	void FMixerDevice::AddPatchInputForAudioBus(const FPatchInput& InPatchInput, uint32 InAudioBusId, float InPatchGain)
-	{
-	}
-
-	//UE_DEPRECATED(5.2, "AddPatchInputForAudioBus_GameThread is deprecated.  Use UAudioBusSubsystem::AddPatchInputForAudioBus.")
-	void FMixerDevice::AddPatchInputForAudioBus_GameThread(const FPatchInput& InPatchInput, uint32 InAudioBusId, float InPatchGain)
-	{
-		PRAGMA_DISABLE_DEPRECATION_WARNINGS
-		FMixerDevice::AddPatchInputForAudioBus(InPatchInput, InAudioBusId, InPatchGain);
-		PRAGMA_ENABLE_DEPRECATION_WARNINGS
-	}
-
 	FPatchOutputStrongPtr FMixerDevice::MakePatch(int32 InFrames, int32 InChannels, float InGain) const
 	{
 		// Assume the mixer will consume SourceManager->GetNumOutputFrames() per iteration and an input patch will generate InFrames per iteration.
diff --git a/Engine/Source/Runtime/AudioMixer/Public/AudioMixerDevice.h b/Engine/Source/Runtime/AudioMixer/Public/AudioMixerDevice.h
--- a/Engine/Source/Runtime/AudioMixer/Public/AudioMixerDevice.h
+++ b/Engine/Source/Runtime/AudioMixer/Public/AudioMixerDevice.h
@@ -148,11 +148,6 @@ namespace Audio
 		AUDIOMIXER_API virtual void RegisterSoundSubmix(USoundSubmixBase* SoundSubmix, bool bInit = true) override;
 		AUDIOMIXER_API virtual void UnregisterSoundSubmix(const USoundSubmixBase* SoundSubmix, const bool bReparentChildren) override;
 
-		AUDIOMIXER_API virtual void InitSoundEffectPresets() override;
-		UE_DEPRECATED(5.2, "The functionality for this has been moved to UAudioBusSubsystem::InitDefaultAudioBuses, which is now automatically called on subsystem creation.")
-		AUDIOMIXER_API virtual void InitDefaultAudioBuses() override;
-		UE_DEPRECATED(5.2, "The functionality for this has been moved to UAudioBusSubsystem::ShutdownDefaultAudioBuses.")
-		AUDIOMIXER_API virtual void ShutdownDefaultAudioBuses() override;
 		AUDIOMIXER_API virtual int32 GetNumActiveSources() const override;
 
 		// Updates the source effect chain (using unique object id). 
@@ -337,26 +332,6 @@ namespace Audio
 
 		static AUDIOMIXER_API bool IsEndpointSubmix(const USoundSubmixBase* InSubmix);
 
-		// Audio bus API - these are deprecated. Use corresponding calls in UAudioBusSubsystem instead
-		UE_DEPRECATED(5.2, "This function is deprecated. Use UAudioBusSubsystem::StartAudioBus instead.")
-		AUDIOMIXER_API virtual void StartAudioBus(uint32 InAudioBusId, int32 InNumChannels, bool bInIsAutomatic) override;
-		UE_DEPRECATED(5.2, "This function is deprecated. Use UAudioBusSubsystem::StopAudioBus instead.")
-		AUDIOMIXER_API virtual void StopAudioBus(uint32 InAudioBusId) override;
-		UE_DEPRECATED(5.2, "This function is deprecated. Use UAudioBusSubsystem::IsAudioBusActive instead.")
-		AUDIOMIXER_API virtual bool IsAudioBusActive(uint32 InAudioBusId) const override;
-
-		UE_DEPRECATED(5.2, "AddPatchForAudioBus is deprecated.  Use UAudioBusSubsystem::AddPatchOutputForAudioBus.")
-		AUDIOMIXER_API virtual FPatchOutputStrongPtr AddPatchForAudioBus(uint32 InAudioBusId, float InPatchGain = 1.0f) override;
-
-		UE_DEPRECATED(5.2, "AddPatchForAudioBus_GameThread is deprecated.  Use UAudioBusSubsystem::AddPatchOutputForAudioBus.")
-		AUDIOMIXER_API virtual FPatchOutputStrongPtr AddPatchForAudioBus_GameThread(uint32 InAudioBusId, float InPatchGain = 1.0f) override;
-
-		UE_DEPRECATED(5.2, "This overload of AddPatchInputForAudioBus is deprecated and non-functional.  Use the overload in UAudioBusSubsystem that takes the number of frames and channels as parameters.")
-		AUDIOMIXER_API virtual void AddPatchInputForAudioBus(const FPatchInput& InPatchInput, uint32 InAudioBusId, float InPatchGain = 1.0f) override;
-
-		UE_DEPRECATED(5.2, "AddPatchInputForAudioBus_GameThread is deprecated.  Use UAudioBusSubsystem::AddPatchInputForAudioBus instead.")
-		AUDIOMIXER_API virtual void AddPatchInputForAudioBus_GameThread(const FPatchInput& InPatchInput, uint32 InAudioBusId, float InPatchGain = 1.0f) override;
-
 		AUDIOMIXER_API FPatchOutputStrongPtr MakePatch(int32 InFrames, int32 InChannels, float InGain) const;
 
 		// Clock Manager for quantized event handling on Audio Render Thread
diff --git a/Engine/Source/Runtime/Engine/Classes/Sound/SoundSubmix.h b/Engine/Source/Runtime/Engine/Classes/Sound/SoundSubmix.h
--- a/Engine/Source/Runtime/Engine/Classes/Sound/SoundSubmix.h
+++ b/Engine/Source/Runtime/Engine/Classes/Sound/SoundSubmix.h
@@ -338,18 +338,6 @@ public:
 	UPROPERTY(EditAnywhere, BlueprintReadWrite, Category = EnvelopeFollower, meta = (ClampMin = "0", UIMin = "0"))
 	int32 EnvelopeFollowerReleaseTime;
 
-	/** Deprecated -- The output volume of the submix. Applied after submix effects and analysis are performed.*/
-	UPROPERTY(meta = (DeprecatedProperty, DeprecationMessage = "5.0 - Removed in favor of OutputVolumeModulation."))
-	float OutputVolume;
-
-	/** Deprecated -- The wet level of the submix. Applied after submix effects and analysis are performed. */
-	UPROPERTY(meta = (DeprecatedProperty, DeprecationMessage = "5.0 - Removed in favor of WetLevelModulation."))
-	float WetLevel;
-
-	/** Deprecated -- The dry level of the submix. Applied before submix effects and analysis are performed. */
-	UPROPERTY(meta = (DeprecatedProperty, DeprecationMessage = "5.0 - Removed in favor of DryLevelModulation."))
-	float DryLevel;
-
 	/** The output volume of the submix in Decibels. Applied after submix effects and analysis are performed.*/
 	UPROPERTY(EditAnywhere, BlueprintSetter=SetOutputVolumeModulation, Category = SubmixLevel, meta = (DisplayName = "Output Volume (dB)", AudioParam = "Volume", AudioParamClass = "SoundModulationParameterVolume", ClampMin = "-96.0", ClampMax = "0.0", UIMin = "-96.0", UIMax = "0.0"))
 	FSoundModulationDestinationSettings OutputVolumeModulation;
diff --git a/Engine/Source/Runtime/Engine/Private/AudioDevice.cpp b/Engine/Source/Runtime/Engine/Private/AudioDevice.cpp
--- a/Engine/Source/Runtime/Engine/Private/AudioDevice.cpp
+++ b/Engine/Source/Runtime/Engine/Private/AudioDevice.cpp
@@ -537,7 +537,6 @@ bool FAudioDevice::Init(Audio::FDeviceId InDeviceID, int32 InMaxSources, int32 I
 
 	// Parses sound classes.
 	InitSoundClasses();
-	InitSoundEffectPresets();
 
 	// Audio mixer needs to create effects manager before initializing the plugins.
 	if (IsStoppingVoicesEnabled())
@@ -673,8 +672,6 @@ bool FAudioDevice::Init(Audio::FDeviceId InDeviceID, int32 InMaxSources, int32 I
 	FCoreUObjectDelegates::GetPreGarbageCollectDelegate().AddRaw(this, &FAudioDevice::OnPreGarbageCollect);
 	FCoreUObjectDelegates::PreGarbageCollectConditionalBeginDestroy.AddRaw(this, &FAudioDevice::OnPreGarbageCollect);
 
-	InitDefaultAudioBuses();
-
 	UE_LOG(LogInit, Log, TEXT("FAudioDevice initialized with ID %d."), InDeviceID);
 
 	return true;
@@ -924,8 +921,6 @@ void FAudioDevice::Teardown()
 	checkf(!bDidRemoveDelegate, TEXT("DelegateHandle not removed. Should be removed during FAudioDevice::OnDeviceDestroyed(...)"));
 #endif
 
-	ShutdownDefaultAudioBuses();
-
 	// Make sure we process any pending game thread tasks before tearing down the audio device.
 	FTaskGraphInterface::Get().ProcessThreadUntilIdle(ENamedThreads::GameThread);
 
@@ -6864,16 +6859,6 @@ Audio::FPatchOutputStrongPtr FAudioDevice::AddPatchForSubmix(uint32 InObjectId,
 	return nullptr;
 }
 
-Audio::FPatchOutputStrongPtr FAudioDevice::AddPatchForAudioBus(uint32 InAudioBusId, float InPatchGain)
-{
-	return nullptr;
-}
-
-Audio::FPatchOutputStrongPtr FAudioDevice::AddPatchForAudioBus_GameThread(uint32 InAudioBusId, float InPatchGain)
-{
-	return nullptr;
-}
-
 Audio::FPatchInput FAudioDevice::AddPatchInputForAudioBus(uint32 InAudioBusId, int32 InFrames, int32 InChannels, float InGain)
 {
 	return Audio::FPatchInput();
diff --git a/Engine/Source/Runtime/Engine/Private/AudioDeviceManager.cpp b/Engine/Source/Runtime/Engine/Private/AudioDeviceManager.cpp
--- a/Engine/Source/Runtime/Engine/Private/AudioDeviceManager.cpp
+++ b/Engine/Source/Runtime/Engine/Private/AudioDeviceManager.cpp
@@ -843,12 +843,7 @@ void FAudioDeviceManager::InitSoundSubmixes()
 
 void FAudioDeviceManager::InitSoundEffectPresets()
 {
-	IterateOverAllDevices(
-		[](Audio::FDeviceId, FAudioDevice* InDevice)
-		{
-			InDevice->InitSoundEffectPresets();
-		}
-	);
+	// Deprecated.
 }
 
 void FAudioDeviceManager::UpdateSourceEffectChain(const uint32 SourceEffectChainId, const TArray<FSourceEffectChainEntry>& SourceEffectChain, const bool bPlayEffectChainTails)
diff --git a/Engine/Source/Runtime/Engine/Private/SoundSubmix.cpp b/Engine/Source/Runtime/Engine/Private/SoundSubmix.cpp
--- a/Engine/Source/Runtime/Engine/Private/SoundSubmix.cpp
+++ b/Engine/Source/Runtime/Engine/Private/SoundSubmix.cpp
@@ -44,9 +44,6 @@ USoundSubmix::USoundSubmix(const FObjectInitializer& ObjectInitializer)
 	, AmbisonicsPluginSettings(nullptr)
 	, EnvelopeFollowerAttackTime(10)
 	, EnvelopeFollowerReleaseTime(500)
-	, OutputVolume(-1.0f)
-	, WetLevel(-1.0f)
-	, DryLevel(-1.0f)
 {
 	OutputVolumeModulation.Value = 0.f;
 	WetLevelModulation.Value = 0.f;
@@ -63,46 +60,6 @@ void USoundSubmix::Serialize(FArchive& Ar)
 		// use -96dB as a noise floor when fixing up linear volume settings
 		static constexpr float LinearNeg96dB = 0.0000158489319f;
 
-		// convert any old deprecated values to the new value
-		if (OutputVolume >= 0.0f)
-		{
-			if (OutputVolume <= LinearNeg96dB)
-			{
-				OutputVolumeModulation.Value = -96.f;
-			}
-			else
-			{
-				OutputVolumeModulation.Value = Audio::ConvertToDecibels(OutputVolume);
-			}
-			OutputVolume = -1.0f;
-		}
-
-		if (WetLevel >= 0.0f)
-		{
-			if (WetLevel <= LinearNeg96dB)
-			{
-				WetLevelModulation.Value = -96.f;
-			}
-			else
-			{
-				WetLevelModulation.Value = Audio::ConvertToDecibels(WetLevel);
-			}
-			WetLevel = -1.0f;
-		}
-
-		if (DryLevel >= 0.0f)
-		{
-			if (DryLevel <= LinearNeg96dB)
-			{
-				DryLevelModulation.Value = -96.f;
-			}
-			else
-			{
-				DryLevelModulation.Value = Audio::ConvertToDecibels(DryLevel);
-			}
-			DryLevel = -1.0f;
-		}
-
 		// fix values previously saved as linear values
 		if (OutputVolumeModulation.Value > 0.0f)
 		{
diff --git a/Engine/Source/Runtime/Engine/Public/AudioDevice.h b/Engine/Source/Runtime/Engine/Public/AudioDevice.h
--- a/Engine/Source/Runtime/Engine/Public/AudioDevice.h
+++ b/Engine/Source/Runtime/Engine/Public/AudioDevice.h
@@ -943,12 +943,6 @@ public:
 	*/
 	ENGINE_API void UnregisterSoundClass(USoundClass* SoundClass);
 
-	/* Initialized audio buses marked as default that are to be enabled for the duration of the application. */
-	virtual void InitDefaultAudioBuses() {}
-
-	/* Shutdown all audio buses marked as default. */
-	virtual void ShutdownDefaultAudioBuses() {}
-
 	/** Initializes sound submixes. */
 	virtual void InitSoundSubmixes() {}
 
@@ -979,41 +973,10 @@ public:
 
 	ENGINE_API virtual Audio::FPatchOutputStrongPtr AddPatchForSubmix(uint32 InObjectId, float InPatchGain);
 
-	virtual void StartAudioBus(uint32 InAudioBusId, int32 InNumChannels, bool bInIsAutomatic)
-	{
-	}
-
-	virtual void StopAudioBus(uint32 InAudioBusId)
-	{
-	}
-
-	virtual bool IsAudioBusActive(uint32 InAudioBusId) const
-	{
-		return false;
-	}
-
-	UE_DEPRECATED(5.2, "AddPatchForAudioBus is deprecated.  Use AddPatchOutputForAudioBus.")
-	ENGINE_API virtual Audio::FPatchOutputStrongPtr AddPatchForAudioBus(uint32 InAudioBusId, float InPatchGain = 1.0f);
-
-	UE_DEPRECATED(5.2, "AddPatchForAudioBus_GameThread is deprecated.  Use AddPatchOutputForAudioBus.")
-	ENGINE_API virtual Audio::FPatchOutputStrongPtr AddPatchForAudioBus_GameThread(uint32 InAudioBusId, float InPatchGain = 1.0f);
-
-	UE_DEPRECATED(5.2, "This overload of AddPatchInputForAudioBus is deprecated.  Use the overload that takes the number of frames and channels as parameters.")
-	virtual void AddPatchInputForAudioBus(const Audio::FPatchInput& InPatchInput, uint32 InAudioBusId, float InPatchGain = 1.0f)
-	{
-	}
-
-	UE_DEPRECATED(5.2, "AddPatchInputForAudioBus_GameThread is deprecated.  Use AddPatchInputForAudioBus.")
-	virtual void AddPatchInputForAudioBus_GameThread(const Audio::FPatchInput& InPatchInput, uint32 InAudioBusId, float InPatchGain = 1.0f)
-	{
-	}
-
 	ENGINE_API virtual Audio::FPatchInput AddPatchInputForAudioBus(uint32 InAudioBusId, int32 InFrames, int32 InChannels, float InGain = 1.f);
 
 	ENGINE_API virtual Audio::FPatchOutputStrongPtr AddPatchOutputForAudioBus(uint32 InAudioBusId, int32 InFrames, int32 InChannels, float InGain = 1.f);
 
-	virtual void InitSoundEffectPresets() {}
-
 	/**
 	* Gets the current properties of a sound class, if the sound class hasn't been registered, then it returns nullptr
 	*
diff --git a/Engine/Source/Runtime/Engine/Public/AudioDeviceManager.h b/Engine/Source/Runtime/Engine/Public/AudioDeviceManager.h
--- a/Engine/Source/Runtime/Engine/Public/AudioDeviceManager.h
+++ b/Engine/Source/Runtime/Engine/Public/AudioDeviceManager.h
@@ -205,6 +205,7 @@ public:
 	ENGINE_API void InitSoundSubmixes();
 
 	/** Initialize all sound effect presets. */
+	UE_DEPRECATED(5.4, "Will be removed in upcoming versions of this code")
 	ENGINE_API void InitSoundEffectPresets();
 
 	/** Updates source effect chain on all sources currently using the source effect chain. */
