Fix Resonance Double Attenuation due to Manual Modulation at Mixer Source

diff --git a/Engine/Plugins/Runtime/ResonanceAudio/Source/ResonanceAudio/Private/ResonanceAudioModule.h b/Engine/Plugins/Runtime/ResonanceAudio/Source/ResonanceAudio/Private/ResonanceAudioModule.h
--- a/Engine/Plugins/Runtime/ResonanceAudio/Source/ResonanceAudio/Private/ResonanceAudioModule.h
+++ b/Engine/Plugins/Runtime/ResonanceAudio/Source/ResonanceAudio/Private/ResonanceAudioModule.h
@@ -41,6 +41,8 @@ namespace ResonanceAudio
 			return 2;
 		}
 
+		virtual bool ReturnsToSubmixGraph() const override { return true; }
+
 	};
 
 	/******************************************************/
@@ -63,6 +65,7 @@ namespace ResonanceAudio
 
 		virtual TAudioReverbPtr CreateNewReverbPlugin(FAudioDevice* OwningDevice) override;
 		virtual bool IsExternalSend() override { return true; }
+		virtual bool ReturnsToSubmixGraph() const override { return true; }
 	};
 
 	/*********************************************************/
diff --git a/Engine/Source/Runtime/AudioExtensions/Public/IAudioExtensionPlugin.h b/Engine/Source/Runtime/AudioExtensions/Public/IAudioExtensionPlugin.h
--- a/Engine/Source/Runtime/AudioExtensions/Public/IAudioExtensionPlugin.h
+++ b/Engine/Source/Runtime/AudioExtensions/Public/IAudioExtensionPlugin.h
@@ -227,6 +227,14 @@ public:
 	{
 		return false;
 	}
+
+	/*
+	*  @return true if the plugin returns from its external submix to a submix in the plugin.
+	*/
+	virtual bool ReturnsToSubmixGraph() const
+	{
+		return false;
+	}
 };
 
 /************************************************************************/
diff --git a/Engine/Source/Runtime/AudioMixer/Private/AudioMixerSource.cpp b/Engine/Source/Runtime/AudioMixer/Private/AudioMixerSource.cpp
--- a/Engine/Source/Runtime/AudioMixer/Private/AudioMixerSource.cpp
+++ b/Engine/Source/Runtime/AudioMixer/Private/AudioMixerSource.cpp
@@ -1642,6 +1642,29 @@ namespace Audio
 
 	float FMixerSource::GetInheritedSubmixVolumeModulation() const
 	{
+		FAudioDevice::FAudioSpatializationInterfaceInfo SpatializationInfo = MixerDevice->GetCurrentSpatializationPluginInterfaceInfo();
+		check(SpatializationInfo.bSpatializationIsExternalSend);
+
+		// if there is a return submix, we need to figure out where to stop manually attenuating
+		// Because the submix will modulate itself later
+		// Since the graph has tree-like structure, we can create a list of the return submix's ancestors
+		// to use while traversing the other submix's ancestors
+		TArray<uint32> ReturnSubmixAncestors;
+		if (SpatializationInfo.bReturnsToSubmixGraph)
+		{
+			USoundSubmix* ReturnSubmix = MixerDevice->ReverbPluginInterface->GetSubmix();
+			FMixerSubmixWeakPtr CurrReturnSubmixWeakPtr = MixerDevice->GetSubmixInstance(ReturnSubmix);
+			FMixerSubmixPtr CurrReturnSubmixPtr = CurrReturnSubmixWeakPtr.Pin();
+			while (CurrReturnSubmixPtr && CurrReturnSubmixPtr->IsValid())
+			{
+				ReturnSubmixAncestors.Add(CurrReturnSubmixPtr->GetId());
+
+				CurrReturnSubmixWeakPtr = CurrReturnSubmixPtr->GetParent();
+				CurrReturnSubmixPtr = CurrReturnSubmixWeakPtr.Pin();
+			}
+
+		}
+
 		float SubmixModVolume = 1.0f;
 
 		FMixerSubmixWeakPtr CurrSubmixWeakPtr = MixerDevice->GetSubmixInstance(WaveInstance->SoundSubmix);
@@ -1649,6 +1672,13 @@ namespace Audio
 		// Check the submix and all its parents in the graph for active modulation
 		while (CurrSubmixPtr && CurrSubmixPtr->IsValid())
 		{
+			// Matching ID means the external spatializer has returned to the submix graph at this point,
+			// so we no longer need to manually apply volume modulation
+			if (SpatializationInfo.bReturnsToSubmixGraph && ReturnSubmixAncestors.Contains(CurrSubmixPtr->GetId()))
+			{
+				break;
+			}
+
 			FModulationDestination* SubmixOutVolDest = CurrSubmixPtr->GetOutputVolumeDestination();
 			FModulationDestination* SubmixWetVolDest = CurrSubmixPtr->GetWetVolumeDestination();
 			if (SubmixOutVolDest)
diff --git a/Engine/Source/Runtime/Engine/Private/AudioDevice.cpp b/Engine/Source/Runtime/Engine/Private/AudioDevice.cpp
--- a/Engine/Source/Runtime/Engine/Private/AudioDevice.cpp
+++ b/Engine/Source/Runtime/Engine/Private/AudioDevice.cpp
@@ -7357,6 +7357,7 @@ FAudioDevice::FAudioSpatializationInterfaceInfo::FAudioSpatializationInterfaceIn
 	// cache metadata from the incoming plugin interface 
 	bSpatializationIsExternalSend = InAudioSpatializationFactoryPtr->IsExternalSend();
 	MaxChannelsSupportedBySpatializationPlugin = InAudioSpatializationFactoryPtr->GetMaxSupportedChannels();
+	bReturnsToSubmixGraph = InAudioSpatializationFactoryPtr->ReturnsToSubmixGraph();
 }
 
 bool FAudioDevice::FAudioSpatializationInterfaceInfo::IsValid() const
diff --git a/Engine/Source/Runtime/Engine/Public/AudioDevice.h b/Engine/Source/Runtime/Engine/Public/AudioDevice.h
--- a/Engine/Source/Runtime/Engine/Public/AudioDevice.h
+++ b/Engine/Source/Runtime/Engine/Public/AudioDevice.h
@@ -2008,6 +2008,7 @@ public:
 		int32 MaxChannelsSupportedBySpatializationPlugin = 1;
 		uint8 bSpatializationIsExternalSend:1;
 		uint8 bIsInitialized:1;
+		uint8 bReturnsToSubmixGraph:1;
 	};
 
 	UE_DEPRECATED(5.1, "Do not access this member directly, it is not used. Call GetSpatializationPluginInterface() instead.")
