BaseDefaultSubmix not working in all situations

diff --git a/Engine/Source/Runtime/AudioMixer/Private/AudioMixerSource.cpp b/Engine/Source/Runtime/AudioMixer/Private/AudioMixerSource.cpp
--- a/Engine/Source/Runtime/AudioMixer/Private/AudioMixerSource.cpp
+++ b/Engine/Source/Runtime/AudioMixer/Private/AudioMixerSource.cpp
@@ -1864,9 +1864,9 @@ namespace Audio
 			{
 				SubmixPtr = MixerDevice->GetSubmixInstance(WaveInstance->SoundSubmix);
 			}
-			else if (!WaveInstance->bIsDynamic) // Dynamic submixes don't auto connect to master.
+			else if (!WaveInstance->bIsDynamic) // Dynamic submixes don't auto connect.
 			{
-				SubmixPtr = MixerDevice->GetMasterSubmix();
+				SubmixPtr = MixerDevice->GetBaseDefaultSubmix(); // This will try base default and fall back to master if that fails.
 			}
 
 
