Move Resonance submix load to its own function called during engine init

diff --git a/Engine/Plugins/Runtime/ResonanceAudio/Source/ResonanceAudio/Private/ResonanceAudioReverb.cpp b/Engine/Plugins/Runtime/ResonanceAudio/Source/ResonanceAudio/Private/ResonanceAudioReverb.cpp
--- a/Engine/Plugins/Runtime/ResonanceAudio/Source/ResonanceAudio/Private/ResonanceAudioReverb.cpp
+++ b/Engine/Plugins/Runtime/ResonanceAudio/Source/ResonanceAudio/Private/ResonanceAudioReverb.cpp
@@ -186,7 +186,11 @@ namespace ResonanceAudio
 
 	USoundSubmix* FResonanceAudioReverb::GetSubmix()
 	{
+		return ReverbSubmixPtr;
+	}
 
+	USoundSubmix* FResonanceAudioReverb::LoadSubmix()
+	{
 		if (ReverbSubmixPtr)
 		{
 			return ReverbSubmixPtr;
diff --git a/Engine/Plugins/Runtime/ResonanceAudio/Source/ResonanceAudio/Private/ResonanceAudioReverb.h b/Engine/Plugins/Runtime/ResonanceAudio/Source/ResonanceAudio/Private/ResonanceAudioReverb.h
--- a/Engine/Plugins/Runtime/ResonanceAudio/Source/ResonanceAudio/Private/ResonanceAudioReverb.h
+++ b/Engine/Plugins/Runtime/ResonanceAudio/Source/ResonanceAudio/Private/ResonanceAudioReverb.h
@@ -114,6 +114,7 @@ namespace ResonanceAudio
 		virtual void OnReleaseSource(const uint32 SourceId) override;
 		virtual FSoundEffectSubmixPtr GetEffectSubmix() override;
 		virtual USoundSubmix* GetSubmix() override;
+		virtual USoundSubmix* LoadSubmix() override;
 		virtual void ProcessSourceAudio(const FAudioPluginSourceInputData& InputData, FAudioPluginSourceOutputData& OutputData) override;
 
 		void SetResonanceAudioApi(vraudio::ResonanceAudioApi* InResonanceAudioApi) { ResonanceAudioApi = InResonanceAudioApi; };
diff --git a/Engine/Source/Runtime/AudioExtensions/Public/IAudioExtensionPlugin.h b/Engine/Source/Runtime/AudioExtensions/Public/IAudioExtensionPlugin.h
--- a/Engine/Source/Runtime/AudioExtensions/Public/IAudioExtensionPlugin.h
+++ b/Engine/Source/Runtime/AudioExtensions/Public/IAudioExtensionPlugin.h
@@ -662,6 +662,11 @@ public:
 	/** Returns the plugin-managed effect submix instance */
 	virtual FSoundEffectSubmixPtr GetEffectSubmix() = 0;
 
+	virtual USoundSubmix* LoadSubmix()
+	{
+		return GetSubmix();
+	}
+
 	/** Returns the plugin-managed effect submix */
 	virtual USoundSubmix* GetSubmix() = 0;
 
diff --git a/Engine/Source/Runtime/AudioMixer/Private/AudioMixerDevice.cpp b/Engine/Source/Runtime/AudioMixer/Private/AudioMixerDevice.cpp
--- a/Engine/Source/Runtime/AudioMixer/Private/AudioMixerDevice.cpp
+++ b/Engine/Source/Runtime/AudioMixer/Private/AudioMixerDevice.cpp
@@ -1097,7 +1097,7 @@ namespace Audio
 		if (IsReverbPluginEnabled() && ReverbPluginInterface)
 		{
 			LLM_SCOPE(ELLMTag::AudioMixerPlugins);
-			USoundSubmix* ReverbPluginSubmix = ReverbPluginInterface->GetSubmix();
+			USoundSubmix* ReverbPluginSubmix = ReverbPluginInterface->LoadSubmix();
 			check(ReverbPluginSubmix);
 			ReverbPluginSubmix->AddToRoot();
 
diff --git a/Engine/Source/Runtime/AudioMixer/Private/AudioMixerSource.cpp b/Engine/Source/Runtime/AudioMixer/Private/AudioMixerSource.cpp
--- a/Engine/Source/Runtime/AudioMixer/Private/AudioMixerSource.cpp
+++ b/Engine/Source/Runtime/AudioMixer/Private/AudioMixerSource.cpp
@@ -1653,14 +1653,17 @@ namespace Audio
 		if (SpatializationInfo.bReturnsToSubmixGraph)
 		{
 			USoundSubmix* ReturnSubmix = MixerDevice->ReverbPluginInterface->GetSubmix();
-			FMixerSubmixWeakPtr CurrReturnSubmixWeakPtr = MixerDevice->GetSubmixInstance(ReturnSubmix);
-			FMixerSubmixPtr CurrReturnSubmixPtr = CurrReturnSubmixWeakPtr.Pin();
-			while (CurrReturnSubmixPtr && CurrReturnSubmixPtr->IsValid())
+			if (ReturnSubmix)
 			{
-				ReturnSubmixAncestors.Add(CurrReturnSubmixPtr->GetId());
+				FMixerSubmixWeakPtr CurrReturnSubmixWeakPtr = MixerDevice->GetSubmixInstance(ReturnSubmix);
+				FMixerSubmixPtr CurrReturnSubmixPtr = CurrReturnSubmixWeakPtr.Pin();
+				while (CurrReturnSubmixPtr && CurrReturnSubmixPtr->IsValid())
+				{
+					ReturnSubmixAncestors.Add(CurrReturnSubmixPtr->GetId());
 
-				CurrReturnSubmixWeakPtr = CurrReturnSubmixPtr->GetParent();
-				CurrReturnSubmixPtr = CurrReturnSubmixWeakPtr.Pin();
+					CurrReturnSubmixWeakPtr = CurrReturnSubmixPtr->GetParent();
+					CurrReturnSubmixPtr = CurrReturnSubmixWeakPtr.Pin();
+				}
 			}
 
 		}
