Fix crash when swapping spatialization methods for binaural source

diff --git a/Engine/Source/Runtime/AudioMixer/Private/AudioMixerSource.cpp b/Engine/Source/Runtime/AudioMixer/Private/AudioMixerSource.cpp
--- a/Engine/Source/Runtime/AudioMixer/Private/AudioMixerSource.cpp
+++ b/Engine/Source/Runtime/AudioMixer/Private/AudioMixerSource.cpp
@@ -1643,7 +1643,12 @@ namespace Audio
 	float FMixerSource::GetInheritedSubmixVolumeModulation() const
 	{
 		FAudioDevice::FAudioSpatializationInterfaceInfo SpatializationInfo = MixerDevice->GetCurrentSpatializationPluginInterfaceInfo();
-		check(SpatializationInfo.bSpatializationIsExternalSend);
+		// We only hit this condition if, while the sound is playing, the spatializer changes from an external send to a non-external one.
+		// If that happens, the submix will catch all modulation so this function's logic is not needed.
+		if (!SpatializationInfo.bSpatializationIsExternalSend)
+		{
+			return 1.0f;
+		}
 
 		// if there is a return submix, we need to figure out where to stop manually attenuating
 		// Because the submix will modulate itself later
