Remove old, unused channel map code
#tests trivial

diff --git a/Engine/Source/Runtime/AudioMixer/Private/AudioMixerChannelMaps.cpp b/Engine/Source/Runtime/AudioMixer/Private/AudioMixerChannelMaps.cpp
--- a/Engine/Source/Runtime/AudioMixer/Private/AudioMixerChannelMaps.cpp
+++ b/Engine/Source/Runtime/AudioMixer/Private/AudioMixerChannelMaps.cpp
@@ -8,66 +8,6 @@
 
 namespace Audio
 {
-	/** Channel type matrix for submix speaker channel mappings. */
-	const TArray<TArray<EAudioMixerChannel::Type>> SubmixOutputChannelMatrix
-	{
-		// ESubmixChannelFormat::Device
-		// Placeholder: Should never be used as Device signifies dynamically set
-		{
-		},
-
-		// ESubmixChannelFormat::Stereo
-		{
-			EAudioMixerChannel::FrontLeft,
-			EAudioMixerChannel::FrontRight
-		},
-
-		// ESubmixChannelFormat::Quad
-		{
-			EAudioMixerChannel::FrontLeft,
-			EAudioMixerChannel::FrontRight,
-			EAudioMixerChannel::SideLeft,
-			EAudioMixerChannel::SideRight
-		},
-
-		// ESubmixChannelFormat::FiveDotOne
-		{
-			EAudioMixerChannel::FrontLeft,
-			EAudioMixerChannel::FrontRight,
-			EAudioMixerChannel::FrontCenter,
-			EAudioMixerChannel::LowFrequency,
-			EAudioMixerChannel::SideLeft,
-			EAudioMixerChannel::SideRight
-		},
-
-		// ESubmixChannelFormat::SevenDotOne
-		{
-			EAudioMixerChannel::FrontLeft,
-			EAudioMixerChannel::FrontRight,
-			EAudioMixerChannel::FrontCenter,
-			EAudioMixerChannel::LowFrequency,
-			EAudioMixerChannel::BackLeft,
-			EAudioMixerChannel::BackRight,
-			EAudioMixerChannel::SideLeft,
-			EAudioMixerChannel::SideRight
-		},
-
-		// ESubmixChannelFormat::Ambisonics
-		// Ambisonics output is encoded to max encoded channel (i.e. 7.1).
-		// To support ambisonic encoded output, will need to convert to
-		// Ambisonics_W/X/Y/Z alias values.
-		{
-			EAudioMixerChannel::FrontLeft,
-			EAudioMixerChannel::FrontRight,
-			EAudioMixerChannel::FrontCenter,
-			EAudioMixerChannel::LowFrequency,
-			EAudioMixerChannel::BackLeft,
-			EAudioMixerChannel::BackRight,
-			EAudioMixerChannel::SideLeft,
-			EAudioMixerChannel::SideRight
-		}, 
-	};
-
 	// Make a channel map cache
 	static TArray<TArray<float>> ChannelMapCache;
 	static TArray<TArray<float>> VorbisChannelMapCache;
