Fix for ASan stack buffer overload warning

diff --git a/Engine/Source/Runtime/AudioMixer/Private/SoundFileIO/SoundFileIOManagerImpl.cpp b/Engine/Source/Runtime/AudioMixer/Private/SoundFileIO/SoundFileIOManagerImpl.cpp
--- a/Engine/Source/Runtime/AudioMixer/Private/SoundFileIO/SoundFileIOManagerImpl.cpp
+++ b/Engine/Source/Runtime/AudioMixer/Private/SoundFileIO/SoundFileIOManagerImpl.cpp
@@ -393,9 +393,13 @@ namespace Audio
 
 			for (const uint32 Id : OptionalChunkIds)
 			{
-				FSoundFileChunkInfo ChunkLookup;
-				ChunkLookup.ChunkIdSize = sizeof(Id);
-				FCStringAnsi::Strncpy(ChunkLookup.ChunkId, (ANSICHAR*)&(INTEL_ORDER32(Id)), ChunkLookup.ChunkIdSize + 1);
+				FSoundFileChunkInfo ChunkLookup;				
+				// Copy chunk ID over. DWORD (4 bytes, each is ANSI char)
+				*reinterpret_cast<uint32*>(ChunkLookup.ChunkId) = Id;
+				ChunkLookup.ChunkId[4] = 0;		// Null terminate the string just in case.
+				
+				ChunkLookup.ChunkIdSize = 5;	// 4 bytes, + null.
+				
 
 				// Lookup chunk of given Id. Multiple chunks can exist of a given type
 				// so we loop here.
