Log error if audio render thread id changes before render commands are completed.

diff --git a/Engine/Source/Runtime/AudioMixer/Private/AudioMixerDevice.cpp b/Engine/Source/Runtime/AudioMixer/Private/AudioMixerDevice.cpp
--- a/Engine/Source/Runtime/AudioMixer/Private/AudioMixerDevice.cpp
+++ b/Engine/Source/Runtime/AudioMixer/Private/AudioMixerDevice.cpp
@@ -745,6 +745,15 @@ namespace Audio
 			// be flushed. Audio Rendering Thread ID will be reset again in call
 			// to FMixerDevice::OnProcessAudio
 			ResetAudioRenderingThreadId();
+			
+			// Cache the audio platform thread id for debugging purposes. This is
+			// an attempt to narrow down the causes of UE-209237. The theory is
+			// that the there are multiple threads attempting to run audio rendering
+			// commands. To catch the issue we attempt to cache the thread id before
+			// running audio rendering commands. We then check that the thread id
+			// has not changed for the duration that render thread commands have 
+			// run.
+			std::atomic<int32> CurrentAudioPlatformThreadId = AudioPlatformThreadId.load();
 
 			// Force source manager to incorporate device channel count change.
 			FlushAudioRenderingCommands(true /* bPumpSynchronously */);
@@ -754,6 +763,9 @@ namespace Audio
 				AudioDeviceNotifSubsystem->OnDeviceSwitched(PlatformInfo.DeviceId);
 			}
 
+			// Related to earlier mention of UE-209237
+			UE_CLOG(CurrentAudioPlatformThreadId != AudioPlatformThreadId, LogAudioMixer, Error, TEXT("Platform audio thread id changed while flushing render commands. Expected %d, found %d. May result in corrupt internal audio source state."), CurrentAudioPlatformThreadId.load(), AudioPlatformThreadId.load());
+
 			// Audio rendering was suspended in CheckAudioDeviceChange if it changed.
 			AudioMixerPlatform->ResumePlaybackOnNewDevice();
 		}
@@ -893,6 +905,15 @@ namespace Audio
 		// This function could be called in a task manager, which means the thread ID may change between calls.
 		ResetAudioRenderingThreadId();
 
+		// Cache the audio platform thread id for debugging purposes. This is
+		// an attempt to narrow down the causes of UE-209237. The theory is
+		// that the there are multiple threads attempting to run audio rendering
+		// commands. To catch the issue we attempt to cache the thread id before
+		// running audio rendering commands. We then check that the thread id
+		// has not changed for the duration that render thread commands have 
+		// run.
+		std::atomic<int32> CurrentAudioPlatformThreadId = AudioPlatformThreadId.load();
+
 		// Update the audio render thread time at the head of the render
 		AudioThreadTimingData.AudioRenderThreadTime = FPlatformTime::Seconds() - AudioThreadTimingData.StartTime;
 
@@ -974,6 +995,12 @@ namespace Audio
 		NotifyAudioDevicePostRender(RenderInfo);
 
 		KickQueuedTasks((Audio::AudioTaskQueueId)DeviceID);
+
+
+		// Related to earlier mention of UE-209237
+		UE_CLOG(CurrentAudioPlatformThreadId != AudioPlatformThreadId, LogAudioMixer, Error, TEXT("Platform audio thread id changed while flushing render commands. Expected %d, found %d. May result in corrupt internal audio source state."), CurrentAudioPlatformThreadId.load(), AudioPlatformThreadId.load());
+
+
 		return true;
 	}
 
diff --git a/Engine/Source/Runtime/AudioMixer/Public/AudioMixerDevice.h b/Engine/Source/Runtime/AudioMixer/Public/AudioMixerDevice.h
--- a/Engine/Source/Runtime/AudioMixer/Public/AudioMixerDevice.h
+++ b/Engine/Source/Runtime/AudioMixer/Public/AudioMixerDevice.h
@@ -464,7 +464,7 @@ namespace Audio
 		mutable int32 GameOrAudioThreadId;
 
 		/** ThreadId for the low-level platform audio mixer. */
-		mutable int32 AudioPlatformThreadId;
+		mutable std::atomic<int32> AudioPlatformThreadId;
 
 		/** Command queue to send commands to audio render thread from game thread or audio thread. */
 		TQueue<TFunction<void()>> CommandQueue;
