Safeguard Modulation Submix fix against using external spatializer with internal reverb

diff --git a/Engine/Source/Runtime/AudioMixer/Private/AudioMixerSource.cpp b/Engine/Source/Runtime/AudioMixer/Private/AudioMixerSource.cpp
--- a/Engine/Source/Runtime/AudioMixer/Private/AudioMixerSource.cpp
+++ b/Engine/Source/Runtime/AudioMixer/Private/AudioMixerSource.cpp
@@ -776,11 +776,18 @@ namespace Audio
 				{
 					if ((SubmixPtr->IsSoundfieldSubmix() || SubmixPtr->IsSoundfieldEndpointSubmix()))
 					{
-						UE_LOG(LogAudioMixer, Warning, TEXT("Ignoring soundfield Base Submix destination being set on SoundWave (%s) because spatializaition method is set to Binaural.")
+						UE_LOG(LogAudioMixer, Warning, TEXT("Ignoring soundfield Base Submix destination being set on SoundWave (%s) because spatialization method is set to Binaural.")
 							, *InWaveInstance->GetName());
 					}
 					
-					bBypassingSubmixModulation = true;
+					if (MixerDevice->ReverbPluginInterface)
+					{
+						bBypassingSubmixModulation = true;
+					}
+					else
+					{
+						UE_LOG(LogAudioMixer, Warning, TEXT("No Reverb Plugin loaded. Sounds using an external spatialization plugin without a reverb return submix will not have submix modulation applied."));
+					}
 				}
 			}
 
@@ -1657,20 +1664,22 @@ namespace Audio
 		TArray<uint32> ReturnSubmixAncestors;
 		if (SpatializationInfo.bReturnsToSubmixGraph)
 		{
-			USoundSubmix* ReturnSubmix = MixerDevice->ReverbPluginInterface->GetSubmix();
-			if (ReturnSubmix)
+			if (MixerDevice && MixerDevice->ReverbPluginInterface)
 			{
-				FMixerSubmixWeakPtr CurrReturnSubmixWeakPtr = MixerDevice->GetSubmixInstance(ReturnSubmix);
-				FMixerSubmixPtr CurrReturnSubmixPtr = CurrReturnSubmixWeakPtr.Pin();
-				while (CurrReturnSubmixPtr && CurrReturnSubmixPtr->IsValid())
+				USoundSubmix* ReturnSubmix = MixerDevice->ReverbPluginInterface->GetSubmix();
+				if (ReturnSubmix)
 				{
-					ReturnSubmixAncestors.Add(CurrReturnSubmixPtr->GetId());
+					FMixerSubmixWeakPtr CurrReturnSubmixWeakPtr = MixerDevice->GetSubmixInstance(ReturnSubmix);
+					FMixerSubmixPtr CurrReturnSubmixPtr = CurrReturnSubmixWeakPtr.Pin();
+					while (CurrReturnSubmixPtr && CurrReturnSubmixPtr->IsValid())
+					{
+						ReturnSubmixAncestors.Add(CurrReturnSubmixPtr->GetId());
 
-					CurrReturnSubmixWeakPtr = CurrReturnSubmixPtr->GetParent();
-					CurrReturnSubmixPtr = CurrReturnSubmixWeakPtr.Pin();
+						CurrReturnSubmixWeakPtr = CurrReturnSubmixPtr->GetParent();
+						CurrReturnSubmixPtr = CurrReturnSubmixWeakPtr.Pin();
+					}
 				}
 			}
-
 		}
 
 		float SubmixModVolume = 1.0f;
