Fix Static Analysis warning in Audio Mixer Source

diff --git a/Engine/Source/Runtime/AudioMixer/Private/AudioMixerSource.cpp b/Engine/Source/Runtime/AudioMixer/Private/AudioMixerSource.cpp
--- a/Engine/Source/Runtime/AudioMixer/Private/AudioMixerSource.cpp
+++ b/Engine/Source/Runtime/AudioMixer/Private/AudioMixerSource.cpp
@@ -1649,6 +1649,11 @@ namespace Audio
 
 	float FMixerSource::GetInheritedSubmixVolumeModulation() const
 	{
+		if (!MixerDevice)
+		{
+			return 1.0f;
+		}
+
 		FAudioDevice::FAudioSpatializationInterfaceInfo SpatializationInfo = MixerDevice->GetCurrentSpatializationPluginInterfaceInfo();
 		// We only hit this condition if, while the sound is playing, the spatializer changes from an external send to a non-external one.
 		// If that happens, the submix will catch all modulation so this function's logic is not needed.
