Add ability to remove envelope follower delegates from submixes

#rnx

diff --git a/Engine/Source/Runtime/AudioMixer/Private/AudioMixerDevice.cpp b/Engine/Source/Runtime/AudioMixer/Private/AudioMixerDevice.cpp
--- a/Engine/Source/Runtime/AudioMixer/Private/AudioMixerDevice.cpp
+++ b/Engine/Source/Runtime/AudioMixer/Private/AudioMixerDevice.cpp
@@ -2501,6 +2501,35 @@ namespace Audio
 		}
 	}
 
+	void FMixerDevice::RemoveEnvelopeFollowerDelegate(USoundSubmix* InSubmix, const FOnSubmixEnvelopeBP& OnSubmixEnvelopeBP)
+	{
+		if (!IsInAudioThread())
+		{
+			DECLARE_CYCLE_STAT(TEXT("FAudioThreadTask.RemoveEnvelopeFollowerDelegate"), STAT_RemoveEnvelopeFollowerDelegate, STATGROUP_AudioThreadCommands);
+
+			FAudioThread::RunCommandOnAudioThread([this, InSubmix, OnSubmixEnvelopeBP]()
+			{
+				CSV_SCOPED_TIMING_STAT(Audio, RemoveEnvelopeFollowerDelegate);
+				RemoveEnvelopeFollowerDelegate(InSubmix, OnSubmixEnvelopeBP);
+			}, GET_STATID(STAT_RemoveEnvelopeFollowerDelegate));
+			return;
+		}
+
+		// Fallback to the master submix if the provided submix isn't found to match behavior from ::AddEnvelopeFollowerDelegate
+		FMixerSubmixPtr FoundSubmix = GetSubmixInstance(InSubmix).Pin();
+		if (FoundSubmix.IsValid())
+		{
+			FoundSubmix->RemoveEnvelopeFollowerDelegate(OnSubmixEnvelopeBP);
+		}
+		else
+		{
+			FMixerSubmixWeakPtr MainSubmix = GetMasterSubmix();
+			FMixerSubmixPtr MainSubmixPtr = MainSubmix.Pin();
+			check(MainSubmixPtr.IsValid());
+
+			MainSubmixPtr->RemoveEnvelopeFollowerDelegate(OnSubmixEnvelopeBP);
+		}
+	}
 
 	void FMixerDevice::StartSpectrumAnalysis(USoundSubmix* InSubmix, const FSoundSpectrumAnalyzerSettings& InSettings)
 	{
diff --git a/Engine/Source/Runtime/AudioMixer/Private/AudioMixerSubmix.cpp b/Engine/Source/Runtime/AudioMixer/Private/AudioMixerSubmix.cpp
--- a/Engine/Source/Runtime/AudioMixer/Private/AudioMixerSubmix.cpp
+++ b/Engine/Source/Runtime/AudioMixer/Private/AudioMixerSubmix.cpp
@@ -2205,6 +2205,11 @@ namespace Audio
 		OnSubmixEnvelope.AddUnique(OnSubmixEnvelopeBP);
 	}
 
+	void FMixerSubmix::RemoveEnvelopeFollowerDelegate(const FOnSubmixEnvelopeBP& OnSubmixEnvelopeBP)
+	{
+		OnSubmixEnvelope.Remove(OnSubmixEnvelopeBP);
+	}
+
 	void FMixerSubmix::AddSpectralAnalysisDelegate(const FSoundSpectrumAnalyzerDelegateSettings& InDelegateSettings, const FOnSubmixSpectralAnalysisBP& OnSubmixSpectralAnalysisBP)
 	{
 		FSpectrumAnalysisDelegateInfo NewDelegateInfo;
diff --git a/Engine/Source/Runtime/AudioMixer/Public/AudioMixerDevice.h b/Engine/Source/Runtime/AudioMixer/Public/AudioMixerDevice.h
--- a/Engine/Source/Runtime/AudioMixer/Public/AudioMixerDevice.h
+++ b/Engine/Source/Runtime/AudioMixer/Public/AudioMixerDevice.h
@@ -184,6 +184,7 @@ namespace Audio
 		AUDIOMIXER_API virtual void StartEnvelopeFollowing(USoundSubmix* InSubmix) override;
 		AUDIOMIXER_API virtual void StopEnvelopeFollowing(USoundSubmix* InSubmix) override;
 		AUDIOMIXER_API virtual void AddEnvelopeFollowerDelegate(USoundSubmix* InSubmix, const FOnSubmixEnvelopeBP& OnSubmixEnvelopeBP) override;
+		AUDIOMIXER_API virtual void RemoveEnvelopeFollowerDelegate(USoundSubmix* InSubmix, const FOnSubmixEnvelopeBP& OnSubmixEnvelopeBP) override;
 
 		// Submix Spectrum Analysis
 		AUDIOMIXER_API virtual void StartSpectrumAnalysis(USoundSubmix* InSubmix, const FSoundSpectrumAnalyzerSettings& InSettings) override;
diff --git a/Engine/Source/Runtime/AudioMixer/Public/AudioMixerSubmix.h b/Engine/Source/Runtime/AudioMixer/Public/AudioMixerSubmix.h
--- a/Engine/Source/Runtime/AudioMixer/Public/AudioMixerSubmix.h
+++ b/Engine/Source/Runtime/AudioMixer/Public/AudioMixerSubmix.h
@@ -274,6 +274,9 @@ namespace Audio
 		// Adds an envelope follower delegate
 		AUDIOMIXER_API void AddEnvelopeFollowerDelegate(const FOnSubmixEnvelopeBP& OnSubmixEnvelopeBP);
 
+		// Removes an existing envelope follower delegate
+		AUDIOMIXER_API void RemoveEnvelopeFollowerDelegate(const FOnSubmixEnvelopeBP& OnSubmixEnvelopeBP);
+
 		// Initializes a new FFT analyzer for this submix and immediately begins feeding audio to it.
 		AUDIOMIXER_API void StartSpectrumAnalysis(const FSoundSpectrumAnalyzerSettings& InSettings);
 
diff --git a/Engine/Source/Runtime/Engine/Classes/Sound/SoundSubmix.h b/Engine/Source/Runtime/Engine/Classes/Sound/SoundSubmix.h
--- a/Engine/Source/Runtime/Engine/Classes/Sound/SoundSubmix.h
+++ b/Engine/Source/Runtime/Engine/Classes/Sound/SoundSubmix.h
@@ -410,6 +410,13 @@ public:
 	UFUNCTION(BlueprintCallable, Category = "Audio|EnvelopeFollowing", meta = (WorldContext = "WorldContextObject"))
 	ENGINE_API void AddEnvelopeFollowerDelegate(const UObject* WorldContextObject, const FOnSubmixEnvelopeBP& OnSubmixEnvelopeBP);
 
+	/**
+	 *	Remove an envelope follower delegate.
+	 *	@param	OnSubmixEnvelopeBP	The event delegate to remove.
+	 */
+	UFUNCTION(BlueprintCallable, Category = "Audio|EnvelopeFollowing", meta = (WorldContext = "WorldContextObject"))
+	ENGINE_API void RemoveEnvelopeFollowerDelegate(const UObject* WorldContextObject, const FOnSubmixEnvelopeBP& OnSubmixEnvelopeBP);
+
 	/**
 	 *	Adds a spectral analysis delegate to receive notifications when this submix has spectrum analysis enabled.
 	 *	@param	InBandsettings					The frequency bands to analyze and their envelope-following settings.
diff --git a/Engine/Source/Runtime/Engine/Private/AudioDevice.cpp b/Engine/Source/Runtime/Engine/Private/AudioDevice.cpp
--- a/Engine/Source/Runtime/Engine/Private/AudioDevice.cpp
+++ b/Engine/Source/Runtime/Engine/Private/AudioDevice.cpp
@@ -5277,6 +5277,11 @@ void FAudioDevice::AddEnvelopeFollowerDelegate(USoundSubmix* InSubmix, const FOn
 	UE_LOG(LogAudio, Error, TEXT("Envelope following submixes only works with the audio mixer. Please run using -audiomixer or set INI file to use submix recording."));
 }
 
+void FAudioDevice::RemoveEnvelopeFollowerDelegate(USoundSubmix* InSubmix, const FOnSubmixEnvelopeBP& OnSubmixEnvelopeBP)
+{
+	UE_LOG(LogAudio, Error, TEXT("Envelope following submixes only works with the audio mixer. Please run using -audiomixer or set INI file to use submix recording."));
+}
+
 void FAudioDevice::StartSpectrumAnalysis(USoundSubmix* InSubmix, const FSoundSpectrumAnalyzerSettings& InSettings)
 {
 	UE_LOG(LogAudio, Error, TEXT("Spectrum analysis of submixes only works with the audio mixer. Please run using -audiomixer or set INI file to use submix recording."));
diff --git a/Engine/Source/Runtime/Engine/Private/SoundSubmix.cpp b/Engine/Source/Runtime/Engine/Private/SoundSubmix.cpp
--- a/Engine/Source/Runtime/Engine/Private/SoundSubmix.cpp
+++ b/Engine/Source/Runtime/Engine/Private/SoundSubmix.cpp
@@ -277,6 +277,23 @@ void USoundSubmix::AddEnvelopeFollowerDelegate(const UObject* WorldContextObject
 	}
 }
 
+void USoundSubmix::RemoveEnvelopeFollowerDelegate(const UObject* WorldContextObject, const FOnSubmixEnvelopeBP& OnSubmixEnvelopeBP)
+{
+	if (!GEngine)
+	{
+		return;
+	}
+
+	// Find device for this specific audio recording thing.
+	if (UWorld* ThisWorld = GEngine->GetWorldFromContextObject(WorldContextObject, EGetWorldErrorMode::LogAndReturnNull))
+	{
+		if (FAudioDevice* AudioDevice = ThisWorld->GetAudioDeviceRaw())
+		{
+			AudioDevice->RemoveEnvelopeFollowerDelegate(this, OnSubmixEnvelopeBP);
+		}
+	}
+}
+
 void USoundSubmix::AddSpectralAnalysisDelegate(const UObject* WorldContextObject, const TArray<FSoundSubmixSpectralAnalysisBandSettings>& InBandSettings, const FOnSubmixSpectralAnalysisBP& OnSubmixSpectralAnalysisBP, float UpdateRate,  float DecibelNoiseFloor, bool bDoNormalize, bool bDoAutoRange, float AutoRangeAttackTime, float AutoRangeReleaseTime)
 {
 
diff --git a/Engine/Source/Runtime/Engine/Public/AudioDevice.h b/Engine/Source/Runtime/Engine/Public/AudioDevice.h
--- a/Engine/Source/Runtime/Engine/Public/AudioDevice.h
+++ b/Engine/Source/Runtime/Engine/Public/AudioDevice.h
@@ -1472,6 +1472,7 @@ public:
 
 	/** Adds an envelope follower delegate to the submix for this audio device. */
 	ENGINE_API virtual void AddEnvelopeFollowerDelegate(USoundSubmix* InSubmix, const FOnSubmixEnvelopeBP& OnSubmixEnvelopeBP);
+	ENGINE_API virtual void RemoveEnvelopeFollowerDelegate(USoundSubmix* InSubmix, const FOnSubmixEnvelopeBP& OnSubmixEnvelopeBP);
 
 	ENGINE_API virtual void StartSpectrumAnalysis(USoundSubmix* InSubmix, const FSoundSpectrumAnalyzerSettings& InSettings);
 	ENGINE_API virtual void StopSpectrumAnalysis(USoundSubmix* InSubmix);
