Fix Source Effect Modulation Ensure

diff --git a/Engine/Source/Runtime/AudioMixer/Private/AudioMixerDevice.cpp b/Engine/Source/Runtime/AudioMixer/Private/AudioMixerDevice.cpp
--- a/Engine/Source/Runtime/AudioMixer/Private/AudioMixerDevice.cpp
+++ b/Engine/Source/Runtime/AudioMixer/Private/AudioMixerDevice.cpp
@@ -1487,7 +1487,22 @@ namespace Audio
 			SourceEffectChainOverrides.Add(SourceEffectChainId, SourceEffectChain);
 		}
 
-		SourceManager->UpdateSourceEffectChain(SourceEffectChainId, SourceEffectChain, bPlayEffectChainTails);
+		FAudioThread::RunCommandOnAudioThread([MixerDeviceID = DeviceID, SourceEffectChainId, SourceEffectChain, bPlayEffectChainTails]()
+		{
+			if (FAudioDeviceManager* Manager = FAudioDeviceManager::Get())
+			{
+				if (FAudioDevice* Device = Manager->GetAudioDeviceRaw(MixerDeviceID))
+				{
+					FMixerDevice* MixerDevice = static_cast<FMixerDevice*>(Device);
+					if (MixerDevice && MixerDevice->SourceManager)
+					{
+						MixerDevice->SourceManager->UpdateSourceEffectChain(SourceEffectChainId, SourceEffectChain, bPlayEffectChainTails);
+					}
+				}
+			}
+		});
+
+		
 	}
 
 	void FMixerDevice::UpdateSubmixProperties(USoundSubmixBase* InSoundSubmix)
