Fix for audible vibration endpoints

diff --git a/Engine/Source/Runtime/AudioMixer/Private/AudioMixerDevice.cpp b/Engine/Source/Runtime/AudioMixer/Private/AudioMixerDevice.cpp
--- a/Engine/Source/Runtime/AudioMixer/Private/AudioMixerDevice.cpp
+++ b/Engine/Source/Runtime/AudioMixer/Private/AudioMixerDevice.cpp
@@ -67,6 +67,14 @@ FAutoConsoleVariableRef CVarDisableSubmixMutationLock(
 	TEXT("0: Not Disabled (Default), 1: Disabled"),
 	ECVF_Default);
 
+static int32 EnableAudibleDefaultEndpointSubmixesCVar = 0;
+FAutoConsoleVariableRef CVarEnableAudibleDefaultEndpointSubmixes(
+	TEXT("au.submix.audibledefaultendpoints"),
+	EnableAudibleDefaultEndpointSubmixesCVar,
+	TEXT("Allows audio sent to defaulted (typically silent) endpoint submixes to be audible via master. (useful for debugging)\n")
+	TEXT("0: Disabled (Default), 1: Enabled"),
+	ECVF_Default);
+
 static int32 DebugGeneratorEnableCVar = 0;
 FAutoConsoleVariableRef CVarDebugGeneratorEnable(
 	TEXT("au.Debug.Generator"),
@@ -949,16 +957,19 @@ namespace Audio
 			SCOPE_CYCLE_COUNTER(STAT_AudioMixerEndpointSubmixes);
 
 			FScopeLock ScopeLock(&EndpointSubmixesMutationLock);
-			for (FMixerSubmixPtr& Submix : DefaultEndpointSubmixes)
+			if (EnableAudibleDefaultEndpointSubmixesCVar !=0 )
 			{
-				// If this hit, a submix was added to the default submix endpoint array
-				// even though it's not an endpoint, or a parent was set on an endpoint submix
-				// and it wasn't removed from DefaultEndpointSubmixes.
-				ensure(Submix->IsDefaultEndpointSubmix());
-
-				// Any endpoint submixes that don't specify an endpoint
-				// are summed into our master output.
-				Submix->ProcessAudio(Output);
+				for (const FMixerSubmixPtr& Submix : DefaultEndpointSubmixes)
+				{
+					// If this hit, a submix was added to the default submix endpoint array
+					// even though it's not an endpoint, or a parent was set on an endpoint submix
+					// and it wasn't removed from DefaultEndpointSubmixes.
+					ensure(Submix->IsDefaultEndpointSubmix());
+
+					// Any endpoint submixes that don't specify an endpoint
+					// are summed into our master output.
+					Submix->ProcessAudio(Output);
+				}
 			}
 			
 			for (FMixerSubmixPtr& Submix : ExternalEndpointSubmixes)
