Quartz Command Queue Refactor (motivated by the need for audio render thread metronome subscribers)

diff --git a/Engine/Source/Runtime/AudioMixer/Private/Quartz/AudioMixerClock.cpp b/Engine/Source/Runtime/AudioMixer/Private/Quartz/AudioMixerClock.cpp
--- a/Engine/Source/Runtime/AudioMixer/Private/Quartz/AudioMixerClock.cpp
+++ b/Engine/Source/Runtime/AudioMixer/Private/Quartz/AudioMixerClock.cpp
@@ -106,9 +106,14 @@ namespace Audio
 
 	bool FQuartzClockProxy::SendCommandToClock(TFunction<void(FQuartzClock*)> InCommand)
 	{
-		if (auto QueuePtr = SharedQueue.Pin())
+		if (FQuartzClockCommandQueuePtr QueuePtr = SharedQueue.Pin())
 		{
-			QueuePtr->PushCommand(InCommand);
+			QueuePtr->PushLambda<Quartz::IQuartzClock>(
+				[Command = MoveTemp(InCommand)](Quartz::IQuartzClock& InListener)
+				{
+					Command(static_cast<FQuartzClock*>(&InListener));
+				});
+
 			return true;
 		}
 
@@ -258,12 +263,12 @@ namespace Audio
 
 	void FQuartzClock::Shutdown()
 	{
-		for (auto& PendingCommand : PendingCommands)
+		for (PendingCommand& PendingCommand : PendingCommands)
 		{
 			PendingCommand.Command->Cancel();
 		}
 
-		for (auto& PendingCommand : ClockAlteringPendingCommands)
+		for (PendingCommand& PendingCommand : ClockAlteringPendingCommands)
 		{
 			PendingCommand.Command->Cancel();
 		}
@@ -276,7 +281,7 @@ namespace Audio
 	{
 		TRACE_CPUPROFILER_EVENT_SCOPE(QuartzClock::Tick_LowRes);
 		UE_LOG(LogAudioQuartz, Verbose, TEXT("Quartz Clock Tick (low-res): %s"), *Name.ToString());
-		PreTickCommands->PumpCommandQueue(this);
+		PreTickCommands->PumpCommandQueue(*this);
 		Tick(static_cast<int32>(InDeltaTimeSeconds * Metronome.GetTickRate().GetSampleRate()));
 	}
 
@@ -286,8 +291,8 @@ namespace Audio
 		TRACE_CPUPROFILER_EVENT_SCOPE(QuartzClock::GameThreadCommands);
 
 		UE_LOG(LogAudioQuartz, Verbose, TEXT("Quartz Clock Tick: %s"), *Name.ToString());
-		
-		PreTickCommands->PumpCommandQueue(this);
+
+		PreTickCommands->PumpCommandQueue(*this);
 
 		if (!bIsRunning)
 		{
@@ -322,7 +327,7 @@ namespace Audio
 	{
 		if (!PreTickCommands.IsValid())
 		{
-			PreTickCommands = TQuartzShareableCommandQueue<FQuartzClock>::Create();
+			PreTickCommands = MakeShared<FQuartzClockCommandQueueType>();
 		}
 
 		return PreTickCommands;
diff --git a/Engine/Source/Runtime/AudioMixer/Private/Quartz/AudioMixerQuantizedCommands.cpp b/Engine/Source/Runtime/AudioMixer/Private/Quartz/AudioMixerQuantizedCommands.cpp
--- a/Engine/Source/Runtime/AudioMixer/Private/Quartz/AudioMixerQuantizedCommands.cpp
+++ b/Engine/Source/Runtime/AudioMixer/Private/Quartz/AudioMixerQuantizedCommands.cpp
@@ -101,6 +101,19 @@ namespace Audio
 		return PlayCommandName;
 	}
 
+	int32 FQuantizedPlayCommand::OverrideFramesUntilExec(int32 NumFramesUntilExec)
+	{
+//		if (OwningClockPtr)
+//		{
+//				if(const FMixerDevice* MixerDevice = OwningClockPtr->GetMixerDevice())
+//				{
+//					return NumFramesUntilExec + 2 * MixerDevice->GetBufferLength();
+//				}
+//		}
+//		
+		return NumFramesUntilExec;
+	}
+
 	void FQuantizedQueueCommand::SetQueueCommand(const FAudioComponentCommandInfo& InAudioComponentData)
 	{
 		AudioComponentData = InAudioComponentData;
@@ -119,7 +132,7 @@ namespace Audio
 	int32 FQuantizedQueueCommand::OverrideFramesUntilExec(int32 NumFramesUntilExec)
 	{
 		// Calculate the amount of time before taking up a voice slot
-		int32 NumFramesBeforeVoiceSlot = NumFramesUntilExec - static_cast<int32>(OwningClockPtr->GetTickRate().GetFramesPerDuration(AudioComponentData.AnticapatoryBoundary.Quantization));
+		int32 NumFramesBeforeVoiceSlot = NumFramesUntilExec - static_cast<int32>(OwningClockPtr->GetTickRate().GetFramesPerDuration(AudioComponentData.AnticipatoryBoundary.Quantization));
 
 		//If NumFramesBeforeVoiceSlot is less than 0, change the boundary back to the original, and mark this command as having 0 frames till exec
 		if (NumFramesBeforeVoiceSlot < 0)
diff --git a/Engine/Source/Runtime/AudioMixer/Private/Quartz/QuartzMetronome.cpp b/Engine/Source/Runtime/AudioMixer/Private/Quartz/QuartzMetronome.cpp
--- a/Engine/Source/Runtime/AudioMixer/Private/Quartz/QuartzMetronome.cpp
+++ b/Engine/Source/Runtime/AudioMixer/Private/Quartz/QuartzMetronome.cpp
@@ -24,34 +24,24 @@ namespace Audio
 	void FQuartzMetronome::Tick(int32 InNumSamples, int32 FramesOfLatency)
 	{
 		LastTickCpuCycles64 = FPlatformTime::Cycles64();
-		
-		static bool bHasWarned = false;
-		if (!bHasWarned && (MusicalDurationsInFrames[EQuartzCommandQuantization::ThirtySecondNote] < InNumSamples))
-		{
-			// TODO: fire duplicate events if this occurs to facilitate game play-side counting logic
-			UE_LOG(LogAudioQuartz, Warning
-				, TEXT("Small note durations are shorter than the audio callback size. Some musical events may not fire delegates"));
-
-			bHasWarned = true;
-		}
 
-		int32 ToUpdateBitField = 0;
 
 		for (int i = 0; i < static_cast<int32>(EQuartzCommandQuantization::Count); ++i)
 		{
-			EQuartzCommandQuantization DurationType = static_cast<EQuartzCommandQuantization>(i);
+			const EQuartzCommandQuantization DurationType = static_cast<EQuartzCommandQuantization>(i);
+			int32 EventFrame = FramesLeftInMusicalDuration[DurationType]; 
 			FramesLeftInMusicalDuration[DurationType] -= InNumSamples;
 			
 			if (FramesLeftInMusicalDuration[DurationType] < 0)
 			{
-				// flag this duration for an update
-				ToUpdateBitField |= (1 << i);
-
 				// the beat value is constant
 				if (!(DurationType == EQuartzCommandQuantization::Beat && PulseDurations.Num()))
 				{
 					do
 					{
+						PendingMetronomeEvents.Add(DurationType, EventFrame);
+						
+						EventFrame += MusicalDurationsInFrames[DurationType];
 						FramesLeftInMusicalDuration[DurationType] += MusicalDurationsInFrames[DurationType];
 					}
 					while (FramesLeftInMusicalDuration[DurationType] <= 0);
@@ -61,11 +51,13 @@ namespace Audio
 					// the beat value can change
 					do
 					{
+						PendingMetronomeEvents.Add(DurationType, EventFrame);
 						if (++PulseDurationIndex == PulseDurations.Num())
 						{
 							PulseDurationIndex = 0;
 						}
 
+						EventFrame += MusicalDurationsInFrames[DurationType];
 						FramesLeftInMusicalDuration[DurationType] += PulseDurations[PulseDurationIndex];
 						MusicalDurationsInFrames[DurationType] = PulseDurations[PulseDurationIndex];
 					}
@@ -74,14 +66,13 @@ namespace Audio
 			}
 		}
 
-
 		// update transport
-		if (ToUpdateBitField & (1 << static_cast<int>(EQuartzCommandQuantization::Bar)))
+		if (PendingMetronomeEvents.HasPendingEvent(EQuartzCommandQuantization::Bar))
 		{
 			++CurrentTimeStamp.Bars;
 			CurrentTimeStamp.Beat = 1;
 		}
-		else if (ToUpdateBitField & (1 << static_cast<int>(EQuartzCommandQuantization::Beat)))
+		else if (PendingMetronomeEvents.HasPendingEvent(EQuartzCommandQuantization::Beat))
 		{
 			++CurrentTimeStamp.Beat;
 		}
@@ -97,7 +88,8 @@ namespace Audio
 
 		TimeSinceStart += double(InNumSamples) / CurrentTickRate.GetSampleRate(); 
 		CurrentTimeStamp.Seconds = TimeSinceStart;
-		FireEvents(ToUpdateBitField);
+		FireEvents();
+		PendingMetronomeEvents.Reset();
 	}
 
 	void FQuartzMetronome::SetTickRate(FQuartzClockTickRate InNewTickRate, int32 NumFramesLeft)
@@ -287,7 +279,7 @@ namespace Audio
 		}
 		else
 		{
-			float BarProgress = 1.0f - (FramesLeftInMusicalDuration[EQuartzCommandQuantization::Bar] / static_cast<float>(MusicalDurationsInFrames[EQuartzCommandQuantization::Bar]));
+			const float BarProgress = 1.0f - (FramesLeftInMusicalDuration[EQuartzCommandQuantization::Bar] / static_cast<float>(MusicalDurationsInFrames[EQuartzCommandQuantization::Bar]));
 			Count += (BarProgress * CountNumSubdivisionsPerBar(InSubdivision));
 		}
 
@@ -296,12 +288,64 @@ namespace Audio
 
 	float FQuartzMetronome::CountNumSubdivisionsSinceStart(EQuartzCommandQuantization InSubdivision) const
 	{
-		int32 NumPerBar = CountNumSubdivisionsPerBar(InSubdivision);
-		int32 NumInThisBar = CountNumSubdivisionsSinceBarStart(InSubdivision);
+		const int32 NumPerBar = CountNumSubdivisionsPerBar(InSubdivision);
+		const int32 NumInThisBar = CountNumSubdivisionsSinceBarStart(InSubdivision);
 
 		return (CurrentTimeStamp.Bars - 1) * NumPerBar + NumInThisBar;
 	}
 
+	void FQuartzMetronome::FFramesInTimeValue::Reset()
+	{
+		for (double& FrameCount : FramesInTimeValueInternal)
+		{
+			FrameCount = 0.0;
+		}
+	}
+
+	double& FQuartzMetronome::FFramesInTimeValue::operator[](EQuartzCommandQuantization InTimeValue)
+	{
+		return FramesInTimeValueInternal[static_cast<int32>(InTimeValue)];
+	}
+
+	const double& FQuartzMetronome::FFramesInTimeValue::operator[](EQuartzCommandQuantization InTimeValue) const
+	{
+		return FramesInTimeValueInternal[static_cast<int32>(InTimeValue)];
+	}
+
+	double& FQuartzMetronome::FFramesInTimeValue::operator[](int32 Index)
+	{
+		return FramesInTimeValueInternal[Index];
+	}
+
+	const double& FQuartzMetronome::FFramesInTimeValue::operator[](int32 Index) const
+	{
+		return FramesInTimeValueInternal[Index];
+	}
+
+	void FQuartzMetronome::FMetronomeEventEntry::Reset()
+	{
+		EventFrames.Reset();
+	}
+
+	void FQuartzMetronome::FPendingMetronomeEvents::Reset()
+	{
+		for (FMetronomeEventEntry& Entry : CurrentMetronomeEvents)
+		{
+			Entry.Reset();
+		}
+	}
+
+	bool FQuartzMetronome::FPendingMetronomeEvents::HasPendingEvent(const EQuartzCommandQuantization InDuration) const
+	{
+		return CurrentMetronomeEvents[static_cast<int32>(InDuration)].EventFrames.Num() != 0;
+	}
+
+	void FQuartzMetronome::FPendingMetronomeEvents::Add(const EQuartzCommandQuantization InDuration,
+	                                                    const int32 InFrame)
+	{
+		CurrentMetronomeEvents[static_cast<int32>(InDuration)].EventFrames.Add(InFrame);
+	}
+
 	void FQuartzMetronome::CalculateDurationPhases(float (&OutPhases)[static_cast<int32>(EQuartzCommandQuantization::Count)]) const
 	{
 		for (int i = 0; i < static_cast<int32>(EQuartzCommandQuantization::Count); ++i)
@@ -356,10 +400,7 @@ namespace Audio
 	{
 		CurrentTimeStamp.Reset();
 
-		for (double& FrameCount : FramesLeftInMusicalDuration.FramesInTimeValueInternal)
-		{
-			FrameCount = 0.0;
-		}
+		FramesLeftInMusicalDuration.Reset();
 
 		TimeSinceStart = 0.0;
 		PulseDurationIndex = -1;
@@ -437,35 +478,38 @@ namespace Audio
 		}
 	}
 
-	void FQuartzMetronome::FireEvents(int32 EventFlags)
+	void FQuartzMetronome::FireEvents()
 	{
-		if (!(EventFlags &= ListenerFlags))
-		{
-			// no events occurred that we have listeners for
-			return;
-		}
-
 		FQuartzMetronomeDelegateData Data;
 		Data.Bar = (CurrentTimeStamp.Bars);
 		Data.Beat = (CurrentTimeStamp.Beat);
 		Data.BeatFraction = (CurrentTimeStamp.BeatFraction);
 		Data.ClockName = ClockName;
-
-		// loop through quantization boundaries
+		
+		// loop through subscribers of each quantization boundary
 		int32 i = -1;
 		for (TArray<MetronomeCommandQueuePtr>& QuantizationBoundarySubscribers : MetronomeSubscriptionMatrix)
 		{
-			if (EventFlags & (1 << ++i))
-			{
-				Data.Quantization = static_cast<EQuartzCommandQuantization>(i);
+			Data.Quantization = static_cast<EQuartzCommandQuantization>(++i);
 
-				// loop through subscribers to that boundary
-				for (MetronomeCommandQueuePtr& Subscriber : QuantizationBoundarySubscribers)
+			// if this quantization boundary had any events...
+			if (PendingMetronomeEvents.HasPendingEvent(Data.Quantization))
+			{
+				// ...for each subscriber...
+				for (const MetronomeCommandQueuePtr& Subscriber : QuantizationBoundarySubscribers)
 				{
-					Subscriber->PushEvent(Data);
+					// fire an event for each instance of that metronome event in this buffer
+					for (const int32& MetronomeEventFrame : PendingMetronomeEvents.CurrentMetronomeEvents[i].EventFrames)
+					{
+						Data.FrameOffset = MetronomeEventFrame;
+						Subscriber->PushLambda<Quartz::IMetronomeEventListener>(
+							[=](Quartz::IMetronomeEventListener& InListener)
+							{
+								InListener.OnMetronomeEvent(Data);
+							});
+					}
 				}
 			}
 		}
 	}
-
 } // namespace Audio
\ No newline at end of file
diff --git a/Engine/Source/Runtime/AudioMixer/Public/Quartz/AudioMixerClock.h b/Engine/Source/Runtime/AudioMixer/Public/Quartz/AudioMixerClock.h
--- a/Engine/Source/Runtime/AudioMixer/Public/Quartz/AudioMixerClock.h
+++ b/Engine/Source/Runtime/AudioMixer/Public/Quartz/AudioMixerClock.h
@@ -6,6 +6,9 @@
 #include "HAL/ThreadSafeBool.h"
 #include "Sound/QuartzQuantizationUtilities.h"
 #include "Quartz/QuartzMetronome.h"
+#include "Sound/QuartzSubscription.h"
+#include "Sound/QuartzInterfaces.h"
+#include "Sound/QuartzCommandQueue.h"
 
 namespace Audio
 {
@@ -15,12 +18,9 @@ namespace Audio
 	class FMixerSourceManager;
 	class FQuartzClockManager;
 
-	template<class ListenerType>
-	class TQuartzShareableCommandQueue;
-
-	using FQuartzClockCommandQueuePtr = TSharedPtr<TQuartzShareableCommandQueue<FQuartzClock>, ESPMode::ThreadSafe>;
-	using FQuartzClockCommandQueueWeakPtr = TWeakPtr<TQuartzShareableCommandQueue<FQuartzClock>, ESPMode::ThreadSafe>;
-
+	using FQuartzClockCommandQueueType = Audio::Quartz::PrivateDefs::TQuartzCommandQueue<IQuartzClock>;
+	using FQuartzClockCommandQueuePtr = TSharedPtr<FQuartzClockCommandQueueType, ESPMode::ThreadSafe>;
+	using FQuartzClockCommandQueueWeakPtr = TWeakPtr<FQuartzClockCommandQueueType, ESPMode::ThreadSafe>;
 	
 	/**
 	 *	FQuartzClockProxy:
@@ -28,7 +28,7 @@ namespace Audio
 	 *		This class is a C++ handle to the underlying clock.
 	 *		
 	 *		It is mostly a wrapper around a TWeakPtr<FQuartzClock> and
-	 *		TSharedPtr<TQuartzShareableCommandQueue<FQuartzClock>.
+	 *		FQuartzClockCommandQueueType
 	 *		
 	 *		The getters query the underlying FQuartzClock directly,
 	 *		which returns values updated during the last audio-engine tick
@@ -44,37 +44,37 @@ namespace Audio
 	 *		FQuartzClock itself when it pumps the command queue.
 	 *
 	 */
-	class FQuartzClockProxy
+	class AUDIOMIXER_API FQuartzClockProxy
 	{
 	public:
 		// ctor
 		FQuartzClockProxy() {}
 		FQuartzClockProxy(const FName& Name) : ClockId(Name){ } // conv ctor from FName
-		AUDIOMIXER_API FQuartzClockProxy(TSharedPtr<FQuartzClock, ESPMode::ThreadSafe> InClock);
+		FQuartzClockProxy(TSharedPtr<FQuartzClock, ESPMode::ThreadSafe> InClock);
 
 		FName GetClockName() const { return ClockId; }
 
-		AUDIOMIXER_API bool IsValid() const;
+		bool IsValid() const;
 		operator bool() const { return IsValid(); }
 
 		bool operator==(const FName& Name) const { return ClockId == Name; }
 
-		AUDIOMIXER_API bool DoesClockExist() const;
+		bool DoesClockExist() const;
 
-		AUDIOMIXER_API bool IsClockRunning() const;
+		bool IsClockRunning() const;
 
-		AUDIOMIXER_API Audio::FQuartzClockTickRate GetTickRate() const;
+		Audio::FQuartzClockTickRate GetTickRate() const;
 
-		AUDIOMIXER_API float GetEstimatedClockRunTimeSeconds() const;
+		float GetEstimatedClockRunTimeSeconds() const;
 
-		AUDIOMIXER_API FQuartzTransportTimeStamp GetCurrentClockTimestamp() const;
+		FQuartzTransportTimeStamp GetCurrentClockTimestamp() const;
 
-		AUDIOMIXER_API float GetDurationOfQuantizationTypeInSeconds(const EQuartzCommandQuantization& QuantizationType, float Multiplier) const;
+		float GetDurationOfQuantizationTypeInSeconds(const EQuartzCommandQuantization& QuantizationType, float Multiplier) const;
 
-		AUDIOMIXER_API float GetBeatProgressPercent(const EQuartzCommandQuantization& QuantizationType) const;
+		float GetBeatProgressPercent(const EQuartzCommandQuantization& QuantizationType) const;
 
 		// returns false if the clock is not valid or has shut down
-		AUDIOMIXER_API bool SendCommandToClock(TFunction<void(FQuartzClock*)> InCommand);
+		bool SendCommandToClock(TFunction<void(FQuartzClock*)> InCommand);
 
 		// implicit cast to underlying ID (FName)
 		operator const FName&() const { return ClockId; }
@@ -104,7 +104,8 @@ namespace Audio
 	 *		UpdateCachedState() updates a game-thread copy of data accessed via FQuartzClockProxy
 	 *		(see FQuartzClockState)
 	 */
-	class FQuartzClock
+
+	class AUDIOMIXER_API FQuartzClock : public FQuartzClockCommandQueueType::TConsumerBase<Audio::Quartz::IQuartzClock>
 	{
 	public:
 
@@ -112,88 +113,88 @@ namespace Audio
 		FQuartzClock(const FName& InName, const FQuartzClockSettings& InClockSettings, FQuartzClockManager* InOwningClockManagerPtr = nullptr);
 
 		// dtor
-		AUDIOMIXER_API ~FQuartzClock();
+		virtual ~FQuartzClock() override;
 
 		// Transport Control:
 		// alter the tick rate (take by-value to make sample-rate adjustments in-place)
-		AUDIOMIXER_API void ChangeTickRate(FQuartzClockTickRate InNewTickRate, int32 NumFramesLeft = 0);
+		void ChangeTickRate(FQuartzClockTickRate InNewTickRate, int32 NumFramesLeft = 0);
 
-		AUDIOMIXER_API void ChangeTimeSignature(const FQuartzTimeSignature& InNewTimeSignature);
+		void ChangeTimeSignature(const FQuartzTimeSignature& InNewTimeSignature);
 
-		AUDIOMIXER_API void Resume();
+		virtual void Resume() override;
 
-		AUDIOMIXER_API void Pause();
+		virtual void Pause() override;
 
-		AUDIOMIXER_API void Restart(bool bPause = true);
+		virtual void Restart(bool bPause = true) override;
 
-		AUDIOMIXER_API void Stop(bool CancelPendingEvents); // Pause + Restart
+		virtual void Stop(bool CancelPendingEvents) override;  // Pause + Restart
 
-		AUDIOMIXER_API void SetSampleRate(float InNewSampleRate);
+		void SetSampleRate(float InNewSampleRate);
 
-		AUDIOMIXER_API void ResetTransport(const int32 NumFramesToTickBeforeReset = 0);
+		void ResetTransport(const int32 NumFramesToTickBeforeReset = 0);
 
 		// (used for StartOtherClock command to handle the sub-tick as the target clock)
-		AUDIOMIXER_API void AddToTickDelay(int32 NumFramesOfDelayToAdd);
+		void AddToTickDelay(int32 NumFramesOfDelayToAdd);
 
 		// (used for StartOtherClock command to handle the sub-tick as the target clock)
-		AUDIOMIXER_API void SetTickDelay(int32 NumFramesOfDelay);
+		void SetTickDelay(int32 NumFramesOfDelay);
 
-		AUDIOMIXER_API void Shutdown();
+		void Shutdown();
 
 		// Getters:
-		AUDIOMIXER_API FQuartzClockTickRate GetTickRate();
+		FQuartzClockTickRate GetTickRate();
 
-		AUDIOMIXER_API FName GetName() const;
+		FName GetName() const;
 
-		AUDIOMIXER_API bool IgnoresFlush() const;
+		bool IgnoresFlush() const;
 
-		AUDIOMIXER_API bool DoesMatchSettings(const FQuartzClockSettings& InClockSettings) const;
+		bool DoesMatchSettings(const FQuartzClockSettings& InClockSettings) const;
 
-		AUDIOMIXER_API bool HasPendingEvents() const;
+		bool HasPendingEvents() const;
 
-		AUDIOMIXER_API int32 NumPendingEvents() const;
+		int32 NumPendingEvents() const;
 
-		AUDIOMIXER_API bool IsRunning() const;
+		bool IsRunning() const;
 
-		AUDIOMIXER_API float GetDurationOfQuantizationTypeInSeconds(const EQuartzCommandQuantization& QuantizationType, float Multiplier);
+		float GetDurationOfQuantizationTypeInSeconds(const EQuartzCommandQuantization& QuantizationType, float Multiplier);
 
-		AUDIOMIXER_API float GetBeatProgressPercent(const EQuartzCommandQuantization& QuantizationType) const;
+		float GetBeatProgressPercent(const EQuartzCommandQuantization& QuantizationType) const;
 
-		AUDIOMIXER_API FQuartzTransportTimeStamp GetCurrentTimestamp();
+		FQuartzTransportTimeStamp GetCurrentTimestamp();
 
-		AUDIOMIXER_API float GetEstimatedRunTime();
+		float GetEstimatedRunTime();
 
-		AUDIOMIXER_API FMixerDevice* GetMixerDevice();
+		FMixerDevice* GetMixerDevice();
 
-		AUDIOMIXER_API FMixerSourceManager* GetSourceManager();
+		FMixerSourceManager* GetSourceManager();
 
-		AUDIOMIXER_API FQuartzClockManager* GetClockManager();
+		FQuartzClockManager* GetClockManager();
 
-		AUDIOMIXER_API FQuartzClockCommandQueueWeakPtr GetCommandQueue() const;
+		FQuartzClockCommandQueueWeakPtr GetCommandQueue() const;
 
 		// Metronome Event Subscription:
-		AUDIOMIXER_API void SubscribeToTimeDivision(FQuartzGameThreadSubscriber InSubscriber, EQuartzCommandQuantization InQuantizationBoundary);
+		virtual void SubscribeToTimeDivision(FQuartzGameThreadSubscriber InSubscriber, EQuartzCommandQuantization InQuantizationBoundary) override;
 
-		AUDIOMIXER_API void SubscribeToAllTimeDivisions(FQuartzGameThreadSubscriber InSubscriber);
+		virtual void SubscribeToAllTimeDivisions(FQuartzGameThreadSubscriber InSubscriber) override;
 
-		AUDIOMIXER_API void UnsubscribeFromTimeDivision(FQuartzGameThreadSubscriber InSubscriber, EQuartzCommandQuantization InQuantizationBoundary);
+		virtual void UnsubscribeFromTimeDivision(FQuartzGameThreadSubscriber InSubscriber, EQuartzCommandQuantization InQuantizationBoundary) override;
 
-		AUDIOMIXER_API void UnsubscribeFromAllTimeDivisions(FQuartzGameThreadSubscriber InSubscriber);
+		virtual void UnsubscribeFromAllTimeDivisions(FQuartzGameThreadSubscriber InSubscriber) override;
 
 		// Quantized Command Management:
-		AUDIOMIXER_API void AddQuantizedCommand(FQuartzQuantizedRequestData& InQuantizedRequestData);
-		AUDIOMIXER_API void AddQuantizedCommand(FQuartzQuantizedCommandInitInfo& InQuantizationCommandInitInfo);
+		virtual void AddQuantizedCommand(FQuartzQuantizedRequestData& InQuantizedRequestData) override;
+		virtual void AddQuantizedCommand(FQuartzQuantizedCommandInitInfo& InQuantizationCommandInitInfo) override;
 
-		AUDIOMIXER_API void AddQuantizedCommand(FQuartzQuantizationBoundary InQuantizationBondary, TSharedPtr<IQuartzQuantizedCommand> InNewEvent);
+		virtual void AddQuantizedCommand(FQuartzQuantizationBoundary InQuantizationBoundary, TSharedPtr<IQuartzQuantizedCommand> InNewEvent) override;
 
-		AUDIOMIXER_API bool CancelQuantizedCommand(TSharedPtr<IQuartzQuantizedCommand> InCommandPtr);
+		bool CancelQuantizedCommand(TSharedPtr<IQuartzQuantizedCommand> InCommandPtr);
 		
 		// low-resolution clock update
 		// (not sample-accurate!, useful when running without an Audio Device)
-		AUDIOMIXER_API void LowResolutionTick(float InDeltaTimeSeconds);
+		void LowResolutionTick(float InDeltaTimeSeconds);
 
 		// sample accurate clock update
-		AUDIOMIXER_API void Tick(int32 InNumFramesUntilNextTick);
+		void Tick(int32 InNumFramesUntilNextTick);
 
 	private:
 		// Contains the pending command and the number of frames it has to wait to fire
@@ -215,7 +216,7 @@ namespace Audio
 
 		// mutex-protected update at the end of Tick()
 		FCriticalSection CachedClockStateCritSec;
-		AUDIOMIXER_API void UpdateCachedState();
+		void UpdateCachedState();
 
 		// data is cached when an FQuartzClock is ticked
 		struct FQuartzClockState
@@ -230,9 +231,9 @@ namespace Audio
 			
 		} CachedClockState;
 
-		AUDIOMIXER_API void TickInternal(int32 InNumFramesUntilNextTick, TArray<PendingCommand>& CommandsToTick, int32 FramesOfLatency = 0, int32 FramesOfDelay = 0);
+		void TickInternal(int32 InNumFramesUntilNextTick, TArray<PendingCommand>& CommandsToTick, int32 FramesOfLatency = 0, int32 FramesOfDelay = 0);
 
-		AUDIOMIXER_API bool CancelQuantizedCommandInternal(TSharedPtr<IQuartzQuantizedCommand> InCommandPtr, TArray<PendingCommand>& CommandsToTick);
+		bool CancelQuantizedCommandInternal(TSharedPtr<IQuartzQuantizedCommand> InCommandPtr, TArray<PendingCommand>& CommandsToTick);
 
 		// don't allow default ctor, a clock needs to be ready to be used
 		// by the clock manager / FMixerDevice once constructed
@@ -260,5 +261,4 @@ namespace Audio
 		int32 TickDelayLengthInFrames{ 0 };
 
 	}; // class FQuartzClock
-
 } // namespace Audio
diff --git a/Engine/Source/Runtime/AudioMixer/Public/Quartz/AudioMixerQuantizedCommands.h b/Engine/Source/Runtime/AudioMixer/Public/Quartz/AudioMixerQuantizedCommands.h
--- a/Engine/Source/Runtime/AudioMixer/Public/Quartz/AudioMixerQuantizedCommands.h
+++ b/Engine/Source/Runtime/AudioMixer/Public/Quartz/AudioMixerQuantizedCommands.h
@@ -30,6 +30,8 @@ namespace Audio
 		// for your implementation, a new EQuartzCommandType needs to be defined in QuartzQuantizationUtilities.h
 		virtual EQuartzCommandType GetCommandType() const { return EQuartzCommandType::PlaySound; };
 
+		virtual int32 OverrideFramesUntilExec(int32 NumFramesUntilExec);
+
 	protected:
 		TSharedPtr<FQuartzClock> OwningClockPtr{ nullptr };
 
diff --git a/Engine/Source/Runtime/AudioMixer/Public/Quartz/QuartzMetronome.h b/Engine/Source/Runtime/AudioMixer/Public/Quartz/QuartzMetronome.h
--- a/Engine/Source/Runtime/AudioMixer/Public/Quartz/QuartzMetronome.h
+++ b/Engine/Source/Runtime/AudioMixer/Public/Quartz/QuartzMetronome.h
@@ -12,9 +12,7 @@ class FQuartzTickableObject;
 
 namespace Audio
 {
-	template <typename T> class TQuartzShareableCommandQueue;
-
-	using MetronomeCommandQueuePtr = TSharedPtr<TQuartzShareableCommandQueue<FQuartzTickableObject>, ESPMode::ThreadSafe>;
+	using MetronomeCommandQueuePtr = Quartz::FQuartzGameThreadCommandQueuePtr; // todo: rename, "GameThread is not accurate"
 
 	// Class to track the passage of musical time, and allow subscribers to be notified when these musical events take place
 	class FQuartzMetronome
@@ -68,7 +66,7 @@ namespace Audio
 		// Helpers:
 		void RecalculateDurations();
 
-		void FireEvents(int32 EventFlags);
+		void FireEvents();
 
 		float CountNumSubdivisionsPerBar(EQuartzCommandQuantization InSubdivision) const;
 
@@ -89,39 +87,38 @@ namespace Audio
 		TArray<MetronomeCommandQueuePtr> MetronomeSubscriptionMatrix[static_cast<int32>(EQuartzCommandQuantization::Count)];
 		
 		// wrapper around our array so it can be indexed into by different Enums that represent musical time
-		struct FramesInTimeValue
+		struct FFramesInTimeValue
 		{
-		public:
-			// index operators for EQuartzCommandQuantization
-			double& operator[](EQuartzCommandQuantization InTimeValue)
-			{
-				return FramesInTimeValueInternal[static_cast<int32>(InTimeValue)];
-			}
-
-			const double& operator[](EQuartzCommandQuantization InTimeValue) const
-			{
-				return FramesInTimeValueInternal[static_cast<int32>(InTimeValue)];
-			}
-
-			// index operators for int32
-			double& operator[](int32 Index)
-			{
-				return FramesInTimeValueInternal[Index];
-			}
-
-			const double& operator[](int32 Index) const
-			{
-				return FramesInTimeValueInternal[Index];
-			}
+			void Reset();
+			
+			double& operator[](EQuartzCommandQuantization InTimeValue);
+			const double& operator[](EQuartzCommandQuantization InTimeValue) const;
+			double& operator[](int32 Index);
+			const double& operator[](int32 Index) const;
 
 			double FramesInTimeValueInternal[static_cast<int32>(EQuartzCommandQuantization::Count)]{ 0.0 };
 		};
 
+		struct FMetronomeEventEntry
+		{
+			TArray<int32> EventFrames;
+			void Reset();
+		};
+
+		struct FPendingMetronomeEvents
+		{
+			void Add(const EQuartzCommandQuantization InDuration, const int32 InFrame);
+			void Reset();
+			bool HasPendingEvent(const EQuartzCommandQuantization InDuration) const;
+			
+			FMetronomeEventEntry CurrentMetronomeEvents[static_cast<int32>(EQuartzCommandQuantization::Count)];
+		} PendingMetronomeEvents;
+
 		// array of lengths of musical durations (in audio frames)
-		FramesInTimeValue MusicalDurationsInFrames;
+		FFramesInTimeValue MusicalDurationsInFrames;
 
 		// array of the number of audio frames left until the respective musical duration
-		FramesInTimeValue FramesLeftInMusicalDuration;
+		FFramesInTimeValue FramesLeftInMusicalDuration;
 
 		// optional array of pulse duration overrides (for odd meters)
 		TArray<double> PulseDurations;
diff --git a/Engine/Source/Runtime/AudioMixer/Public/Quartz/QuartzSubsystem.h b/Engine/Source/Runtime/AudioMixer/Public/Quartz/QuartzSubsystem.h
--- a/Engine/Source/Runtime/AudioMixer/Public/Quartz/QuartzSubsystem.h
+++ b/Engine/Source/Runtime/AudioMixer/Public/Quartz/QuartzSubsystem.h
@@ -3,7 +3,6 @@
 #pragma once
 
 #include "CoreMinimal.h"
-#include "UObject/ObjectMacros.h"
 #include "Subsystems/WorldSubsystem.h"
 #include "Quartz/AudioMixerClockManager.h"
 #include "Sound/QuartzQuantizationUtilities.h"
@@ -16,15 +15,13 @@ namespace Audio
 {
 	class FMixerDevice;
 	class FQuartzClockManager;
-
-	template<class ListenerType>
-	class TQuartzShareableCommandQueue;
+	class FQuartzShareableCommandQueue;
 
 }
 
 class FQuartzTickableObject;
 class UQuartzClockHandle;
-using MetronomeCommandQueuePtr = TSharedPtr<Audio::TQuartzShareableCommandQueue<FQuartzTickableObject>, ESPMode::ThreadSafe>;
+using MetronomeCommandQueuePtr = TSharedPtr<Audio::FQuartzShareableCommandQueue, ESPMode::ThreadSafe>;
 
 
 
@@ -170,7 +167,6 @@ private:
 	AUDIOMIXER_API void PruneStaleProxies();
 	static AUDIOMIXER_API void PruneStaleProxiesInternal(TArray<Audio::FQuartzClockProxy>& ContainerToPrune);
 
-
 	// sharable tickable object manager to allow for non-UObject subscription / un-subscription
 	TSharedPtr<FQuartzTickableObjectsManager> TickableObjectManagerPtr { MakeShared<FQuartzTickableObjectsManager>() };
 
diff --git a/Engine/Source/Runtime/Engine/Classes/Components/AudioComponent.h b/Engine/Source/Runtime/Engine/Classes/Components/AudioComponent.h
--- a/Engine/Source/Runtime/Engine/Classes/Components/AudioComponent.h
+++ b/Engine/Source/Runtime/Engine/Classes/Components/AudioComponent.h
@@ -568,7 +568,7 @@ private:
 	// Data to hold pending quartz commands 
 	struct FAudioComponentPendingQuartzCommandData
 	{
-		FQuartzQuantizationBoundary AnticapatoryBoundary;
+		FQuartzQuantizationBoundary AnticipatoryBoundary;
 		FOnQuartzCommandEventBP Delegate;
 		float StartTime{ 0.0f };
 		float FadeDuration{ 0.0f };
diff --git a/Engine/Source/Runtime/Engine/Classes/Sound/QuartzCommandQueue.h b/Engine/Source/Runtime/Engine/Classes/Sound/QuartzCommandQueue.h
new file mode 100644
--- /dev/null
+++ b/Engine/Source/Runtime/Engine/Classes/Sound/QuartzCommandQueue.h
@@ -0,0 +1,77 @@
+// Copyright Epic Games, Inc. All Rights Reserved.
+
+#pragma once
+
+#include <type_traits>
+#include "Containers/ConsumeAllMpmcQueue.h"
+#include "QuartzCompileTimeVisitor.h"
+
+
+// forwards
+namespace Audio::Quartz::PrivateDefs
+{
+	template <typename... Ts>
+	class  TQuartzCommandQueue : public TVisitorPatternBase<Ts...>
+	{
+	private:
+		// type aliasing for private code...
+		template <class ConsumerBaseType>
+		using TCommand = typename TVisitorPatternBase<Ts...>::template TVisitWithLambda<ConsumerBaseType>;
+		using ICommandBase = typename TVisitorPatternBase<Ts...>::IVisitorBase;
+		
+	public:
+		// type aliasing for client code...
+		// IConsumerBase supports polymorphic access / containers of consumers
+		using IConsumerBase = typename TVisitorPatternBase<Ts...>::IListenerBase;
+
+		// TConsumerBase is the template base class which concrete consumer types should inherit from
+		template <typename... ConsumerInterfaces>
+		using TConsumerBase = typename TVisitorPatternBase<Ts...>::template TElementBase<ConsumerInterfaces...>;
+	
+	public:
+		// begin TQuartzCommandQueue interface
+		TQuartzCommandQueue() = default;
+		~TQuartzCommandQueue() = default;
+
+		/**
+			PushLambda():
+			* push a function that takes a TargetInterface& as an argument.
+			* Listeners will execute this code with (*static_cast<TargetInterface*>(this)) as the input()
+		* */
+		template <class TargetInterface>
+		void PushLambda(TFunction<void(TargetInterface&)> InLambda)
+		{
+			Queue.ProduceItem(MakeUnique<TCommand<TargetInterface>>(MoveTempIfPossible(InLambda)));
+		}
+
+		// note: this is not currently needed for this implementation, but is an example of
+		// polymorphic consumer usage and having multiple listeners listen to a single
+		// set of pending commands.
+		/**
+		void PumpCommandQueue(TArray<IConsumerBase*>& InListeners)
+		{
+			Queue.ConsumeAllFifo([&InListeners](TUniquePtr<ICommandBase> InCommand)
+			{
+				for(TUniquePtr<ICommandBase> Listener : InListeners)
+				{
+					if (ensure(InCommand))
+					{
+						Listener->Accept(*InCommand);
+					}
+				}
+			});
+		}
+	**/
+
+		void PumpCommandQueue(IConsumerBase& InListener)
+		{
+			Queue.ConsumeAllFifo([&InListener](TUniquePtr<ICommandBase> InCommand)
+			{
+				InListener.Accept(*InCommand);
+			});
+		}
+
+	private:
+		UE::TConsumeAllMpmcQueue<TUniquePtr<ICommandBase>> Queue;
+	}; // class TQuartzCommandQueue
+} // namespace Audio::Quartz::PrivateDefs
\ No newline at end of file
diff --git a/Engine/Source/Runtime/Engine/Classes/Sound/QuartzCompileTimeVisitor.h b/Engine/Source/Runtime/Engine/Classes/Sound/QuartzCompileTimeVisitor.h
new file mode 100644
--- /dev/null
+++ b/Engine/Source/Runtime/Engine/Classes/Sound/QuartzCompileTimeVisitor.h
@@ -0,0 +1,179 @@
+// Copyright Epic Games, Inc. All Rights Reserved.
+
+#pragma once
+
+#include <type_traits>
+
+// forwards
+namespace Audio::Quartz::PrivateDefs
+{
+	// Notes on the Visitor Pattern (https://en.wikipedia.org/wiki/Visitor_pattern)
+	// this is a pattern that helps solve two things:
+	// - adds operations for a class without modifying the class itself
+	// 		(and more importantly for Quartz's use-case:)
+	// - implements Double-Dispatch (https://en.wikipedia.org/wiki/Double_dispatch)
+	// 
+	// C++ supports single-dispatch through polymorphism, where a concrete function
+	// called depends on the dynamic type of a SINGLE object.
+	// 
+	// i.e.:  MyBasePtr->DoThing(MyConcreteType);
+	// 
+	// Double dispatch is being able resolve to a concrete function for how
+	// TWO dynamic types should interact
+	//
+	// i.e.: MyBaseA_ptr->DoThing(MyBaseB_ptr)
+	//
+	// where we could resolve to different concrete functions for each combination of 
+	// the RUNTIME TYPES of objects derived from MyBaseA and MyBaseB.
+	//
+	// Concretely, Quartz has metronome listeners, command listeners, etc
+	// and has things that can be both a metronome and a command listener.
+	//
+	// The DOWNSIDE to the visitor pattern is the element (or listener) types must be known
+	// at compile time.  This makes it hard to write reusable objects like command queues.
+	// 
+	// This implementation uses some template metaprogramming to abstract away the visitor pattern
+	// and let client code build reusable things like command queues and FSMs without needing to know
+	// the final concrete types.
+	//
+	// It also helps avoid some diamond inheritance problems that are easy to run into
+	// when attempting to have a "consumer" base type, and then types that can be multiple
+	// kinds of consumers.
+
+	
+	/** class TVisitorPatternBase:
+	* Summary:
+	* 	Used as a base class to implement the visitor pattern for the provided Ts types.
+	* 
+	* Examples:
+	* 	For an example usage for client code derived from this see TQuartzCommandQueue
+	* 	For an example of TQuartzCommandQueue client code that accepts multiple visitors see FQuartzTickableObject
+	* 
+	* 
+	* 
+	**/
+	template <typename... Ts>
+	class  TVisitorPatternBase
+	{ 
+	private:
+		/**
+		IVisitorInterface and IVisitorBase:
+		* These classes use template metaprogramming to recursively build
+		* a virtual listener interface (IVisitorBase) for our visitor pattern
+		*  with a virtual Visit(T&) overload for each type in the parameter pack (Ts...)
+		*
+		*  for example, if we want an IVisitorBase class which looks like:
+		*  class IVisitorBase
+		*  {
+		*		public:
+		*			virtual void Visit(ConsumerTypeA& InConsumer) {}
+		*			virtual void Visit(ConsumerTypeB& InConsumer) {}
+		*			virtual void Visit(ConsumerTypeC& InConsumer) {}
+		*  };
+		*
+		*  TVisitorInterface expands a template  parameter pack so that
+		*  IVisitorBase ends up with a virtual Visit() overload for each type in
+		*  the Interfaces parameter pack
+		* */
+		
+		// Forward declaration for the primary template
+		template <typename... Interfaces>
+		class TVisitorInterface;
+		
+		// Specialization for the base case
+		template <typename Interface>
+		class TVisitorInterface<Interface>
+		{
+		public:
+			virtual ~TVisitorInterface<Interface>() = default;
+			virtual void Visit(Interface&) {};
+		}; // class TVisitorInterface
+
+		// Specialization for the recursive case
+		template <typename First, typename... Rest>
+		class TVisitorInterface<First, Rest...> : public TVisitorInterface<Rest...>
+		{
+		public:
+			// Use Visit methods from the base class
+			using TVisitorInterface<Rest...>::Visit;
+
+			// Override Visit for the first type in the pack
+			virtual void Visit(First&) {};
+
+		}; // class TVisitorInterface
+
+	protected:
+		class IVisitorBase : public TVisitorInterface<Ts...>
+		{
+			// todo: comment about the inherited interface
+		}; // class IVisitorBase
+
+		/*
+		IListenerBase: 
+		* this should not be inherited from directly, but is public for
+		* polymorphic access to client listeners via IListenerBase*
+		*/
+
+		class IListenerBase
+		{
+		public:
+			virtual ~IListenerBase() = default;
+			virtual void Accept(IVisitorBase&) = 0;
+
+		}; // class IListenerBase
+
+	public:
+		/**
+		TElementBase:
+		* Client code defines listener interfaces, and then the client's concrete types
+		* should inherit from those through this template base class.
+		* (Concrete types can inherit from multiple listener interfaces)
+		* the automatically-defined Accept method will loop through the interfaces
+		* and cast this to each of the interface types to visit.
+		* 
+		* */
+		template <typename... ListenerInterfaces>
+		class TElementBase : public IListenerBase, public ListenerInterfaces...
+		{
+		public:
+			virtual void Accept(IVisitorBase& InVisitor) override final
+			{
+				(..., (void)InVisitor.Visit(*static_cast<ListenerInterfaces*>(this)));
+			}
+		}; // class TElementBase
+
+	protected:
+		/**
+		TVisitWithLambda
+		* This templatized visitor is used in ::PushLambda()
+		* external code does not need to worry about the visitor pattern
+		* */
+		template <class TargetInterface>
+		class TVisitWithLambda : public IVisitorBase
+		{
+		public:
+			// Bring all versions of Visit into the scope
+			// otherwise the compiler complains we hide the ones we aren't overriding
+			using IVisitorBase::Visit; 
+			
+			TVisitWithLambda() = delete;
+
+			TVisitWithLambda<TargetInterface>(TFunction<void(TargetInterface&)> InLambda)
+			: Lambda (MoveTempIfPossible(InLambda))
+			{}
+
+			virtual ~TVisitWithLambda() override = default;
+
+			virtual void Visit(TargetInterface& InTarget) override final
+			{
+				if(Lambda)
+				{
+					Lambda(InTarget);
+				}
+			}
+
+		private:
+			TFunction<void(TargetInterface&)> Lambda;
+		}; // class TVisitWithLambda
+	}; // class TVisitorPatternBase
+} // namespace Audio::Quartz::PrivateDefs
\ No newline at end of file
diff --git a/Engine/Source/Runtime/Engine/Classes/Sound/QuartzInterfaces.h b/Engine/Source/Runtime/Engine/Classes/Sound/QuartzInterfaces.h
new file mode 100644
--- /dev/null
+++ b/Engine/Source/Runtime/Engine/Classes/Sound/QuartzInterfaces.h
@@ -0,0 +1,90 @@
+// Copyright Epic Games, Inc. All Rights Reserved.
+
+#pragma once
+
+#include "QuartzQuantizationUtilities.h"
+#include "QuartzCommandQueue.h"
+
+namespace Audio
+{
+	// Struct used to communicate command state back to the game play thread
+	struct FQuartzQuantizedCommandDelegateData : public FQuartzCrossThreadMessage
+	{
+		EQuartzCommandType CommandType;
+		EQuartzCommandDelegateSubType DelegateSubType;
+
+		// ID so the clock handle knows which delegate to fire
+		int32 DelegateID{ -1 };
+
+	}; // struct FQuartzQuantizedCommandDelegateData
+
+	// Struct used to communicate metronome events back to the game play thread
+	struct FQuartzMetronomeDelegateData : public FQuartzCrossThreadMessage
+	{
+		int32 Bar;
+		int32 Beat;
+		float BeatFraction;
+		EQuartzCommandQuantization Quantization;
+		FName ClockName;
+		int32 FrameOffset = 0;
+	}; // struct FQuartzMetronomeDelegateData
+
+	// Struct used to queue events to be sent to the Audio Render thread closer to their start time
+	ENGINE_API struct FQuartzQueueCommandData : public FQuartzCrossThreadMessage
+	{
+		FAudioComponentCommandInfo AudioComponentCommandInfo;
+		FName ClockName;
+
+		FQuartzQueueCommandData(const FAudioComponentCommandInfo& InAudioComponentCommandInfo, FName InClockName)
+			: AudioComponentCommandInfo(InAudioComponentCommandInfo)
+			, ClockName(InClockName)
+		{}
+	}; // struct FQuartzQueueCommandData
+}
+
+namespace Audio::Quartz
+{
+	class ENGINE_API IMetronomeEventListener
+	{
+	public:
+		virtual ~IMetronomeEventListener() = default;
+		virtual void OnMetronomeEvent(const FQuartzMetronomeDelegateData&) = 0;
+	}; // class IMetronomeEventListener
+
+	class ENGINE_API ICommandListener
+	{
+	public:
+		virtual ~ICommandListener() = default;
+		virtual void OnCommandEvent(const FQuartzQuantizedCommandDelegateData&) = 0;
+	}; // class ICommandListener
+
+	class ENGINE_API IQueueCommandListener
+	{
+	public:
+		virtual ~IQueueCommandListener() = default;
+		virtual void OnQueueCommandEvent(const FQuartzQueueCommandData&) = 0;
+	};
+
+	class ENGINE_API IQuartzClock
+	{
+	public:
+		virtual ~IQuartzClock() = default;
+
+		// Transport control
+		virtual void Resume() = 0;
+		virtual void Pause() = 0;
+		virtual void Restart(bool bPause = true) = 0;
+		virtual void Stop(bool CancelPendingEvents) = 0; // Pause + Restart
+
+		// Metronome Event Subscription:
+		virtual void SubscribeToTimeDivision(FQuartzGameThreadSubscriber InSubscriber, EQuartzCommandQuantization InQuantizationBoundary) = 0;
+		virtual void SubscribeToAllTimeDivisions(FQuartzGameThreadSubscriber InSubscriber) = 0;
+		virtual void UnsubscribeFromTimeDivision(FQuartzGameThreadSubscriber InSubscriber, EQuartzCommandQuantization InQuantizationBoundary) = 0;
+		virtual void UnsubscribeFromAllTimeDivisions(FQuartzGameThreadSubscriber InSubscriber) = 0;
+
+		// Quantized Command Management:
+		virtual void AddQuantizedCommand(FQuartzQuantizedRequestData& InQuantizedRequestData) = 0;
+		virtual void AddQuantizedCommand(FQuartzQuantizedCommandInitInfo& InQuantizationCommandInitInfo) = 0;
+		virtual void AddQuantizedCommand(FQuartzQuantizationBoundary InQuantizationBoundary, TSharedPtr<IQuartzQuantizedCommand> InNewEvent) = 0;
+	}; // IQuartzClock
+} // namespace Audio::Quartz
\ No newline at end of file
diff --git a/Engine/Source/Runtime/Engine/Classes/Sound/QuartzQuantizationUtilities.cpp b/Engine/Source/Runtime/Engine/Classes/Sound/QuartzQuantizationUtilities.cpp
--- a/Engine/Source/Runtime/Engine/Classes/Sound/QuartzQuantizationUtilities.cpp
+++ b/Engine/Source/Runtime/Engine/Classes/Sound/QuartzQuantizationUtilities.cpp
@@ -489,7 +489,7 @@ namespace Audio
 
 		if (Audio::FMixerDevice* MixerDevice = InCommandInitInfo.OwningClockPointer->GetMixerDevice())
 		{
-			MixerDevice->QuantizedEventClockManager.PushLatencyTrackerResult(FQuartzCrossThreadMessage::RequestRecieved());
+			MixerDevice->QuantizedEventClockManager.PushLatencyTrackerResult(FQuartzCrossThreadMessage::RequestReceived());
 		}
 
 		GameThreadSubscribers.Append(InCommandInitInfo.GameThreadSubscribers);
@@ -739,7 +739,7 @@ namespace Audio
 		Timer.StartTimer();
 	}
 
-	double FQuartzCrossThreadMessage::RequestRecieved() const
+	double FQuartzCrossThreadMessage::RequestReceived() const
 	{
 		Timer.StopTimer();
 		return GetResultsMilliseconds();
@@ -771,7 +771,10 @@ namespace Audio
 	{
 		if(ensure(Queue.IsValid()))
 		{
-			Queue->PushEvent(Data);
+			Queue->PushLambda<ICommandListener>([=](ICommandListener& InListener)
+			{
+				InListener.OnCommandEvent(Data);
+			});
 
 			// raise the flag if this was a CommandOnAboutToStart notification
 			if(!bHasBeenNotifiedOfAboutToStart)
@@ -785,7 +788,10 @@ namespace Audio
 	{
 		if(ensure(Queue.IsValid()))
 		{
-			Queue->PushEvent(Data);
+			Queue->PushLambda<IMetronomeEventListener>([=](IMetronomeEventListener& InListener)
+			{
+				InListener.OnMetronomeEvent(Data);
+			});
 		}
 	}
 
@@ -793,7 +799,10 @@ namespace Audio
 	{
 		if(ensure(Queue.IsValid()))
 		{
-			Queue->PushEvent(Data);
+			Queue->PushLambda<IQueueCommandListener>([=](IQueueCommandListener& InListener)
+			{
+				InListener.OnQueueCommandEvent(Data);
+			});
 		}
 	}
 
diff --git a/Engine/Source/Runtime/Engine/Classes/Sound/QuartzQuantizationUtilities.h b/Engine/Source/Runtime/Engine/Classes/Sound/QuartzQuantizationUtilities.h
--- a/Engine/Source/Runtime/Engine/Classes/Sound/QuartzQuantizationUtilities.h
+++ b/Engine/Source/Runtime/Engine/Classes/Sound/QuartzQuantizationUtilities.h
@@ -3,6 +3,7 @@
 #pragma once
 
 #include "CoreMinimal.h"
+#include "QuartzCommandQueue.h"
 #include "UObject/ObjectMacros.h"
 #include "UObject/Object.h"
 #include "UObject/Class.h"
@@ -30,9 +31,6 @@ namespace Audio
 	class FQuartzClock;
 	class FShareableQuartzCommandQueue;
 
-	template<typename T>
-	class TQuartzShareableCommandQueue;
-
 	class FMixerDevice;
 
 	struct FQuartzQuantizedCommandDelegateData;
@@ -386,7 +384,7 @@ namespace Audio
 
 		ENGINE_API void RequestSent();
 
-		ENGINE_API double RequestRecieved() const;
+		ENGINE_API double RequestReceived() const;
 
 		ENGINE_API double GetResultsMilliseconds() const;
 
@@ -425,10 +423,16 @@ namespace Audio
 		TOptional<TPair<EQuartzCommandQuantization, double>> OffsetAsDuration;
 
 	};
+	
+	namespace Quartz
+	{
+		class IQueueCommandListener;
+		class ICommandListener;
+		class IMetronomeEventListener;
 
-
-	using FQuartzGameThreadCommandQueue = Audio::TQuartzShareableCommandQueue<FQuartzTickableObject>;
-	using FQuartzGameThreadCommandQueuePtr = TSharedPtr<FQuartzGameThreadCommandQueue, ESPMode::ThreadSafe>;
+		using FQuartzSubscriberCommandQueue = Audio::Quartz::PrivateDefs::TQuartzCommandQueue<IMetronomeEventListener, ICommandListener, IQueueCommandListener>;
+		using FQuartzGameThreadCommandQueuePtr = TSharedPtr<FQuartzSubscriberCommandQueue, ESPMode::ThreadSafe>;
+	} // namespace Quartz
 
 	struct FQuartzGameThreadSubscriber
 	{
@@ -436,13 +440,13 @@ namespace Audio
 
 		// this is only for back-compat until metronomes support FQuartzGameThreadSubscribers instead of raw queue ptrs
 		// (for Metronome event offsets)
-		FQuartzGameThreadSubscriber(const FQuartzGameThreadCommandQueuePtr& InQueuePtr)
+		FQuartzGameThreadSubscriber(const Quartz::FQuartzGameThreadCommandQueuePtr& InQueuePtr)
 		: Queue(InQueuePtr)
 		{ }
 
 
 		// copy ctor
-		FQuartzGameThreadSubscriber(const FQuartzGameThreadCommandQueuePtr& InQueuePtr, FQuartzOffset InOffset)
+		FQuartzGameThreadSubscriber(const Quartz::FQuartzGameThreadCommandQueuePtr& InQueuePtr, FQuartzOffset InOffset)
 		: Offset(InOffset)
 		, Queue(InQueuePtr)
 		{ }
@@ -455,14 +459,13 @@ namespace Audio
 		// comparison
 		ENGINE_API bool operator==(const FQuartzGameThreadSubscriber& Other) const;
 
-		// todo: templatize to match teh underlying TQUartzShareableCommandQueue
 		// notify
 		ENGINE_API void PushEvent(const FQuartzQuantizedCommandDelegateData& Data);
 		ENGINE_API void PushEvent(const FQuartzMetronomeDelegateData& Data);
 		ENGINE_API void PushEvent(const FQuartzQueueCommandData& Data);
 
 		// allow implicit casting to the underlying queue
-		operator FQuartzGameThreadCommandQueuePtr() const { return Queue; }
+		operator Quartz::FQuartzGameThreadCommandQueuePtr() const { return Queue; }
 
 		// positive: anticipatory amount, negative:
 		ENGINE_API int32 FinalizeOffset(const FQuartzClockTickRate& TickRate);
@@ -472,7 +475,7 @@ namespace Audio
 
 	private:
 		FQuartzOffset Offset;
-		FQuartzGameThreadCommandQueuePtr Queue;
+		Quartz::FQuartzGameThreadCommandQueuePtr Queue;
 		bool bOffsetConvertedToFrames = false;
 		bool bHasBeenNotifiedOfAboutToStart = false;
 		int32 OffsetInAudioFrames = 0;
@@ -552,7 +555,7 @@ namespace Audio
 	// info that derived classes need can be added here
 	struct FQuartzQuantizedCommandInitInfo
 	{
-		FQuartzQuantizedCommandInitInfo() {}
+		FQuartzQuantizedCommandInitInfo() = default;
 
 		// conversion ctor from FQuartzQuantizedRequestData
 		ENGINE_API FQuartzQuantizedCommandInitInfo(const FQuartzQuantizedRequestData& RHS, float InSampleRate, int32 InSourceID = INDEX_NONE);
@@ -573,7 +576,7 @@ namespace Audio
 
 		// Audio Render thread-specific data:
 		TSharedPtr<Audio::FQuartzClock> OwningClockPointer{ nullptr };
-		float SampleRate;
+		float SampleRate{0};
 		int32 SourceID{ -1 };
 
 		// Number of frames used for any FramesTilExec overrides
@@ -675,13 +678,13 @@ struct FAudioComponentCommandInfo
 
  	FAudioComponentCommandInfo(Audio::FQuartzGameThreadSubscriber InSubscriber, FQuartzQuantizationBoundary InAnticaptoryBoundary)
 		: Subscriber(InSubscriber)
-		, AnticapatoryBoundary(InAnticaptoryBoundary)
+		, AnticipatoryBoundary(InAnticaptoryBoundary)
 	{
 		static uint32 CommandIDs = 0;
 		CommandID = CommandIDs++;
 	}
 
 	Audio::FQuartzGameThreadSubscriber Subscriber;
-	FQuartzQuantizationBoundary AnticapatoryBoundary;
+	FQuartzQuantizationBoundary AnticipatoryBoundary;
 	uint32 CommandID{ (uint32)INDEX_NONE };
 };
diff --git a/Engine/Source/Runtime/Engine/Classes/Sound/QuartzSubscription.cpp b/Engine/Source/Runtime/Engine/Classes/Sound/QuartzSubscription.cpp
--- a/Engine/Source/Runtime/Engine/Classes/Sound/QuartzSubscription.cpp
+++ b/Engine/Source/Runtime/Engine/Classes/Sound/QuartzSubscription.cpp
@@ -1,6 +1,7 @@
 // Copyright Epic Games, Inc. All Rights Reserved.
 
 #include "Sound/QuartzSubscription.h"
+
 #include "Quartz/QuartzSubsystem.h"
 #include "HAL/IConsoleManager.h"
 
@@ -13,24 +14,13 @@ FAutoConsoleVariableRef CVarDecrementSlotIndexOnStarted(
 	TEXT("1: New Behavior, 0: Old Behavior"),
 	ECVF_Default);
 
-namespace Audio
-{
-	FQuartzQueueCommandData::FQuartzQueueCommandData(const FAudioComponentCommandInfo& InAudioComponentCommandInfo, FName InClockName)
-	: AudioComponentCommandInfo(InAudioComponentCommandInfo)
-	, ClockName(InClockName)
-	{
-	}
-} // namespace Audio
-
 FQuartzTickableObject::FQuartzTickableObject()
 {}
 
-PRAGMA_DISABLE_DEPRECATION_WARNINGS
 FQuartzTickableObject::~FQuartzTickableObject()
 {
 	QuartzUnsubscribe();
 }
-PRAGMA_ENABLE_DEPRECATION_WARNINGS
 
 FQuartzTickableObject* FQuartzTickableObject::Init(UWorld* InWorldPtr)
 {
@@ -40,12 +30,10 @@ FQuartzTickableObject* FQuartzTickableObject::Init(UWorld* InWorldPtr)
 		return this;
 	}
 
-	PRAGMA_DISABLE_DEPRECATION_WARNINGS
 	if(!CommandQueuePtr.IsValid())
 	{
-		CommandQueuePtr = Audio::TQuartzShareableCommandQueue<FQuartzTickableObject>::Create();
+		CommandQueuePtr = MakeShared<FQuartzSubscriberCommandQueue>();
 	}
-	PRAGMA_ENABLE_DEPRECATION_WARNINGS
 
 	UQuartzSubsystem* QuartzSubsystemPtr = UQuartzSubsystem::Get(InWorldPtr);
 	QuartzSubscriptionToken.Subscribe(this, QuartzSubsystemPtr);
@@ -77,20 +65,13 @@ int32 FQuartzTickableObject::AddCommandDelegate(const FOnQuartzCommandEventBP& I
 	return SlotId;
 }
 
-
-UQuartzSubsystem* FQuartzTickableObject::GetQuartzSubsystem() const
-{
-	return nullptr;
-}
-
-
-void FQuartzTickableObject::ExecCommand(const Audio::FQuartzQuantizedCommandDelegateData& Data)
+void FQuartzTickableObject::OnCommandEvent(const Audio::FQuartzQuantizedCommandDelegateData& Data)
 {
 	checkSlow(Data.DelegateSubType < EQuartzCommandDelegateSubType::Count);
-	
+
 	if(const TSharedPtr<FQuartzTickableObjectsManager> ObjManagerPtr = QuartzSubscriptionToken.GetTickableObjectManager())
 	{
-		ObjManagerPtr->PushLatencyTrackerResult(Data.RequestRecieved());
+		ObjManagerPtr->PushLatencyTrackerResult(Data.RequestReceived());
 	}
 
 	// Broadcast to the BP delegate if we have one bound
@@ -114,7 +95,7 @@ void FQuartzTickableObject::ExecCommand(const Audio::FQuartzQuantizedCommandDele
 		// (end of a command)
 		bool bShouldDecrement = Data.DelegateSubType == EQuartzCommandDelegateSubType::CommandOnCanceled;
 		bShouldDecrement |= (DecrementSlotIndexOnStartedCvar && Data.DelegateSubType == EQuartzCommandDelegateSubType::CommandOnStarted);
-			
+
 		// are all the commands for this delegate done?
 		if (bShouldDecrement && (GameThreadEntry.RefCount.Decrement() <= 0))
 		{
@@ -128,11 +109,11 @@ void FQuartzTickableObject::ExecCommand(const Audio::FQuartzQuantizedCommandDele
 	ProcessCommand(Data);
 }
 
-void FQuartzTickableObject::ExecCommand(const Audio::FQuartzMetronomeDelegateData& Data)
+void FQuartzTickableObject::OnMetronomeEvent(const Audio::FQuartzMetronomeDelegateData& Data)
 {
-	if(const TSharedPtr<FQuartzTickableObjectsManager> ObjManagerPtr = QuartzSubscriptionToken.GetTickableObjectManager())
+	if (const TSharedPtr<FQuartzTickableObjectsManager> ObjManagerPtr = QuartzSubscriptionToken.GetTickableObjectManager())
 	{
-		ObjManagerPtr->PushLatencyTrackerResult(Data.RequestRecieved());
+		ObjManagerPtr->PushLatencyTrackerResult(Data.RequestReceived());
 	}
 
 	MetronomeDelegates[static_cast<int32>(Data.Quantization)].MulticastDelegate
@@ -142,7 +123,7 @@ void FQuartzTickableObject::ExecCommand(const Audio::FQuartzMetronomeDelegateDat
 	ProcessCommand(Data);
 }
 
-void FQuartzTickableObject::ExecCommand(const Audio::FQuartzQueueCommandData& Data)
+void FQuartzTickableObject::OnQueueCommandEvent(const Audio::FQuartzQueueCommandData& Data)
 {
 	// call base-class method
 	ProcessCommand(Data);
@@ -162,28 +143,22 @@ void FQuartzTickableObject::SetNotificationAnticipationAmountMusicalDuration(con
 
 Audio::FQuartzGameThreadSubscriber FQuartzTickableObject::GetQuartzSubscriber()
  {
-	PRAGMA_DISABLE_DEPRECATION_WARNINGS
  	if (!CommandQueuePtr.IsValid())
  	{
- 		CommandQueuePtr = MakeShared<Audio::TQuartzShareableCommandQueue<FQuartzTickableObject>, ESPMode::ThreadSafe>();
+ 		CommandQueuePtr = MakeShared<FQuartzSubscriberCommandQueue, ESPMode::ThreadSafe>();
  	}
 
  	return { CommandQueuePtr, NotificationOffset };
-	PRAGMA_ENABLE_DEPRECATION_WARNINGS
 }
 
 void FQuartzTickableObject::QuartzTick(float DeltaTime)
 {
-	PRAGMA_DISABLE_DEPRECATION_WARNINGS
-	CommandQueuePtr->PumpCommandQueue(this);
-	PRAGMA_ENABLE_DEPRECATION_WARNINGS
+	CommandQueuePtr->PumpCommandQueue(*this);
 }
 
 bool FQuartzTickableObject::QuartzIsTickable() const
 {
-	PRAGMA_DISABLE_DEPRECATION_WARNINGS
 	return CommandQueuePtr.IsValid();
-	PRAGMA_ENABLE_DEPRECATION_WARNINGS
 }
 
 void FQuartzTickableObject::AddMetronomeBpDelegate(EQuartzCommandQuantization InQuantizationBoundary, const FOnQuartzMetronomeEventBP& OnQuantizationEvent)
diff --git a/Engine/Source/Runtime/Engine/Classes/Sound/QuartzSubscription.h b/Engine/Source/Runtime/Engine/Classes/Sound/QuartzSubscription.h
--- a/Engine/Source/Runtime/Engine/Classes/Sound/QuartzSubscription.h
+++ b/Engine/Source/Runtime/Engine/Classes/Sound/QuartzSubscription.h
@@ -3,6 +3,8 @@
 #pragma once
 
 #include "QuartzSubscriptionToken.h"
+#include "Sound/QuartzInterfaces.h"
+#include "Sound/QuartzCommandQueue.h"
 #include "Sound/QuartzQuantizationUtilities.h"
 #include "Containers/ConsumeAllMpmcQueue.h"
 
@@ -17,103 +19,13 @@ namespace Audio
 	struct FQuartzGameThreadSubscriber;
 	class FQuartzClockProxy;
 
-	// Struct used to communicate command state back to the game play thread
-	struct FQuartzQuantizedCommandDelegateData : public FQuartzCrossThreadMessage
-	{
-		EQuartzCommandType CommandType;
-		EQuartzCommandDelegateSubType DelegateSubType;
-
-		// ID so the clock handle knows which delegate to fire
-		int32 DelegateID{ -1 };
-
-	}; // struct FQuartzQuantizedCommandDelegateData
-
-	// Struct used to communicate metronome events back to the game play thread
-	struct FQuartzMetronomeDelegateData : public FQuartzCrossThreadMessage
-	{
-		int32 Bar;
-		int32 Beat;
-		float BeatFraction;
-		EQuartzCommandQuantization Quantization;
-		FName ClockName;
-	}; // struct FQuartzMetronomeDelegateData
-
-	//Struct used to queue events to be sent to the Audio Render thread closer to their start time
-	struct FQuartzQueueCommandData : public FQuartzCrossThreadMessage
-	{
-		FAudioComponentCommandInfo AudioComponentCommandInfo;
-		FName ClockName;
-
-		ENGINE_API FQuartzQueueCommandData(const FAudioComponentCommandInfo& InAudioComponentCommandInfo, FName InClockName);
-	}; // struct FQuartzQueueCommandData
-
-
-	// old non-template\ version of TQuartzShareableCommandQueue
-	class UE_DEPRECATED(5.1, "Message") FShareableQuartzCommandQueue;
-	class FShareableQuartzCommandQueue
+	// old version of TQuartzCommandQueue
+	template<class ListenerType>
+	class UE_DEPRECATED(5.5, "Message") TQuartzShareableCommandQueue
 	{
 	};
-
-	/**
-	*	Template class for mono-directional MPSC command queues 
-	* 
-	*   in order to enforce thread-safe access to the object executing the commands,
-	*	"listener type" is the type of the object that is being accessed in the commands
-	*	that object will have to provide a 'this' ptr (of type ListenerType) in order to 
-	*   invoke the commands on itself. (The lambdas do NOT and should NOT cache a ptr or
-	*	reference to the target).
-	* 
-	*	User-provided lambdas can take any (single) argument type T in PushEvent()
-	*	but there must exist a ListenerType::ExecCommand(T) overload for any PushEvent(T) instantiated. 
-	* 
-	*	(see FQuartzTickableObject and FQuartzClock as examples)
-	* 
-	**/
-	template <class ListenerType>
-	class TQuartzShareableCommandQueue
-	{
-	public:
-		// ctor
-		TQuartzShareableCommandQueue() {}
-
-		// dtor
-		~TQuartzShareableCommandQueue() {}
-
-		// static helper to create a new sharable queue
-		static TSharedPtr<TQuartzShareableCommandQueue<ListenerType>, ESPMode::ThreadSafe> Create()
-		{
-			return MakeShared<TQuartzShareableCommandQueue<ListenerType>, ESPMode::ThreadSafe>();
-		}
-
-		// note: ListenerType must have a ExecCommand() overload for each instantiation of this method
-		template <typename T>
-		void PushEvent(const T& Data)
-		{
-			CommandQueue.ProduceItem([InData = Data](ListenerType* Listener) { Listener->ExecCommand(InData); });
-		}
-
-		void PushCommand(TFunction<void(ListenerType*)> InCommand)
-		{
-			CommandQueue.ProduceItem([=](ListenerType* Listener) { InCommand(Listener); });
-		}
-
-		void PumpCommandQueue(ListenerType* InListener)
-		{
-			CommandQueue.ConsumeAllFifo([InListener](TFunction<void(ListenerType*)> Command)
-			{
-				Command(InListener);
-			});
-		}
-
-	private:
-		UE::TConsumeAllMpmcQueue<TFunction<void(ListenerType*)>> CommandQueue;
-	};
-
 } // namespace Audio
 
-using FQuartzGameThreadCommandQueue = Audio::TQuartzShareableCommandQueue<FQuartzTickableObject>;
-using FQuartzGameThreadCommandQueuePtr = TSharedPtr<FQuartzGameThreadCommandQueue, ESPMode::ThreadSafe>;
-
 /**
  *	FQuartzTickableObject
  *
@@ -123,49 +35,54 @@ using FQuartzGameThreadCommandQueuePtr = TSharedPtr<FQuartzGameThreadCommandQueu
  *		It is a wrapper around TQuartzShareableCommandQueue.
  *		(see UQuartzClockHandle or UAudioComponent as implementation examples)
  */
-class FQuartzTickableObject
+using namespace Audio::Quartz;
+
+// TODO: comment up why we listen to these interfaces
+class ENGINE_API FQuartzTickableObject
+	: public FQuartzSubscriberCommandQueue::TConsumerBase<
+		  IMetronomeEventListener
+		, ICommandListener
+		, IQueueCommandListener
+>
 {
 public:
 	// ctor
-	ENGINE_API FQuartzTickableObject();
+	FQuartzTickableObject();
 
-	// explicitly defaulted ctors (to disable deprecation warnings in compiler-generated functions)
-    PRAGMA_DISABLE_DEPRECATION_WARNINGS
     FQuartzTickableObject(const FQuartzTickableObject& Other) = default;
     FQuartzTickableObject& operator=(const FQuartzTickableObject&) = default;
-    ENGINE_API PRAGMA_ENABLE_DEPRECATION_WARNINGS
 
 	// dtor
-	virtual ~FQuartzTickableObject();
+	virtual ~FQuartzTickableObject() override;
 
-	ENGINE_API FQuartzTickableObject* Init(UWorld* InWorldPtr);
+	FQuartzTickableObject* Init(UWorld* InWorldPtr);
 
 	// called by the associated QuartzSubsystem
-	ENGINE_API void QuartzTick(float DeltaTime);
-
-	ENGINE_API bool QuartzIsTickable() const;
+	void QuartzTick(float DeltaTime);
 
-	UE_DEPRECATED(5.1, "Derived classes should have access to their own UWorld, this function will always return null")
-	UWorld* QuartzGetWorld() const { return nullptr; }
+	bool QuartzIsTickable() const;
 
-	ENGINE_API void AddMetronomeBpDelegate(EQuartzCommandQuantization InQuantizationBoundary, const FOnQuartzMetronomeEventBP& OnQuantizationEvent);
+	void AddMetronomeBpDelegate(EQuartzCommandQuantization InQuantizationBoundary, const FOnQuartzMetronomeEventBP& OnQuantizationEvent);
 
 	bool IsInitialized() const { return QuartzSubscriptionToken.IsSubscribed(); }
 
-	ENGINE_API Audio::FQuartzGameThreadSubscriber GetQuartzSubscriber();
+	Audio::FQuartzGameThreadSubscriber GetQuartzSubscriber();
 
-	ENGINE_API int32 AddCommandDelegate(const FOnQuartzCommandEventBP& InDelegate);
+	int32 AddCommandDelegate(const FOnQuartzCommandEventBP& InDelegate);
 
-	UE_DEPRECATED(5.1, "use FQuartzTickableObject::AddCommandDelegate(const FOnQuartzCommandEventBP& InDelegate) insead")
-	int32 AddCommandDelegate(const FOnQuartzCommandEventBP& InDelegate, TArray<FQuartzGameThreadCommandQueuePtr>& TargetSubscriberArray){ return -1; }
-
-	UE_DEPRECATED(5.1, "This object no longer holds any UObject references. Caller should have their own UWorld* and use static 'UQuartzSubsystem::Get()' instead.")
-	ENGINE_API UQuartzSubsystem* GetQuartzSubsystem() const;
+	UE_DEPRECATED(5.5, "This should not be called directly, use the ICommandListener interface instead.")
+	void ExecCommand(const Audio::FQuartzQuantizedCommandDelegateData& Data) { OnCommandEvent(Data); }
+	
+	UE_DEPRECATED(5.5, "This should not be called directly, use the IMetronomeEventListener interface instead.")
+	void ExecCommand(const Audio::FQuartzMetronomeDelegateData& Data) { OnMetronomeEvent(Data); }
+	
+	UE_DEPRECATED(5.5, "This should not be called directly, use the IQueueCommandListener interface instead.")
+	void ExecCommand(const Audio::FQuartzQueueCommandData& Data) { OnQueueCommandEvent(Data);}
 
 	// required by TQuartzShareableCommandQueue template
-	ENGINE_API void ExecCommand(const Audio::FQuartzQuantizedCommandDelegateData& Data);
-	ENGINE_API void ExecCommand(const Audio::FQuartzMetronomeDelegateData& Data);
-	ENGINE_API void ExecCommand(const Audio::FQuartzQueueCommandData& Data);
+	virtual void OnCommandEvent(const Audio::FQuartzQuantizedCommandDelegateData& Data) override;
+	virtual void OnMetronomeEvent(const Audio::FQuartzMetronomeDelegateData& Data) override;
+	virtual void OnQueueCommandEvent(const Audio::FQuartzQueueCommandData& Data) override;
 
 	// virtual interface (ExecCommand will forward the data to derived classes' ProcessCommand() call)
 	virtual void ProcessCommand(const Audio::FQuartzQuantizedCommandDelegateData& Data) {}
@@ -175,10 +92,10 @@ public:
 	const Audio::FQuartzOffset& GetQuartzOffset() const { return NotificationOffset; }
 
 protected:
-	ENGINE_API void SetNotificationAnticipationAmountMilliseconds(const double Milliseconds);
-	ENGINE_API void SetNotificationAnticipationAmountMusicalDuration(const EQuartzCommandQuantization Duration,  const double Multiplier);
+	void SetNotificationAnticipationAmountMilliseconds(const double Milliseconds);
+	void SetNotificationAnticipationAmountMusicalDuration(const EQuartzCommandQuantization Duration,  const double Multiplier);
 
-	ENGINE_API void QuartzUnsubscribe();
+	void QuartzUnsubscribe();
 
 private:
 	struct FMetronomeDelegateGameThreadData
@@ -198,14 +115,9 @@ private:
 
 	TArray<TFunction<void(FQuartzTickableObject*)>> TempCommandQueue;
 
-public:
-	// deprecate public access
-	UE_DEPRECATED(5.1, "FQuartzTickableObject::CommandQueuePtr should no longer be accessed directly. (This member will be private in a future engine version).  Use GetQuartzSubscriber() instead")
-	TSharedPtr<Audio::TQuartzShareableCommandQueue<FQuartzTickableObject>, ESPMode::ThreadSafe> CommandQueuePtr;
-
 private:
 	Audio::FQuartzOffset NotificationOffset;
-
+	FQuartzGameThreadCommandQueuePtr CommandQueuePtr;
 	FQuartzSubscriptionToken QuartzSubscriptionToken;
 }; // class FQuartzTickableObject
 
diff --git a/Engine/Source/Runtime/Engine/Private/Components/AudioComponent.cpp b/Engine/Source/Runtime/Engine/Private/Components/AudioComponent.cpp
--- a/Engine/Source/Runtime/Engine/Private/Components/AudioComponent.cpp
+++ b/Engine/Source/Runtime/Engine/Private/Components/AudioComponent.cpp
@@ -530,7 +530,7 @@ void UAudioComponent::PlayQuantized(
 	const bool bCommandResetsClock = InQuantizationBoundary.bResetClockOnQueued;
 	if (bStealVoiceSlot || bClockIsNotRunning || bCommandResetsClock)
 	{
-		AudioComponentQuartzCommandData.AnticapatoryBoundary = NewComponentCommandInfo.AnticapatoryBoundary = InQuantizationBoundary; // use the desired target boundary
+		AudioComponentQuartzCommandData.AnticipatoryBoundary = NewComponentCommandInfo.AnticipatoryBoundary = InQuantizationBoundary; // use the desired target boundary
 
 		// Add to the list of pending data
 		// todo: avoid this copy for OUR call to PlayQueuedQuantizedInternal()
@@ -579,7 +579,7 @@ void UAudioComponent::PlayQueuedQuantizedInternal(const UObject* WorldContextObj
 				// confirm a valid handle
 				if (Handle != nullptr)
 				{
-					InternalRequestData.QuantizedRequestData = UQuartzSubsystem::CreateRequestDataForSchedulePlaySound(Handle, PendingData.Delegate, PendingData.AnticapatoryBoundary);
+					InternalRequestData.QuantizedRequestData = UQuartzSubsystem::CreateRequestDataForSchedulePlaySound(Handle, PendingData.Delegate, PendingData.AnticipatoryBoundary);
 					UGameplayStatics::PrimeSound(Sound);
 				}
 
