Fix for rare condition with SoundWave GC and playing Wave.

diff --git a/Engine/Source/Runtime/AudioMixer/Private/AudioMixerSourceBuffer.cpp b/Engine/Source/Runtime/AudioMixer/Private/AudioMixerSourceBuffer.cpp
--- a/Engine/Source/Runtime/AudioMixer/Private/AudioMixerSourceBuffer.cpp
+++ b/Engine/Source/Runtime/AudioMixer/Private/AudioMixerSourceBuffer.cpp
@@ -199,7 +199,18 @@ namespace Audio
 	{
 		FScopeTryLock Lock(&SoundWaveCritSec);
 
-		if (!Lock.IsLocked() || (NumBuffersQeueued == 0 && bBufferFinished) || (bProcedural && !SoundWave) || (bHasError))
+		// If the buffer is flagged as complete and there's nothing queued remaining.
+		const bool bBufferCompleted = (NumBuffersQeueued == 0 && bBufferFinished);
+		
+		// If we're procedural we must have a procedural SoundWave pointer to continue.
+		const bool bProceduralStateBad = (bProcedural && !SoundWave);
+		
+		// If we're non-procedural and we don't have a decoder, bail. This can happen when the wave is GC'd.
+		// The Decoder and SoundWave is deleted on the GameThread via FMixerSourceBuffer::OnBeginDestroy
+		// Although this is bad state it's not an error, so just bail here.
+		const bool bDecompressionStateBad = (!bProcedural && DecompressionState == nullptr);
+
+		if (!Lock.IsLocked() || bBufferCompleted || bProceduralStateBad || bDecompressionStateBad || bHasError )
 		{
 			return;
 		}
