Submix Send Stomp Fix

diff --git a/Engine/Source/Runtime/AudioMixer/Private/AudioMixerSource.cpp b/Engine/Source/Runtime/AudioMixer/Private/AudioMixerSource.cpp
--- a/Engine/Source/Runtime/AudioMixer/Private/AudioMixerSource.cpp
+++ b/Engine/Source/Runtime/AudioMixer/Private/AudioMixerSource.cpp
@@ -91,33 +91,6 @@ namespace Audio
 
 			return SoundClass;
 		}
-
-		template<typename SendInfo>
-		void ClearPreviousSubmixSends(const TArray<SendInfo>& InPreviousSendInfos, const TArray<SendInfo>& InNewSendInfos, FMixerDevice* InMixerDevice, FMixerSourceVoice* InMixerSourceVoice)
-		{
-			// Loop through every previous send setting
-			for (const SendInfo& PreviousSendSetting : InPreviousSendInfos)
-			{
-				bool bFound = false;
-
-				// See if it's in the current send list
-				for (const SendInfo&  CurrentSendSettings : InNewSendInfos)
-				{
-					if (CurrentSendSettings.SoundSubmix == PreviousSendSetting.SoundSubmix)
-					{
-						bFound = true;
-						break;
-					}
-				}
-
-				// If it's not in the current send list, add to submixes to clear
-				if (!bFound)
-				{
-					FMixerSubmixPtr SubmixPtr = InMixerDevice->GetSubmixInstance(PreviousSendSetting.SoundSubmix).Pin();
-					InMixerSourceVoice->ClearSubmixSendInfo(SubmixPtr);
-				}
-			}
-		}
 		
 	} // namespace MixerSourcePrivate
 
@@ -583,6 +556,14 @@ namespace Audio
 			}
 		}
 
+		// Reset all 'previous' state.
+		PreviousSubmixResolved.Reset();
+		bPreviousBusEnablement = false;
+		bPreviousBaseSubmixEnablement = false;
+		PreviousAzimuth = -1.f;
+		PreviousPlaybackPercent = 0.f;
+		PreviousSubmixSends.Reset();
+
 		// Unfortunately, we need to know if this is a vorbis source since channel maps are different for 5.1 vorbis files
 		bIsVorbis = WaveData->bDecompressedFromOgg;
 
@@ -1768,13 +1749,16 @@ namespace Audio
 		}
 	}
 	
-	void FMixerSource::UpdateSubmixSendLevels(const FSoundSubmixSendInfoBase& InSendInfo, const EMixerSourceSubmixSendStage InSendStage)
+	void FMixerSource::UpdateSubmixSendLevels(const FSoundSubmixSendInfoBase& InSendInfo, const EMixerSourceSubmixSendStage InSendStage, TSet<FMixerSubmixWeakPtr>& OutTouchedSubmixes)
 	{
 		if (InSendInfo.SoundSubmix != nullptr)
 		{
 			const FMixerSubmixWeakPtr SubmixInstance = MixerDevice->GetSubmixInstance(InSendInfo.SoundSubmix);
 			float SendLevel = 1.0f;
 
+			// Add it to our touched submix list.
+			OutTouchedSubmixes.Add(SubmixInstance);
+
 			// calculate send level based on distance if that method is enabled
 			if (!WaveInstance->bEnableSubmixSends)
 			{
@@ -1892,32 +1876,42 @@ namespace Audio
 				SubmixPtr = MixerDevice->GetBaseDefaultSubmix(); // This will try base default and fall back to master if that fails.
 			}
 
-
 			MixerSourceVoice->SetSubmixSendInfo(SubmixPtr, WaveInstance->bEnableBaseSubmix);
 			bPreviousBaseSubmixEnablement = WaveInstance->bEnableBaseSubmix;
 			PreviousSubmixResolved = SubmixPtr;
 			PrevousSubmix = SubmixKey;
 		}
 
+		// We clear sends that aren't used between updates. So tally up the ones that are used.
+		// Including the submix itself. 
+		// It's okay to use "previous" submix here as it's set above or from a previous setting.
+		TSet<FMixerSubmixWeakPtr> TouchedSubmixes;
+		TouchedSubmixes.Add(PreviousSubmixResolved);
+	
+
 		// Attenuation Submix Sends. (these come from Attenuation assets).
-		// These are largely identical to SoundSubmix Sends, but don't specify a send stage, so we pass one here.
+		// These are largely identical to SoundSubmix Sends, but don't specify a send stage, so we pass one here.		
 		for (const FAttenuationSubmixSendSettings& SendSettings : WaveInstance->AttenuationSubmixSends)
 		{
-			UpdateSubmixSendLevels(SendSettings, EMixerSourceSubmixSendStage::PostDistanceAttenuation);
+			UpdateSubmixSendLevels(SendSettings, EMixerSourceSubmixSendStage::PostDistanceAttenuation, TouchedSubmixes);
 		}
-		// Clear any previous sends that may not exist now.
-		MixerSourcePrivate::ClearPreviousSubmixSends(PreviousAttenuationSendSettings, WaveInstance->AttenuationSubmixSends, MixerDevice, MixerSourceVoice);
-		PreviousAttenuationSendSettings = WaveInstance->AttenuationSubmixSends; 
 		
 		// Sound submix Sends. (these come from SoundBase derived assets).
 		for (FSoundSubmixSendInfo& SendInfo : WaveInstance->SoundSubmixSends)
 		{
-			UpdateSubmixSendLevels(SendInfo, MixerSourcePrivate::SubmixSendStageToMixerSourceSubmixSendStage(SendInfo.SendStage));
+			UpdateSubmixSendLevels(SendInfo, MixerSourcePrivate::SubmixSendStageToMixerSourceSubmixSendStage(SendInfo.SendStage), TouchedSubmixes);
 		}
-		// Again, Clear any sends that maybe not exist now.
-		MixerSourcePrivate::ClearPreviousSubmixSends(PreviousSubmixSendSettings, WaveInstance->SoundSubmixSends, MixerDevice, MixerSourceVoice);
-		PreviousSubmixSendSettings = WaveInstance->SoundSubmixSends;
+		
+		// Anything we haven't touched this update we should now clear.
+		const TSet<FMixerSubmixWeakPtr> ToClear = PreviousSubmixSends.Difference(TouchedSubmixes);
+		PreviousSubmixSends = TouchedSubmixes;
 
+		// Clear sends that aren't touched.	
+		for (FMixerSubmixWeakPtr i : ToClear)
+		{
+			MixerSourceVoice->ClearSubmixSendInfo(i);
+		}
+	
 		MixerSourceVoice->SetEnablement(WaveInstance->bEnableBusSends, WaveInstance->bEnableBaseSubmix, WaveInstance->bEnableSubmixSends);
 
 		MixerSourceVoice->SetSourceBufferListener(WaveInstance->SourceBufferListener, WaveInstance->bShouldSourceBufferListenerZeroBuffer);
diff --git a/Engine/Source/Runtime/AudioMixer/Private/AudioMixerSource.h b/Engine/Source/Runtime/AudioMixer/Private/AudioMixerSource.h
--- a/Engine/Source/Runtime/AudioMixer/Private/AudioMixerSource.h
+++ b/Engine/Source/Runtime/AudioMixer/Private/AudioMixerSource.h
@@ -124,7 +124,7 @@ namespace Audio
 		float GetInheritedSubmixVolumeModulation() const;
 
 	private:
-		void UpdateSubmixSendLevels(const FSoundSubmixSendInfoBase& InSendInfo, EMixerSourceSubmixSendStage InSendStage);
+		void UpdateSubmixSendLevels(const FSoundSubmixSendInfoBase& InSendInfo, EMixerSourceSubmixSendStage InSendStage, TSet<FMixerSubmixWeakPtr>& OutTouchedSubmixes);
 
 		FMixerDevice* MixerDevice;
 		FMixerBuffer* MixerBuffer;
@@ -180,10 +180,9 @@ namespace Audio
 		// source may need to live-update during its lifespan
 		TArray<FDynamicBusSendInfo> DynamicBusSendInfos;
 
-		// An array of submix sends from previous update. Allows us to clear out submix sends if they are no longer being sent.
-		TArray<FSoundSubmixSendInfo> PreviousSubmixSendSettings;
-		TArray<FAttenuationSubmixSendSettings> PreviousAttenuationSendSettings;
-
+		// An array of submixes from previous update. Allows us to clear out submix sends if they are no longer being sent.
+		TSet<FMixerSubmixWeakPtr> PreviousSubmixSends;
+		
 		// Whether or not we're currently releasing our resources. Prevents recycling the source until release is finished.
 		FThreadSafeBool bIsReleasing;
 
