Audio Bus Reader-Writer Nodes Don't CreatePatch in Reset
#tests all audio bus engine tests pass

diff --git a/Engine/Plugins/Runtime/Metasound/Source/MetasoundEngine/Private/MetasoundAudioBusReaderNode.cpp b/Engine/Plugins/Runtime/Metasound/Source/MetasoundEngine/Private/MetasoundAudioBusReaderNode.cpp
--- a/Engine/Plugins/Runtime/Metasound/Source/MetasoundEngine/Private/MetasoundAudioBusReaderNode.cpp
+++ b/Engine/Plugins/Runtime/Metasound/Source/MetasoundEngine/Private/MetasoundAudioBusReaderNode.cpp
@@ -296,8 +296,6 @@ namespace Metasound
 			
 			SampleRate = InParams.OperatorSettings.GetSampleRate();
 			BlockSizeFrames = InParams.OperatorSettings.GetNumFramesPerBlock();
-
-			CreatePatchOutput();
 		}
 
 	private:
diff --git a/Engine/Plugins/Runtime/Metasound/Source/MetasoundEngine/Private/MetasoundAudioBusWriterNode.cpp b/Engine/Plugins/Runtime/Metasound/Source/MetasoundEngine/Private/MetasoundAudioBusWriterNode.cpp
--- a/Engine/Plugins/Runtime/Metasound/Source/MetasoundEngine/Private/MetasoundAudioBusWriterNode.cpp
+++ b/Engine/Plugins/Runtime/Metasound/Source/MetasoundEngine/Private/MetasoundAudioBusWriterNode.cpp
@@ -204,8 +204,6 @@ namespace Metasound
 				UE_LOG(LogMetaSound, Warning, TEXT("Audio bus writer node requires audio device ID '%s', audio mixer num output frames '%s' and transmitter id '%s' environment variables")
 					, *SourceInterface::Environment::DeviceID.ToString(), *SourceInterface::Environment::AudioMixerNumOutputFrames.ToString(), *SourceInterface::Environment::TransmitterID.ToString());
 			}
-
-			CreatePatchInput();
 		}
 
 
diff --git a/Engine/Plugins/Runtime/Metasound/Source/MetasoundEngine/Private/MetasoundSource.cpp b/Engine/Plugins/Runtime/Metasound/Source/MetasoundEngine/Private/MetasoundSource.cpp
--- a/Engine/Plugins/Runtime/Metasound/Source/MetasoundEngine/Private/MetasoundSource.cpp
+++ b/Engine/Plugins/Runtime/Metasound/Source/MetasoundEngine/Private/MetasoundSource.cpp
@@ -5,8 +5,6 @@
 #include "Algo/Find.h"
 #include "Algo/Transform.h"
 #include "AssetRegistry/AssetRegistryModule.h"
-#include "AudioBusSubsystem.h"
-#include "AudioDevice.h"
 #include "AudioDeviceManager.h"
 #include "Containers/Ticker.h"
 #include "IAudioParameterInterfaceRegistry.h"
@@ -1047,19 +1045,6 @@ ISoundGeneratorPtr UMetaSoundSource::CreateSoundGenerator(const FSoundGeneratorI
 	if (Generator.IsValid())
 	{
 		TrackGenerator(InParams.AudioComponentId, Generator);
-
-		Generator->AddGraphSetCallback(Metasound::FOnSetGraph::FDelegate::CreateLambda([AudioDeviceID = InParams.AudioDeviceID, InstanceID = InParams.InstanceID]()
-			{
-				if (FAudioDeviceManager* ADM = FAudioDeviceManager::Get())
-				{
-					if (FAudioDevice* AudioDevice = ADM->GetAudioDeviceRaw(AudioDeviceID))
-					{
-						UAudioBusSubsystem* AudioBusSubsystem = AudioDevice->GetSubsystem<UAudioBusSubsystem>();
-						check(AudioBusSubsystem);
-						AudioBusSubsystem->ReadyToConnect(InstanceID);
-					}
-				}
-			}));
 	}
 
 	return ISoundGeneratorPtr(Generator);
diff --git a/Engine/Source/Runtime/AudioMixer/Private/AudioBusSubsystem.cpp b/Engine/Source/Runtime/AudioMixer/Private/AudioBusSubsystem.cpp
--- a/Engine/Source/Runtime/AudioMixer/Private/AudioBusSubsystem.cpp
+++ b/Engine/Source/Runtime/AudioMixer/Private/AudioBusSubsystem.cpp
@@ -187,13 +187,6 @@ void UAudioBusSubsystem::AddPendingConnection(uint64 SoundInstanceID, FPendingCo
 	SoundInstanceConnections.PendingConnections.Add(MoveTemp(PendingConnection));
 }
 
-void UAudioBusSubsystem::ReadyToConnect(uint64 SoundInstanceID)
-{
-	FScopeLock ScopeLock(&Mutex);
-	FSoundInstanceConnections& SoundInstanceConnections = SoundInstanceConnectionMap.FindOrAdd(SoundInstanceID);
-	SoundInstanceConnections.bSoundInstanceReady = true;
-}
-
 void UAudioBusSubsystem::ConnectPatches(uint64 SoundInstanceID)
 {
 	TArray<FPendingConnection> PendingConnections = ExtractPendingConnectionsIfReady(SoundInstanceID);
@@ -227,12 +220,9 @@ TArray<UAudioBusSubsystem::FPendingConnection> UAudioBusSubsystem::ExtractPendin
 	FScopeLock ScopeLock(&Mutex);
 	if (FSoundInstanceConnections* SoundInstanceConnections = SoundInstanceConnectionMap.Find(SoundInstanceID))
 	{
-		if (SoundInstanceConnections->bSoundInstanceReady)
-		{
-			TArray<FPendingConnection> PendingConnections = MoveTemp(SoundInstanceConnections->PendingConnections);
-			SoundInstanceConnections->PendingConnections.Empty();
-			return PendingConnections;
-		}
+		TArray<FPendingConnection> PendingConnections = MoveTemp(SoundInstanceConnections->PendingConnections);
+		SoundInstanceConnections->PendingConnections.Empty();
+		return PendingConnections;
 	}
 	return {};
 }
diff --git a/Engine/Source/Runtime/AudioMixer/Public/AudioBusSubsystem.h b/Engine/Source/Runtime/AudioMixer/Public/AudioBusSubsystem.h
--- a/Engine/Source/Runtime/AudioMixer/Public/AudioBusSubsystem.h
+++ b/Engine/Source/Runtime/AudioMixer/Public/AudioBusSubsystem.h
@@ -88,7 +88,6 @@ public:
 
 	AUDIOMIXER_API Audio::FPatchInput AddPatchInputForSoundAndAudioBus(uint64 SoundInstanceID, Audio::FAudioBusKey AudioBusKey, int32 InFrames, int32 NumChannels, float InGain = 1.f);
 	AUDIOMIXER_API Audio::FPatchOutputStrongPtr AddPatchOutputForSoundAndAudioBus(uint64 SoundInstanceID, Audio::FAudioBusKey AudioBusKey, int32 InFrames, int32 NumChannels, float InGain = 1.f);
-	AUDIOMIXER_API void ReadyToConnect(uint64 SoundInstanceID);
 	AUDIOMIXER_API void ConnectPatches(uint64 SoundInstanceID);
 	AUDIOMIXER_API void RemoveSound(uint64 SoundInstanceID);
 
@@ -122,7 +121,6 @@ private:
 	struct FSoundInstanceConnections
 	{
 		TArray<FPendingConnection> PendingConnections;
-		bool bSoundInstanceReady = false;
 	};
 
 	TArray<FPendingConnection> ExtractPendingConnectionsIfReady(uint64 SoundInstanceID);
