Removinvg stale debug-only ensure.
#rnx

diff --git a/Engine/Source/Runtime/AudioMixer/Private/AudioMixerSourceManager.cpp b/Engine/Source/Runtime/AudioMixer/Private/AudioMixerSourceManager.cpp
--- a/Engine/Source/Runtime/AudioMixer/Private/AudioMixerSourceManager.cpp
+++ b/Engine/Source/Runtime/AudioMixer/Private/AudioMixerSourceManager.cpp
@@ -1161,8 +1161,6 @@ namespace Audio
 #if AUDIO_MIXER_ENABLE_DEBUG_MODE
 			AUDIO_MIXER_CHECK(!SourceInfo.bIsDebugMode);
 			SourceInfo.bIsDebugMode = InitParams.bIsDebugMode;
-
-			AUDIO_MIXER_CHECK(SourceInfo.DebugName.IsEmpty());
 			SourceInfo.DebugName = InitParams.DebugName;
 #endif 
 
