Removed / deprecated need to unregister submix buffer listeners.
This had issues with Unregister(AsShared()) being called in TSharedFromThis-derived constructors.

diff --git a/Engine/Plugins/Experimental/MotoSynth/Source/MotoSynth/Private/MotoSynthPreset.cpp b/Engine/Plugins/Experimental/MotoSynth/Source/MotoSynth/Private/MotoSynthPreset.cpp
--- a/Engine/Plugins/Experimental/MotoSynth/Source/MotoSynth/Private/MotoSynthPreset.cpp
+++ b/Engine/Plugins/Experimental/MotoSynth/Source/MotoSynth/Private/MotoSynthPreset.cpp
@@ -35,19 +35,19 @@ void UMotoSynthPreset::PostLoad()
 void UMotoSynthPreset::StartEnginePreview()
 {
 	// Set all the state of the previewer that needs setting
-	EnginePreviewer.SetSettings(Settings);
+	EnginePreviewer->SetSettings(Settings);
 
 	if (FRichCurve* RichRPMCurve = EnginePreviewRPMCurve.GetRichCurve())
 	{
-		EnginePreviewer.SetPreviewRPMCurve(*RichRPMCurve);
+		EnginePreviewer->SetPreviewRPMCurve(*RichRPMCurve);
 	}
 
-	EnginePreviewer.StartPreviewing();
+	EnginePreviewer->StartPreviewing();
 }
 
 void UMotoSynthPreset::StopEnginePreview()
 {
-	EnginePreviewer.StopPreviewing();
+	EnginePreviewer->StopPreviewing();
 }
 
 void UMotoSynthPreset::DumpRuntimeMemoryUsage()
@@ -80,7 +80,7 @@ void UMotoSynthPreset::PostEditChangeProperty(FPropertyChangedEvent& PropertyCha
 		{
 			if (FRichCurve* RichRPMCurve = EnginePreviewRPMCurve.GetRichCurve())
 			{
-				EnginePreviewer.SetPreviewRPMCurve(*RichRPMCurve);
+				EnginePreviewer->SetPreviewRPMCurve(*RichRPMCurve);
 			}
 		}
 		else
@@ -88,12 +88,12 @@ void UMotoSynthPreset::PostEditChangeProperty(FPropertyChangedEvent& PropertyCha
 			// Only set the settings on the engine previewer when we have both a acceleration and deceleration source
 			if (Settings.AccelerationSource && Settings.DecelerationSource)
 			{
-				EnginePreviewer.SetSettings(Settings);
+				EnginePreviewer->SetSettings(Settings);
 			}
 			else
 			{
 				// Stop previewing if we've cleared out any sources (i.e. were previously previewing an accel/decel source, stop it now if one of them is null)
-				EnginePreviewer.StopPreviewing();
+				EnginePreviewer->StopPreviewing();
 			}
 		}
 	}
@@ -106,11 +106,6 @@ FMotoSynthEnginePreviewer::FMotoSynthEnginePreviewer()
 	SynthEngine = TUniquePtr<FMotoSynthEngine>(new FMotoSynthEngine);
 }
 
-FMotoSynthEnginePreviewer::~FMotoSynthEnginePreviewer()
-{
-	StopPreviewing();
-}
-
 void FMotoSynthEnginePreviewer::SetSettings(const FMotoSynthRuntimeSettings& InSettings)
 {
 	FScopeLock Lock(&PreviewEngineCritSect);
diff --git a/Engine/Plugins/Experimental/MotoSynth/Source/MotoSynth/Public/MotoSynthPreset.h b/Engine/Plugins/Experimental/MotoSynth/Source/MotoSynth/Public/MotoSynthPreset.h
--- a/Engine/Plugins/Experimental/MotoSynth/Source/MotoSynth/Public/MotoSynthPreset.h
+++ b/Engine/Plugins/Experimental/MotoSynth/Source/MotoSynth/Public/MotoSynthPreset.h
@@ -175,7 +175,7 @@ class FMotoSynthEnginePreviewer : public ISubmixBufferListener
 {
 public:
 	FMotoSynthEnginePreviewer();
-	virtual ~FMotoSynthEnginePreviewer();
+	virtual ~FMotoSynthEnginePreviewer() override = default;
 
 	void StartPreviewing();
 	void StopPreviewing();
@@ -248,7 +248,7 @@ public:
 
 #if WITH_EDITORONLY_DATA
 	// The engine previewer which is using this preset
-	FMotoSynthEnginePreviewer EnginePreviewer;
+	TSharedRef<FMotoSynthEnginePreviewer> EnginePreviewer = MakeShared<FMotoSynthEnginePreviewer>();
 #endif
 };
 
diff --git a/Engine/Plugins/FX/Niagara/Source/Niagara/Classes/NiagaraDataInterfaceAudio.h b/Engine/Plugins/FX/Niagara/Source/Niagara/Classes/NiagaraDataInterfaceAudio.h
--- a/Engine/Plugins/FX/Niagara/Source/Niagara/Classes/NiagaraDataInterfaceAudio.h
+++ b/Engine/Plugins/FX/Niagara/Source/Niagara/Classes/NiagaraDataInterfaceAudio.h
@@ -35,7 +35,7 @@ public:
 
 	NIAGARA_API void RegisterToSubmix();
 
-	NIAGARA_API virtual ~FNiagaraSubmixListener();
+	NIAGARA_API virtual ~FNiagaraSubmixListener() override = default;
 
 	/** Returns the current sample rate of the current submix. */
 	NIAGARA_API float GetSampleRate() const;
diff --git a/Engine/Plugins/FX/Niagara/Source/Niagara/Private/NiagaraDataInterfaceAudio.cpp b/Engine/Plugins/FX/Niagara/Source/Niagara/Private/NiagaraDataInterfaceAudio.cpp
--- a/Engine/Plugins/FX/Niagara/Source/Niagara/Private/NiagaraDataInterfaceAudio.cpp
+++ b/Engine/Plugins/FX/Niagara/Source/Niagara/Private/NiagaraDataInterfaceAudio.cpp
@@ -52,11 +52,6 @@ FNiagaraSubmixListener::FNiagaraSubmixListener(FNiagaraSubmixListener&& Other)
 	RegisterToSubmix();
 }
 
-FNiagaraSubmixListener::~FNiagaraSubmixListener()
-{
-	UnregisterFromSubmix();
-}
-
 void FNiagaraSubmixListener::RegisterToSubmix()
 {
 	if (FAudioDevice* AudioDevice = FAudioDeviceManager::Get()->GetAudioDeviceRaw(AudioDeviceId))
diff --git a/Engine/Source/Runtime/AudioMixer/AudioMixer.Build.cs b/Engine/Source/Runtime/AudioMixer/AudioMixer.Build.cs
--- a/Engine/Source/Runtime/AudioMixer/AudioMixer.Build.cs
+++ b/Engine/Source/Runtime/AudioMixer/AudioMixer.Build.cs
@@ -6,6 +6,7 @@ namespace UnrealBuildTool.Rules
 	{
 		public AudioMixer(ReadOnlyTargetRules Target) : base(Target)
 		{
+			OptimizeCode = CodeOptimization.Never;
 			PrivateIncludePathModuleNames.Add("TargetPlatform");
 			PublicIncludePathModuleNames.Add("TargetPlatform");
 
diff --git a/Engine/Source/Runtime/AudioMixer/Private/AudioMixerDevice.cpp b/Engine/Source/Runtime/AudioMixer/Private/AudioMixerDevice.cpp
--- a/Engine/Source/Runtime/AudioMixer/Private/AudioMixerDevice.cpp
+++ b/Engine/Source/Runtime/AudioMixer/Private/AudioMixerDevice.cpp
@@ -2842,20 +2842,22 @@ namespace Audio
 		DECLARE_CYCLE_STAT(TEXT("FAudioThreadTask.UnregisterSubmixBufferListener"), STAT_UnregisterSubmixBufferListener, STATGROUP_AudioThreadCommands);
 
 		const TWeakObjectPtr<USoundSubmix> SubmixPtr(&InSubmix);
-
-		auto UnregisterLambda = [this, InSubmixBufferListener, SubmixPtr]()
+		UPTRINT ListenerPtr = reinterpret_cast<UPTRINT>(&InSubmixBufferListener.Get());
+		FString ListenerName = InSubmixBufferListener->GetListenerName();
+		
+		auto UnregisterLambda = [this, SubmixPtr, ListenerPtr, ListenerName]()
 		{
 			CSV_SCOPED_TIMING_STAT(Audio, UnregisterSubmixBufferListener);
 
 			FMixerSubmixPtr FoundSubmix = GetSubmixInstance(SubmixPtr.Get()).Pin();
 			if (FoundSubmix.IsValid())
 			{
-				UE_LOG(LogAudioMixer, Display, TEXT("Unregistering submix buffer listener '%s' from submix '%s'"), *InSubmixBufferListener->GetListenerName(), *FoundSubmix->SubmixName);
-				FoundSubmix->UnregisterBufferListener(InSubmixBufferListener);
+				UE_LOG(LogAudioMixer, Display, TEXT("Unregistering submix buffer listener '%s' from submix '%s'"), *ListenerName, *FoundSubmix->SubmixName);
+				FoundSubmix->UnregisterBufferListenerInternal(ListenerPtr);
 			}
 			else
 			{
-				UE_LOG(LogAudioMixer, Display, TEXT("Submix buffer listener '%s' not unregistered. Submix not loaded."), *InSubmixBufferListener->GetListenerName());
+				UE_LOG(LogAudioMixer, Display, TEXT("Submix buffer listener '%s' not unregistered. Submix not loaded."), *ListenerName);
 			}
 		};
 
diff --git a/Engine/Source/Runtime/AudioMixer/Private/AudioMixerSubmix.cpp b/Engine/Source/Runtime/AudioMixer/Private/AudioMixerSubmix.cpp
--- a/Engine/Source/Runtime/AudioMixer/Private/AudioMixerSubmix.cpp
+++ b/Engine/Source/Runtime/AudioMixer/Private/AudioMixerSubmix.cpp
@@ -1114,15 +1114,18 @@ namespace Audio
 		{
 			return true;
 		}
-
+		
 		{
 			// query the SubmixBufferListeners to see if they plan to render audio into this buffer
 			FScopeLock Lock(&BufferListenerCriticalSection);
-			for (const TSharedRef<ISubmixBufferListener, ESPMode::ThreadSafe>& BufferListener : BufferListeners)
+			for (const TWeakPtr<ISubmixBufferListener>& BufferListenerWeakPtr : BufferListenerPtrs)	
 			{
-				if (BufferListener->IsRenderingAudio())
+				if (TSharedPtr<ISubmixBufferListener> BufferListener = BufferListenerWeakPtr.Pin())
 				{
-					return true;
+					if (BufferListener->IsRenderingAudio())
+					{
+						return true;
+					}
 				}
 			}
 		}
@@ -1625,13 +1628,18 @@ namespace Audio
 
 			double AudioClock = MixerDevice->GetAudioTime();
 			float SampleRate = MixerDevice->GetSampleRate();
+			
 			FScopeLock Lock(&BufferListenerCriticalSection);
-			for (TSharedRef<ISubmixBufferListener, ESPMode::ThreadSafe>& BufferListener : BufferListeners)
+			for (TWeakPtr<ISubmixBufferListener, ESPMode::ThreadSafe>& BufferListenerWeakPtr : BufferListenerPtrs)
 			{
-				BufferListener->OnNewSubmixBuffer(SoundSubmix, OutAudioBuffer.GetData(), OutAudioBuffer.Num(), NumChannels, SampleRate, AudioClock);
+				if (TSharedPtr<ISubmixBufferListener> BufferListener = BufferListenerWeakPtr.Pin())
+				{
+					BufferListener->OnNewSubmixBuffer(SoundSubmix, OutAudioBuffer.GetData(), OutAudioBuffer.Num(), NumChannels, SampleRate, AudioClock);
+				}
 			}
 
 			PatchSplitter.PushAudio(OutAudioBuffer.GetData(), OutAudioBuffer.Num());
+			PruneSubmixBufferListeners();
 		}
 	}
 
@@ -1645,6 +1653,25 @@ namespace Audio
 		}
 	}
 
+	void FMixerSubmix::UnregisterBufferListenerInternal(UPTRINT ListenerBufferPtr) 
+	{
+		BufferListenerPtrs.RemoveAll(
+			[=](const TWeakPtr<ISubmixBufferListener>& Listener)
+			{
+				Listener.GetWeakPtrTypeHash();
+				return reinterpret_cast<UPTRINT>(Listener.Pin().Get()) == ListenerBufferPtr;
+			});
+	}
+
+	void FMixerSubmix::PruneSubmixBufferListeners()
+	{
+		BufferListenerPtrs.RemoveAll(
+			[](const TWeakPtr<ISubmixBufferListener>& Listener)
+			{
+				return !Listener.IsValid();
+			});
+	}
+
 	bool FMixerSubmix::GenerateEffectChainAudio(FSoundEffectSubmixInputData& InputData, const FAlignedFloatBuffer& InAudioBuffer, TArray<FSoundEffectSubmixPtr>& InEffectChain, FAlignedFloatBuffer& OutBuffer)
 	{
 		checkf(InAudioBuffer.Num() == OutBuffer.Num(), TEXT("Processing effect chain audio must not alter the number of samples in a buffer. Buffer size mismatch InAudioBuffer[%d] != OutBuffer[%d]"), InAudioBuffer.Num(), OutBuffer.Num());
@@ -2152,26 +2179,26 @@ namespace Audio
 	{
 		FScopeLock Lock(&BufferListenerCriticalSection);
 		check(BufferListener);
-		BufferListeners.AddUnique(BufferListener->AsShared());
+		BufferListenerPtrs.AddUnique(BufferListener->AsShared());
 	}
 
 	void FMixerSubmix::RegisterBufferListener(TSharedRef<ISubmixBufferListener, ESPMode::ThreadSafe> BufferListener)
 	{
 		FScopeLock Lock(&BufferListenerCriticalSection);
-		BufferListeners.AddUnique(BufferListener);
+		BufferListenerPtrs.AddUnique(BufferListener);
 	}
 
 	void FMixerSubmix::UnregisterBufferListener(ISubmixBufferListener* BufferListener)
 	{
 		FScopeLock Lock(&BufferListenerCriticalSection);
 		check(BufferListener);
-		BufferListeners.Remove(BufferListener->AsShared());
+		BufferListenerPtrs.Remove(BufferListener->AsShared());
 	}
 
 	void FMixerSubmix::UnregisterBufferListener(TSharedRef<ISubmixBufferListener, ESPMode::ThreadSafe> BufferListener)
 	{
 		FScopeLock Lock(&BufferListenerCriticalSection);
-		BufferListeners.Remove(BufferListener);
+		BufferListenerPtrs.Remove(BufferListener);
 	}
 
 	void FMixerSubmix::StartEnvelopeFollowing(int32 AttackTime, int32 ReleaseTime)
diff --git a/Engine/Source/Runtime/AudioMixer/Public/AudioMixerDevice.h b/Engine/Source/Runtime/AudioMixer/Public/AudioMixerDevice.h
--- a/Engine/Source/Runtime/AudioMixer/Public/AudioMixerDevice.h
+++ b/Engine/Source/Runtime/AudioMixer/Public/AudioMixerDevice.h
@@ -202,6 +202,9 @@ namespace Audio
 		AUDIOMIXER_API virtual void UnregisterSubmixBufferListener(ISubmixBufferListener* InSubmixBufferListener, USoundSubmix* InSubmix = nullptr) override;
 
 		AUDIOMIXER_API virtual void RegisterSubmixBufferListener(TSharedRef<ISubmixBufferListener, ESPMode::ThreadSafe> InSubmixBufferListener, USoundSubmix& InSubmix) override;
+
+		// This is optional, and should only be used to manually unregister before the ISubmixBufferListener is getting destroyed
+		// (do not call this in the destructor of ISubmixBufferListener or derived classes)
 		AUDIOMIXER_API virtual void UnregisterSubmixBufferListener(TSharedRef<ISubmixBufferListener, ESPMode::ThreadSafe> InSubmixBufferListener, USoundSubmix& InSubmix) override;
 
 		AUDIOMIXER_API virtual FPatchOutputStrongPtr AddPatchForSubmix(uint32 InObjectId, float InPatchGain) override;
diff --git a/Engine/Source/Runtime/AudioMixer/Public/AudioMixerSubmix.h b/Engine/Source/Runtime/AudioMixer/Public/AudioMixerSubmix.h
--- a/Engine/Source/Runtime/AudioMixer/Public/AudioMixerSubmix.h
+++ b/Engine/Source/Runtime/AudioMixer/Public/AudioMixerSubmix.h
@@ -588,9 +588,9 @@ namespace Audio
 		// Submix command queue to shuffle commands from audio thread to audio render thread.
 		TQueue<TFunction<void()>> CommandQueue;
 
-		// List of submix buffer listeners.
-		TArray<TSharedRef<ISubmixBufferListener, ESPMode::ThreadSafe>> BufferListeners;
+		// List of submix buffer listeners. (mutable for pruning stale weak references)
 
+		mutable TArray<TWeakPtr<ISubmixBufferListener>> BufferListenerPtrs;
 		// Critical section used for modifying and interacting with buffer listeners
 		mutable FCriticalSection BufferListenerCriticalSection;
 
@@ -684,6 +684,10 @@ namespace Audio
 	private:
 		AUDIOMIXER_API void SendAudioToRegisteredAudioBuses(FAlignedFloatBuffer& OutAudioBuffer);
 
+		void UnregisterBufferListenerInternal(UPTRINT ListenerBufferPtr);
+
+		void PruneSubmixBufferListeners();
+
 		// Registered audio buses
 		TMap<Audio::FAudioBusKey, Audio::FPatchInput> AudioBuses;
 	};
diff --git a/Engine/Source/Runtime/GameplayMediaEncoder/Private/GameplayMediaEncoder.cpp b/Engine/Source/Runtime/GameplayMediaEncoder/Private/GameplayMediaEncoder.cpp
--- a/Engine/Source/Runtime/GameplayMediaEncoder/Private/GameplayMediaEncoder.cpp
+++ b/Engine/Source/Runtime/GameplayMediaEncoder/Private/GameplayMediaEncoder.cpp
@@ -89,7 +89,11 @@ FGameplayMediaEncoder::FGameplayMediaEncoder() {}
 PRAGMA_ENABLE_DEPRECATION_WARNINGS
 
 PRAGMA_DISABLE_DEPRECATION_WARNINGS
-FGameplayMediaEncoder::~FGameplayMediaEncoder() { Shutdown(); }
+FGameplayMediaEncoder::~FGameplayMediaEncoder()
+{
+	bIsInDestructor = true;
+	Shutdown();
+}
 PRAGMA_ENABLE_DEPRECATION_WARNINGS
 
 bool FGameplayMediaEncoder::RegisterListener(IGameplayMediaEncoderListener* Listener)
@@ -374,7 +378,11 @@ void FGameplayMediaEncoder::Stop()
 	{
 		if(FAudioDevice* AudioDevice = GameEngine->GetMainAudioDeviceRaw())
 		{
-			AudioDevice->UnregisterSubmixBufferListener(AsShared(), AudioDevice->GetMainSubmixObject());
+			// AsShared() cannot be called in the destructor of the ref counted object!
+			if(!bIsInDestructor)
+			{
+				AudioDevice->UnregisterSubmixBufferListener(AsShared(), AudioDevice->GetMainSubmixObject());
+			}
 		}
 
 		if(FSlateApplication::IsInitialized())
diff --git a/Engine/Source/Runtime/GameplayMediaEncoder/Public/GameplayMediaEncoder.h b/Engine/Source/Runtime/GameplayMediaEncoder/Public/GameplayMediaEncoder.h
--- a/Engine/Source/Runtime/GameplayMediaEncoder/Public/GameplayMediaEncoder.h
+++ b/Engine/Source/Runtime/GameplayMediaEncoder/Public/GameplayMediaEncoder.h
@@ -169,6 +169,8 @@ private:
 	PRAGMA_DISABLE_DEPRECATION_WARNINGS
 	TMap<TSharedPtr<AVEncoder::FVideoEncoderInputFrame>, FTextureRHIRef> BackBuffers;
 	PRAGMA_ENABLE_DEPRECATION_WARNINGS
+
+	bool bIsInDestructor = false;
 };
 
 
