[CIS] followup: remove module unoptimiaztion (This had issues with Unregister(AsShared()) being called in TSharedFromThis-derived constructors.)

diff --git a/Engine/Source/Runtime/AudioMixer/AudioMixer.Build.cs b/Engine/Source/Runtime/AudioMixer/AudioMixer.Build.cs
--- a/Engine/Source/Runtime/AudioMixer/AudioMixer.Build.cs
+++ b/Engine/Source/Runtime/AudioMixer/AudioMixer.Build.cs
@@ -6,7 +6,6 @@ namespace UnrealBuildTool.Rules
 	{
 		public AudioMixer(ReadOnlyTargetRules Target) : base(Target)
 		{
-			OptimizeCode = CodeOptimization.Never;
 			PrivateIncludePathModuleNames.Add("TargetPlatform");
 			PublicIncludePathModuleNames.Add("TargetPlatform");
 
