Pass on missing ScopeLocks around SubmixBufferListener array mutation.
#rnx

diff --git a/Engine/Source/Runtime/AudioMixer/Private/AudioMixerSubmix.cpp b/Engine/Source/Runtime/AudioMixer/Private/AudioMixerSubmix.cpp
--- a/Engine/Source/Runtime/AudioMixer/Private/AudioMixerSubmix.cpp
+++ b/Engine/Source/Runtime/AudioMixer/Private/AudioMixerSubmix.cpp
@@ -1655,6 +1655,7 @@ namespace Audio
 
 	void FMixerSubmix::UnregisterBufferListenerInternal(UPTRINT ListenerBufferPtr) 
 	{
+		FScopeLock Lock(&BufferListenerCriticalSection);
 		BufferListenerPtrs.RemoveAll(
 			[=](const TWeakPtr<ISubmixBufferListener>& Listener)
 			{
@@ -1665,6 +1666,7 @@ namespace Audio
 
 	void FMixerSubmix::PruneSubmixBufferListeners()
 	{
+		FScopeLock Lock(&BufferListenerCriticalSection);
 		BufferListenerPtrs.RemoveAll(
 			[](const TWeakPtr<ISubmixBufferListener>& Listener)
 			{
