Modulation Settings in Sound Cues

diff --git a/Engine/Source/Runtime/AudioMixer/Private/AudioMixerSource.cpp b/Engine/Source/Runtime/AudioMixer/Private/AudioMixerSource.cpp
--- a/Engine/Source/Runtime/AudioMixer/Private/AudioMixerSource.cpp
+++ b/Engine/Source/Runtime/AudioMixer/Private/AudioMixerSource.cpp
@@ -13,6 +13,7 @@
 #include "IAudioModulation.h"
 #include "ProfilingDebugging/CsvProfiler.h"
 #include "Sound/AudioSettings.h"
+#include "Sound/SoundCue.h"
 #include "Sound/SoundModulationDestination.h"
 #include "Misc/ScopeRWLock.h"
 #include "Templates/Function.h"
@@ -482,6 +483,98 @@ namespace Audio
 
 			return Settings;
 		}
+
+
+		FSoundModulationDestinationSettings InitVolumeModulationFromSoundCue(const USoundCue& InSoundCue, const FActiveSound& InActiveSound)
+		{
+			const EModulationRouting& ActiveSoundRouting = InActiveSound.ModulationRouting.VolumeRouting;
+			const FSoundModulationDestinationSettings& ActiveSoundSettings = InActiveSound.ModulationRouting.VolumeModulationDestination;
+
+			const EModulationRouting& SoundCueRouting = InSoundCue.ModulationSettings.VolumeRouting;
+			const FSoundModulationDestinationSettings& SoundCueSettings = InSoundCue.ModulationSettings.VolumeModulationDestination;
+
+			return InitRoutedDestinationSettings(
+				ActiveSoundRouting,
+				ActiveSoundSettings,
+				SoundCueRouting,
+				SoundCueSettings,
+				InSoundCue.GetSoundClass(),
+				Audio::GetModulationParameter("Volume"),
+				[](const USoundClass& InSoundClass) { return &InSoundClass.Properties.ModulationSettings.VolumeModulationDestination; }
+			);
+		}
+
+		FSoundModulationDestinationSettings InitPitchModulationFromSoundCue(const USoundCue& InSoundCue, const FActiveSound& InActiveSound)
+		{
+			const EModulationRouting& ActiveSoundRouting = InActiveSound.ModulationRouting.PitchRouting;
+			const FSoundModulationDestinationSettings& ActiveSoundSettings = InActiveSound.ModulationRouting.PitchModulationDestination;
+
+			const EModulationRouting& SoundCueRouting = InSoundCue.ModulationSettings.PitchRouting;
+			const FSoundModulationDestinationSettings& SoundCueSettings = InSoundCue.ModulationSettings.PitchModulationDestination;
+
+			return InitRoutedDestinationSettings(
+				ActiveSoundRouting,
+				ActiveSoundSettings,
+				SoundCueRouting,
+				SoundCueSettings,
+				InSoundCue.GetSoundClass(),
+				Audio::GetModulationParameter("Pitch"),
+				[](const USoundClass& InSoundClass) { return &InSoundClass.Properties.ModulationSettings.PitchModulationDestination; }
+			);
+		}
+
+		FSoundModulationDestinationSettings InitHighpassModulationFromSoundCue(const USoundCue& InSoundCue, const FActiveSound& InActiveSound)
+		{
+			const EModulationRouting& ActiveSoundRouting = InActiveSound.ModulationRouting.HighpassRouting;
+			const FSoundModulationDestinationSettings& ActiveSoundSettings = InActiveSound.ModulationRouting.HighpassModulationDestination;
+
+			const EModulationRouting& SoundCueRouting = InSoundCue.ModulationSettings.HighpassRouting;
+			const FSoundModulationDestinationSettings& SoundCueSettings = InSoundCue.ModulationSettings.HighpassModulationDestination;
+
+			return InitRoutedDestinationSettings(
+				ActiveSoundRouting,
+				ActiveSoundSettings,
+				SoundCueRouting,
+				SoundCueSettings,
+				InSoundCue.GetSoundClass(),
+				Audio::GetModulationParameter("HPFCutoffFrequency"),
+				[](const USoundClass& InSoundClass) { return &InSoundClass.Properties.ModulationSettings.HighpassModulationDestination; }
+			);
+		}
+
+		FSoundModulationDestinationSettings InitLowpassModulationFromSoundCue(const USoundCue& InSoundCue, const FActiveSound& InActiveSound)
+		{
+			const EModulationRouting& ActiveSoundRouting = InActiveSound.ModulationRouting.LowpassRouting;
+			const FSoundModulationDestinationSettings& ActiveSoundSettings = InActiveSound.ModulationRouting.LowpassModulationDestination;
+
+			const EModulationRouting& SoundCueRouting = InSoundCue.ModulationSettings.LowpassRouting;
+			const FSoundModulationDestinationSettings& SoundCueSettings = InSoundCue.ModulationSettings.LowpassModulationDestination;
+
+			return InitRoutedDestinationSettings(
+				ActiveSoundRouting,
+				ActiveSoundSettings,
+				SoundCueRouting,
+				SoundCueSettings,
+				InSoundCue.GetSoundClass(),
+				Audio::GetModulationParameter("LPFCutoffFrequency"),
+				[](const USoundClass& InSoundClass) { return &InSoundClass.Properties.ModulationSettings.LowpassModulationDestination; }
+			);
+		}
+
+		FSoundModulationDefaultSettings InitRoutedModulationFromSoundCue(const USoundCue& InSoundCue, FActiveSound* InActiveSound)
+		{
+			FSoundModulationDefaultSettings Settings;
+			if (InActiveSound)
+			{
+				Settings.VolumeModulationDestination = InitVolumeModulationFromSoundCue(InSoundCue, *InActiveSound);
+				Settings.PitchModulationDestination = InitPitchModulationFromSoundCue(InSoundCue, *InActiveSound);
+				Settings.HighpassModulationDestination = InitHighpassModulationFromSoundCue(InSoundCue, *InActiveSound);
+				Settings.LowpassModulationDestination = InitLowpassModulationFromSoundCue(InSoundCue, *InActiveSound);
+			}
+
+			return Settings;
+		}
+
 	} // namespace ModulationUtils
 
 	FMixerSource::FMixerSource(FAudioDevice* InAudioDevice)
@@ -604,7 +697,23 @@ namespace Audio
 			InitParams.bIsSoundfield = WaveInstance->bIsAmbisonics && (WaveData->NumChannels == 4);
 
 			FActiveSound* ActiveSound = WaveInstance->ActiveSound;
-			InitParams.ModulationSettings = ModulationUtils::InitRoutedModulation(*WaveInstance, *WaveData, ActiveSound);
+
+			/*
+			 * Note--this code means that when Sound Cues enable their own Modulation, their settings will completely overwrite 
+			 * the settings of any SoundWave being played by the Sound Cue.
+			 * ie. this MixerSource will use the Modulation settings of the Audio Component, Sound Cue, and/or Sound Cue's Sound Class,
+			 * depending on inheritance settings. This is intended behavior;
+			 * it's established UX that SoundCue settings usually stomp/overwrite the settings of their SoundWaves.
+			 */
+			USoundCue* SoundCue = Cast<USoundCue>(ActiveSound->GetSound());
+			if (SoundCue && SoundCue->bUseModulation)
+			{
+				InitParams.ModulationSettings = ModulationUtils::InitRoutedModulationFromSoundCue(*SoundCue, ActiveSound);
+			}
+			else
+			{
+				InitParams.ModulationSettings = ModulationUtils::InitRoutedModulation(*WaveInstance, *WaveData, ActiveSound);
+			}
 
 			// Copy quantization request data
 			if (WaveInstance->QuantizedRequestData)
diff --git a/Engine/Source/Runtime/Engine/Classes/Sound/SoundCue.h b/Engine/Source/Runtime/Engine/Classes/Sound/SoundCue.h
--- a/Engine/Source/Runtime/Engine/Classes/Sound/SoundCue.h
+++ b/Engine/Source/Runtime/Engine/Classes/Sound/SoundCue.h
@@ -7,6 +7,7 @@
 #include "Templates/SubclassOf.h"
 #include "Sound/SoundAttenuation.h"
 #include "Sound/SoundBase.h"
+#include "Sound/SoundModulationDestination.h"
 #include "Sound/SoundNode.h"
 #include "Sound/SoundWave.h"
 
@@ -102,6 +103,16 @@ class USoundCue : public USoundBase
 	UPROPERTY(EditAnywhere, BlueprintReadOnly, Category = Sound, AssetRegistrySearchable)
 	float PitchMultiplier;
 
+	/* 
+	 * Enables Modulation on this SoundCue. Setting this to true will override all modulation settings on any SoundWaves played through this SoundCue.
+	 * In the inheritance path, Audio Components will inherit Modulation Settings from this sound cue, which inherits Modulation Settings from its own Sound Class.
+	 */
+	UPROPERTY(EditAnywhere, Category = Modulation)
+	bool bUseModulation;
+
+	UPROPERTY(EditAnywhere, Category = Modulation, meta = (EditCondition = "bUseModulation"))
+	FSoundModulationDefaultRoutingSettings ModulationSettings;
+
 	/* Attenuation settings to use if Override Attenuation is set to true */
 	UPROPERTY(EditAnywhere, Category = Attenuation)
 	FSoundAttenuationSettings AttenuationOverrides;
