Update Audio Component ID in Spatialization plugin data
#11617

diff --git a/Engine/Source/Runtime/AudioMixer/Private/AudioMixerSourceManager.cpp b/Engine/Source/Runtime/AudioMixer/Private/AudioMixerSourceManager.cpp
--- a/Engine/Source/Runtime/AudioMixer/Private/AudioMixerSourceManager.cpp
+++ b/Engine/Source/Runtime/AudioMixer/Private/AudioMixerSourceManager.cpp
@@ -2594,6 +2594,7 @@ namespace Audio
 			AudioPluginInputData.NumChannels = SourceInfo.NumInputChannels;
 			AudioPluginInputData.SourceId = SourceId;
 			AudioPluginInputData.SpatializationParams = &SourceInfo.SpatParams;
+			AudioPluginInputData.AudioComponentId = SourceInfo.AudioComponentID;
 
 			if (!SpatialInterfaceInfo.bSpatializationIsExternalSend)
 			{
