Fix Bug where Mixer Source Voice's Modulation settings were being overwritten during MixerSource Update

diff --git a/Engine/Source/Runtime/AudioMixer/Private/AudioMixerSource.cpp b/Engine/Source/Runtime/AudioMixer/Private/AudioMixerSource.cpp
--- a/Engine/Source/Runtime/AudioMixer/Private/AudioMixerSource.cpp
+++ b/Engine/Source/Runtime/AudioMixer/Private/AudioMixerSource.cpp
@@ -482,6 +482,22 @@ namespace Audio
 
 			return Settings;
 		}
+
+		FSoundModulationDefaultRoutingSettings UpdateRoutedModulation(const FWaveInstance& InWaveInstance, const USoundWave& InWaveData, FActiveSound* InActiveSound)
+		{
+			FSoundModulationDefaultRoutingSettings NewRouting;
+
+			if (InActiveSound)
+			{
+				NewRouting.VolumeModulationDestination = InitRoutedVolumeModulation(InWaveInstance, InWaveData, *InActiveSound);
+				NewRouting.PitchModulationDestination = InitRoutedPitchModulation(InWaveInstance, InWaveData, *InActiveSound);
+				NewRouting.HighpassModulationDestination = InitRoutedHighpassModulation(InWaveInstance, InWaveData, *InActiveSound);
+				NewRouting.LowpassModulationDestination = InitRoutedLowpassModulation(InWaveInstance, InWaveData, *InActiveSound);
+			}
+
+			return NewRouting;
+		}
+
 	} // namespace ModulationUtils
 
 	FMixerSource::FMixerSource(FAudioDevice* InAudioDevice)
@@ -1929,7 +1945,15 @@ namespace Audio
 
 		if (ActiveSound->bModulationRoutingUpdated)
 		{
-			MixerSourceVoice->SetModulationRouting(ActiveSound->ModulationRouting);
+			if (WaveInstance->WaveData)
+			{
+				FSoundModulationDefaultRoutingSettings UpdatedRouting = ModulationUtils::UpdateRoutedModulation(*WaveInstance, *(WaveInstance->WaveData), ActiveSound);
+				MixerSourceVoice->SetModulationRouting(UpdatedRouting);
+			}
+			else
+			{
+				MixerSourceVoice->SetModulationRouting(ActiveSound->ModulationRouting);
+			}
 		}
 
 		ActiveSound->bModulationRoutingUpdated = false;
diff --git a/Engine/Source/Runtime/Engine/Private/ActiveSound.cpp b/Engine/Source/Runtime/Engine/Private/ActiveSound.cpp
--- a/Engine/Source/Runtime/Engine/Private/ActiveSound.cpp
+++ b/Engine/Source/Runtime/Engine/Private/ActiveSound.cpp
@@ -517,7 +517,7 @@ void FActiveSound::SetNewModulationRouting(const FSoundModulationDefaultRoutingS
 	bModulationRoutingUpdated = true;
 }
 
-void FActiveSound::AddModulationRouting(const TSet<TObjectPtr<USoundModulatorBase>>& NewModulators, EModulationDestination Destination, const bool bShouldModulationRoutingBeUpdated)
+void FActiveSound::AddModulationRouting(const TSet<TObjectPtr<USoundModulatorBase>>& NewModulators, EModulationDestination Destination)
 {
 	switch (Destination)
 	{
@@ -545,7 +545,7 @@ void FActiveSound::AddModulationRouting(const TSet<TObjectPtr<USoundModulatorBas
 		}
 	}
 	
-	bModulationRoutingUpdated = bShouldModulationRoutingBeUpdated;
+	bModulationRoutingUpdated = true;
 }
 
 void FActiveSound::RemoveModulationRouting(const TSet<TObjectPtr<USoundModulatorBase>>& NewModulators, EModulationDestination Destination)
diff --git a/Engine/Source/Runtime/Engine/Private/Components/AudioComponent.cpp b/Engine/Source/Runtime/Engine/Private/Components/AudioComponent.cpp
--- a/Engine/Source/Runtime/Engine/Private/Components/AudioComponent.cpp
+++ b/Engine/Source/Runtime/Engine/Private/Components/AudioComponent.cpp
@@ -1954,9 +1954,9 @@ void UAudioComponent::AddModulationRouting(const TSet<USoundModulatorBase*>& Mod
 	}
 
 	// Tell the active sounds on the component to use the new Modulation Routing
-	AudioDevice->SendCommandToActiveSounds(AudioComponentID, [NewModulators = ConvertedModulators, Destination, bShouldModulationRoutingBeUpdated](FActiveSound& ActiveSound)
+	AudioDevice->SendCommandToActiveSounds(AudioComponentID, [NewModulators = ConvertedModulators, Destination](FActiveSound& ActiveSound)
 	{
-		ActiveSound.AddModulationRouting(NewModulators, Destination, bShouldModulationRoutingBeUpdated);
+		ActiveSound.AddModulationRouting(NewModulators, Destination);
 	});
 }
 
diff --git a/Engine/Source/Runtime/Engine/Public/ActiveSound.h b/Engine/Source/Runtime/Engine/Public/ActiveSound.h
--- a/Engine/Source/Runtime/Engine/Public/ActiveSound.h
+++ b/Engine/Source/Runtime/Engine/Public/ActiveSound.h
@@ -734,7 +734,7 @@ public:
 	 * This function forces the Destination's routing method to Union.
 	 * To replace existing modulators or set a different routing method, use SetNewModulationRouting.
 	 */
-	ENGINE_API void AddModulationRouting(const TSet<TObjectPtr<USoundModulatorBase>>& NewModulators, EModulationDestination Destination, const bool bShouldModulationRoutingBeUpdated = true);
+	ENGINE_API void AddModulationRouting(const TSet<TObjectPtr<USoundModulatorBase>>& NewModulators, EModulationDestination Destination);
 
 	/* Removes given Modulators from the ActiveSound, if possible. */
 	ENGINE_API void RemoveModulationRouting(const TSet<TObjectPtr<USoundModulatorBase>>& NewModulators, EModulationDestination Destination);
