Fix for crash in dynamics processor after changing output channel counts
[RN] Fix for crash in dynamics processor after changing output channel counts
#tests AudioQA Windows editor changing back and forth with devices of different formats (8 to 2))

diff --git a/Engine/Source/Runtime/AudioMixer/Private/Effects/AudioMixerSubmixEffectDynamicsProcessor.cpp b/Engine/Source/Runtime/AudioMixer/Private/Effects/AudioMixerSubmixEffectDynamicsProcessor.cpp
--- a/Engine/Source/Runtime/AudioMixer/Private/Effects/AudioMixerSubmixEffectDynamicsProcessor.cpp
+++ b/Engine/Source/Runtime/AudioMixer/Private/Effects/AudioMixerSubmixEffectDynamicsProcessor.cpp
@@ -169,6 +169,8 @@ Audio::FMixerDevice* FSubmixEffectDynamicsProcessor::GetMixerDevice()
 
 bool FSubmixEffectDynamicsProcessor::UpdateKeySourcePatch()
 {
+	using namespace Audio;
+	
 	// Default (input as key) does not use source patch, so don't
 	// continue checking or updating state.
 	if (KeySource.GetType() == ESubmixEffectDynamicsKeySource::Default)
@@ -176,61 +178,86 @@ bool FSubmixEffectDynamicsProcessor::UpdateKeySourcePatch()
 		return false;
 	}
 
-	if (KeySource.Patch.IsValid())
+	// Valid object
+	const uint32 ObjectId = KeySource.GetObjectId();
+	if (ObjectId == INDEX_NONE)
+	{
+		return false;
+	}
+	
+	// Retrieving/mutating the MixerDevice is only safe during OnProcessAudio calls if
+	// it is not called during Teardown.  The DynamicsProcessor should be Reset via
+	// the OnDeviceDestroyed callback (prior to FAudioDevice::Teardown), so this call
+	// should never be hit during Teardown.
+	FMixerDevice* MixerDevice = GetMixerDevice();
+	if (!MixerDevice)
 	{
-		return true;
+		return false;
 	}
 
+	// Need to determine if our current patch is still valid.
+	// This is handled differently per-type. We return true for valid.
+	// If it's not-valid, we must try and (re)create a patch
+	
 	switch (KeySource.GetType())
 	{
 		case ESubmixEffectDynamicsKeySource::AudioBus:
 		{
-			// Retrieving/mutating the MixerDevice is only safe during OnProcessAudio calls if
-			// it is not called during Teardown.  The DynamicsProcessor should be Reset via
-			// the OnDeviceDestroyed callback (prior to FAudioDevice::Teardown), so this call
-			// should never be hit during Teardown.
-			if (Audio::FMixerDevice* MixerDevice = GetMixerDevice())
+			if (UAudioBusSubsystem* AudioBusSubsystem = MixerDevice->GetSubsystem<UAudioBusSubsystem>())
 			{
-				const uint32 ObjectId = KeySource.GetObjectId();
-				if (ObjectId != INDEX_NONE)
+				// Check patch is still valid and not stale.
+				// AudioBusses can't change format under us (yet), so just check its validity.
+				if (KeySource.Patch.IsValid() && !KeySource.Patch->IsInputStale())
 				{
-					UAudioBusSubsystem* AudioBusSubsystem = MixerDevice->GetSubsystem<UAudioBusSubsystem>();
-					if (AudioBusSubsystem)
-					{
-						const int32 NumChannels = KeySource.GetNumChannels();
-						AudioBusSubsystem->StartAudioBus(Audio::FAudioBusKey(ObjectId), NumChannels, /*bInIsAutomatic=*/false);
-						KeySource.Patch = AudioBusSubsystem->AddPatchOutputForAudioBus(Audio::FAudioBusKey(ObjectId), MixerDevice->GetNumOutputFrames(), NumChannels);
-						DynamicsProcessor.SetKeyNumChannels(NumChannels);
-					}
+					return true;
 				}
+
+				const int32 NumFrames = MixerDevice->GetNumOutputFrames();
+				const int32 NumChannels = KeySource.GetNumChannels();
+				const FAudioBusKey BusKey(ObjectId);
+
+				// Make sure the bus is started. (if it's started already this does nothing).
+				AudioBusSubsystem->StartAudioBus(BusKey, NumChannels, /*bInIsAutomatic=*/false);
+
+				// Add/Recreate a new patch to this bus
+				KeySource.Patch = AudioBusSubsystem->AddPatchOutputForAudioBus(BusKey, NumFrames, NumChannels);
+
+				// Always set the keyed channels even if we fail to create the patch.
+				// As the nop case requires these to be set.
+				DynamicsProcessor.SetKeyNumChannels(NumChannels);
+
+				// If the patch is created and got valid inputs (non-stale) we consider it valid.
+				return KeySource.Patch.IsValid() && !KeySource.Patch->IsInputStale();				
 			}
 		}
 		break;
 
 		case ESubmixEffectDynamicsKeySource::Submix:
 		{
-			// Retrieving/mutating the MixerDevice is only safe during OnProcessAudio calls if
-			// it is not called during Teardown.  The DynamicsProcessor should be Reset via
-			// the OnDeviceDestroyed callback (prior to FAudioDevice::Teardown), so this call
-			// should never be hit during Teardown.
-			if (Audio::FMixerDevice* MixerDevice = GetMixerDevice())
+			if (const FMixerSubmixPtr SubmixPtr = MixerDevice->FindSubmixInstanceByObjectId(ObjectId); SubmixPtr.IsValid())
 			{
-				const uint32 ObjectId = KeySource.GetObjectId();
-				if (ObjectId != INDEX_NONE)
+				const int32 NumSubmixChannels = SubmixPtr->GetNumOutputChannels();					
+
+				// Is patch still valid? Check that we haven't changed channel count.
+				// This can happen after we've swapped devices.
+				if (KeySource.Patch.IsValid() && !KeySource.Patch->IsInputStale() && KeySource.GetNumChannels() == NumSubmixChannels)
 				{
-					KeySource.Patch = MixerDevice->AddPatchForSubmix(ObjectId, 1.0f /* PatchGain */);
-					if (KeySource.Patch.IsValid())
-					{
-						Audio::FMixerSubmixPtr SubmixPtr = MixerDevice->FindSubmixInstanceByObjectId(KeySource.GetObjectId());
-						if (SubmixPtr.IsValid())
-						{
-							const int32 SubmixNumChannels = SubmixPtr->GetNumOutputChannels();
-							KeySource.SetNumChannels(SubmixNumChannels);
-							DynamicsProcessor.SetKeyNumChannels(SubmixNumChannels);
-							return true;
-						}
-					}
-				}	
+					// Looks good.
+					return true;
+				}
+
+				// We store the num channels the submix key source is running at.
+				KeySource.SetNumChannels(NumSubmixChannels);
+				
+				// Add/Recreate a new patch to this submix.
+				KeySource.Patch = MixerDevice->AddPatchForSubmix(ObjectId, 1.0f /* PatchGain */);
+
+				// Always set the keyed channels even if we fail to create the patch.
+				// As the nop case requires these to be set.
+				DynamicsProcessor.SetKeyNumChannels(NumSubmixChannels);
+
+				// If the patch is created and got valid inputs (non-stale) we consider it valid.
+				return KeySource.Patch.IsValid() && !KeySource.Patch->IsInputStale();				
 			}
 		}
 		break;
@@ -238,10 +265,12 @@ bool FSubmixEffectDynamicsProcessor::UpdateKeySourcePatch()
 		case ESubmixEffectDynamicsKeySource::Default:
 		default:
 		{
+			checkNoEntry();
 		}
 		break;
 	}
-
+	
+	// KeySource invalid.
 	return false;
 }
 
@@ -262,36 +291,44 @@ void FSubmixEffectDynamicsProcessor::OnProcessAudio(const FSoundEffectSubmixInpu
 		return;
 	}
 
-	int32 NumKeyChannels = DynamicsProcessor.GetKeyNumChannels();
-	int32 NumKeySamples = InData.NumFrames * NumKeyChannels;
-
 	AudioExternal.Reset();
-	
+
+	// Update the output channels if necessary.
+	if (InData.NumChannels != DynamicsProcessor.GetNumChannels())
+	{
+		DynamicsProcessor.SetNumChannels(InData.NumChannels);
+	}
+		
 	// If set to default, enforce num key channels to always be number of input channels.
 	// If either unset or KeySource was changed between frames back to 'Default', NumKeyChannels
 	// could be stale or left as initialized number of scratch channels.
 	if (KeySource.GetType() == ESubmixEffectDynamicsKeySource::Default)
 	{
-		if (InData.NumChannels != NumKeyChannels)
+		if (InData.NumChannels != DynamicsProcessor.GetKeyNumChannels())
 		{
-			NumKeyChannels = InData.NumChannels;
-			NumKeySamples = InData.NumFrames * NumKeyChannels;
-			DynamicsProcessor.SetKeyNumChannels(NumKeyChannels);
+			DynamicsProcessor.SetKeyNumChannels(InData.NumChannels);
 		}
-	}
-	else
+	}	
+	else 	
 	{
-		AudioExternal.AddZeroed(NumKeySamples);
-	}
-
-	if (UpdateKeySourcePatch())
-	{
-		KeySource.Patch->PopAudio(AudioExternal.GetData(), NumKeySamples, true /* bUseLatestAudio */);
-	}
-
-	if (InData.NumChannels != DynamicsProcessor.GetNumChannels())
-	{
-		DynamicsProcessor.SetNumChannels(InData.NumChannels);
+		// Update our key source. (this validates the existing patch and create new one if necessary)
+		// NOTE: This call can adjust the DynamicsProcessor.NumKeyChannels, so query that after updating.
+		const bool bPatchIsValid = UpdateKeySourcePatch(); 
+				
+		const int32 NumKeyChannels = DynamicsProcessor.GetKeyNumChannels();
+		const int32 NumKeySamples = InData.NumFrames * NumKeyChannels;		
+		
+		// Make enough space to handle all input samples that we're expecting.
+		// These zeros will act as input if we don't get enough samples.
+		AudioExternal.AddZeroed(NumKeySamples);	
+		
+		if (bPatchIsValid)
+		{
+			// There's a chance of a patch racing with us here and becoming stale, which would return INDEX_NONE (-1) samples here.
+			// This shouldn't matter, but we should be careful of ensuring against >= 0 samples popped.
+			const int32 NumSamplesPopped = KeySource.Patch->PopAudio(AudioExternal.GetData(), AudioExternal.Num(), true /* bUseLatestAudio */);
+			ensure(NumSamplesPopped <= AudioExternal.Num());
+		}
 	}
 
 	// No key assigned (Uses input buffer as key)
