Decompression jobs can return partial results

#tests

diff --git a/Engine/Source/Runtime/AudioMixer/Private/AudioMixerSourceBuffer.cpp b/Engine/Source/Runtime/AudioMixer/Private/AudioMixerSourceBuffer.cpp
--- a/Engine/Source/Runtime/AudioMixer/Private/AudioMixerSourceBuffer.cpp
+++ b/Engine/Source/Runtime/AudioMixer/Private/AudioMixerSourceBuffer.cpp
@@ -458,10 +458,23 @@ namespace Audio
 					FDecodeAudioTaskResults TaskResult;
 					AsyncRealtimeAudioTask->GetResult(TaskResult);
 					bIsFinishedOrLooped = TaskResult.bIsFinishedOrLooped;
+					
+					extern int32 bDecodersCanReturnPartialBuffersCvar;
+					if (bDecodersCanReturnPartialBuffersCvar)
+					{
+						SourceVoiceBuffers[CurrentBuffer]->AudioData.SetNum(TaskResult.NumSamplesWritten);
 #if ENABLE_AUDIO_DEBUG
-					double AudioDuration = static_cast<double>(MONO_PCM_BUFFER_SAMPLES) / FMath::Max(1., static_cast<double>(SampleRate));
-					UpdateCPUCoreUtilization(TaskResult.CPUDuration, AudioDuration);
+						double AudioDuration = static_cast<double>(TaskResult.NumSamplesWritten) / static_cast<double>(FMath::Max(1, NumChannels * SampleRate));
+						UpdateCPUCoreUtilization(TaskResult.CPUDuration, AudioDuration);
 #endif // ENABLE_AUDIO_DEBUG
+					}
+					else
+					{
+#if ENABLE_AUDIO_DEBUG
+						double AudioDuration = static_cast<double>(MONO_PCM_BUFFER_SAMPLES) / FMath::Max(1., static_cast<double>(SampleRate));
+						UpdateCPUCoreUtilization(TaskResult.CPUDuration, AudioDuration);
+#endif //ENABLE_AUDIO_DEBUG
+					}
 				}
 				break;
 
diff --git a/Engine/Source/Runtime/AudioMixer/Private/AudioMixerSourceDecode.cpp b/Engine/Source/Runtime/AudioMixer/Private/AudioMixerSourceDecode.cpp
--- a/Engine/Source/Runtime/AudioMixer/Private/AudioMixerSourceDecode.cpp
+++ b/Engine/Source/Runtime/AudioMixer/Private/AudioMixerSourceDecode.cpp
@@ -27,6 +27,14 @@ FAutoConsoleVariableRef CVarForceSynchronizedAudioTaskKick(
 	TEXT("Force all Audio Tasks created in one \"audio render frame\" to be queued until they can all be \"kicked\" at once at the end of the frame.\n")
 	TEXT("0: Don't Force, 1: Force"),
 	ECVF_Default);
+	
+int32 bDecodersCanReturnPartialBuffersCvar = 1;
+FAutoConsoleVariableRef CVarDecodersCanReturnPartialBuffers(
+	TEXT("au.DecodersCanReturnPartialBuffers"),
+	bDecodersCanReturnPartialBuffersCvar,
+	TEXT(""),
+	ECVF_Default
+	);
 
 namespace Audio
 {
@@ -195,9 +203,23 @@ public:
 					DecodeResult.bIsFinishedOrLooped = DecodeTaskData.DecompressionState->ReadCompressedData(DecodeBuffer.GetData(), DecodeTaskData.bLoopingMode, kPCMBufferSize);
 				}
 
-				// Convert the decoded PCM data into a float buffer while still in the async task
-				Audio::ArrayPcm16ToFloat(MakeArrayView((int16*)DecodeBuffer.GetData(), DecodeTaskData.NumFramesToDecode * NumChannels)
-					, MakeArrayView(DecodeTaskData.AudioData, DecodeTaskData.NumFramesToDecode* NumChannels));
+				if (bDecodersCanReturnPartialBuffersCvar)
+				{
+					const int32 NumSamplesStreamed = NumBytesStreamed / sizeof(int16);
+
+					DecodeResult.NumSamplesWritten = NumSamplesStreamed;
+
+					// Convert the decoded PCM data into a float buffer while still in the async task
+					Audio::ArrayPcm16ToFloat(
+						MakeArrayView((int16*)DecodeBuffer.GetData(), NumSamplesStreamed)
+						, MakeArrayView(DecodeTaskData.AudioData, NumSamplesStreamed));
+				}
+				else
+				{
+					// Convert the decoded PCM data into a float buffer while still in the async task
+					Audio::ArrayPcm16ToFloat(MakeArrayView((int16*)DecodeBuffer.GetData(), DecodeTaskData.NumFramesToDecode * NumChannels)
+						, MakeArrayView(DecodeTaskData.AudioData, DecodeTaskData.NumFramesToDecode* NumChannels));
+				}
 			}
 			break;
 		}
diff --git a/Engine/Source/Runtime/AudioMixer/Private/AudioMixerSourceDecode.h b/Engine/Source/Runtime/AudioMixer/Private/AudioMixerSourceDecode.h
--- a/Engine/Source/Runtime/AudioMixer/Private/AudioMixerSourceDecode.h
+++ b/Engine/Source/Runtime/AudioMixer/Private/AudioMixerSourceDecode.h
@@ -128,6 +128,7 @@ namespace Audio
 
 		// Whether or not the audio buffer looped
 		bool bIsFinishedOrLooped;
+		int32 NumSamplesWritten = 0;
 
 #if ENABLE_AUDIO_DEBUG
 		double CPUDuration = 0;
