Fix to build issues after last Sequencer changes. AudioMixerCore can't depend on MusicEnvironemt. Dependency moved to AudioMixer, which now needs a dependency on MusicEnvironment noted in its .build.cs file.

diff --git a/Engine/Source/Runtime/AudioMixer/AudioMixer.Build.cs b/Engine/Source/Runtime/AudioMixer/AudioMixer.Build.cs
--- a/Engine/Source/Runtime/AudioMixer/AudioMixer.Build.cs
+++ b/Engine/Source/Runtime/AudioMixer/AudioMixer.Build.cs
@@ -32,6 +32,7 @@ namespace UnrealBuildTool.Rules
 					"AudioExtensions",
 					"AudioLinkCore",
 					"HeadMountedDisplay",
+					"MusicEnvironment",
 					"TraceLog"
 				}
 			);
diff --git a/Engine/Source/Runtime/AudioMixer/Private/AudioMixerModule.cpp b/Engine/Source/Runtime/AudioMixer/Private/AudioMixerModule.cpp
--- a/Engine/Source/Runtime/AudioMixer/Private/AudioMixerModule.cpp
+++ b/Engine/Source/Runtime/AudioMixer/Private/AudioMixerModule.cpp
@@ -10,6 +10,7 @@ public:
 	virtual void StartupModule() override
 	{
 		FModuleManager::Get().LoadModuleChecked(TEXT("AudioMixerCore"));
+		FModuleManager::Get().LoadModuleChecked(TEXT("MusicEnvironment"));
 		FModuleManager::Get().LoadModuleChecked(TEXT("SignalProcessing"));
 	}
 };
diff --git a/Engine/Source/Runtime/AudioMixerCore/Private/AudioMixerCoreModule.cpp b/Engine/Source/Runtime/AudioMixerCore/Private/AudioMixerCoreModule.cpp
--- a/Engine/Source/Runtime/AudioMixerCore/Private/AudioMixerCoreModule.cpp
+++ b/Engine/Source/Runtime/AudioMixerCore/Private/AudioMixerCoreModule.cpp
@@ -14,7 +14,6 @@ public:
 	virtual void StartupModule() override
 	{
 		FModuleManager::Get().LoadModuleChecked(TEXT("SignalProcessing"));
-		FModuleManager::Get().LoadModuleChecked(TEXT("MusicEnvironment"));
 	}
 };
 
