Fix for shipping / server CIS break

diff --git a/Engine/Source/Runtime/AudioMixer/Private/AudioMixerSourceManager.cpp b/Engine/Source/Runtime/AudioMixer/Private/AudioMixerSourceManager.cpp
--- a/Engine/Source/Runtime/AudioMixer/Private/AudioMixerSourceManager.cpp
+++ b/Engine/Source/Runtime/AudioMixer/Private/AudioMixerSourceManager.cpp
@@ -1359,7 +1359,7 @@ namespace Audio
 
 			SourceInfo.StartTime = StartTime;
 					
-			UE_CLOG(Audio::MatchesLogFilter(*SourceInfo.DebugName), LogAudioTiming, Verbose,
+			UE_CLOG(Audio::MatchesLogFilter(*SourceInfo.GetDebugName()), LogAudioTiming, Verbose,
 				TEXT("FMixerSourceManager::Play (render thread), Wave=%s, SourceId=%d"), *SourceInfo.GetDebugName(), SourceId);
 
 			AUDIO_MIXER_DEBUG_LOG(SourceId, TEXT("Is playing"));
