Attempt to fix ASan warning

#tests

diff --git a/Engine/Source/Runtime/AudioMixer/Private/AudioMixerSource.cpp b/Engine/Source/Runtime/AudioMixer/Private/AudioMixerSource.cpp
--- a/Engine/Source/Runtime/AudioMixer/Private/AudioMixerSource.cpp
+++ b/Engine/Source/Runtime/AudioMixer/Private/AudioMixerSource.cpp
@@ -1320,8 +1320,7 @@ namespace Audio
 
 		AUDIO_MIXER_TRACE_CPUPROFILER_EVENT_SCOPE(AudioMixerSource::Play);
 
-
-		UE_CLOG(Audio::MatchesLogFilter(*WaveInstance->GetName()), LogAudioTiming, Verbose, 
+		UE_CLOG(WaveInstance && Audio::MatchesLogFilter(*WaveInstance->GetName()), LogAudioTiming, Verbose, 
 			TEXT("FMixerSource::Play, Name=%s, StartTime=%.2f bProcedural=%d, bIsSourceBus=%d, CachedRealTimeFirstBuffer=0x%p, MixerSourceBufferValid=%s, TickCount=%d, InitState=%s"),
 			*WaveInstance->GetName(), WaveInstance->StartTime, WaveInstance->WaveData->bProcedural, WaveInstance->WaveData->bIsSourceBus,
 			WaveInstance->WaveData->CachedRealtimeFirstBuffer,ToCStr(LexToString(MixerSourceBuffer.IsValid())), TickCount, ToCStr(LexToString(InitializationState)));
@@ -1330,7 +1329,7 @@ namespace Audio
 		// we'll just not actually play the source here. Instead we'll call play when the sound finishes loading.
 		if (MixerSourceVoice && InitializationState == EMixerSourceInitializationState::Initialized)
 		{
-			UE_CLOG(Audio::MatchesLogFilter(*WaveInstance->GetName()), LogAudioTiming, Verbose,
+			UE_CLOG(WaveInstance && Audio::MatchesLogFilter(*WaveInstance->GetName()), LogAudioTiming, Verbose,
 				TEXT("FMixerSourceVoice::Play, Name=%s, StartTime=%.2f bProcedural=%d, bIsSourceBus=%d, CachedRealTimeFirstBuffer=0x%p, MixerSourceBufferValid=%s, TickCount=%d, InitState=%s"),
 				*WaveInstance->GetName(), WaveInstance->StartTime, WaveInstance->WaveData->bProcedural, WaveInstance->WaveData->bIsSourceBus,
 				WaveInstance->WaveData->CachedRealtimeFirstBuffer,ToCStr(LexToString(MixerSourceBuffer.IsValid())), TickCount, ToCStr(LexToString(InitializationState)));
