Update key channel count on submix dynamics processor even if teh key source fails to update.

diff --git a/Engine/Source/Runtime/AudioMixer/Private/Effects/AudioMixerSubmixEffectDynamicsProcessor.cpp b/Engine/Source/Runtime/AudioMixer/Private/Effects/AudioMixerSubmixEffectDynamicsProcessor.cpp
--- a/Engine/Source/Runtime/AudioMixer/Private/Effects/AudioMixerSubmixEffectDynamicsProcessor.cpp
+++ b/Engine/Source/Runtime/AudioMixer/Private/Effects/AudioMixerSubmixEffectDynamicsProcessor.cpp
@@ -314,14 +314,25 @@ void FSubmixEffectDynamicsProcessor::OnProcessAudio(const FSoundEffectSubmixInpu
 		// Update our key source. (this validates the existing patch and create new one if necessary)
 		// NOTE: This call can adjust the DynamicsProcessor.NumKeyChannels, so query that after updating.
 		const bool bPatchIsValid = UpdateKeySourcePatch(); 
-				
+		if (!bPatchIsValid)
+		{
+			// If the patch was not created, we emulate a patch by pumping in 
+			// silence with an equal number of channels as the input audio.
+			// The key channel count needs to get updated here as it may not have
+			// been updated to a valid value inside of `UpdateKeySourcePatch()`
+			if (InData.NumChannels != DynamicsProcessor.GetKeyNumChannels())
+			{
+				DynamicsProcessor.SetKeyNumChannels(InData.NumChannels);
+			}
+		}
+
 		const int32 NumKeyChannels = DynamicsProcessor.GetKeyNumChannels();
 		const int32 NumKeySamples = InData.NumFrames * NumKeyChannels;		
 		
 		// Make enough space to handle all input samples that we're expecting.
 		// These zeros will act as input if we don't get enough samples.
 		AudioExternal.AddZeroed(NumKeySamples);	
-		
+
 		if (bPatchIsValid)
 		{
 			// There's a chance of a patch racing with us here and becoming stale, which would return INDEX_NONE (-1) samples here.
