Delay the stopping of a sound started with a delay due to Quartz by one buffer so the sound's output.

#swarm https://p4-swarm.epicgames.net/reviews/39392679

diff --git a/Engine/Source/Runtime/AudioMixer/Private/AudioMixerSourceManager.cpp b/Engine/Source/Runtime/AudioMixer/Private/AudioMixerSourceManager.cpp
--- a/Engine/Source/Runtime/AudioMixer/Private/AudioMixerSourceManager.cpp
+++ b/Engine/Source/Runtime/AudioMixer/Private/AudioMixerSourceManager.cpp
@@ -2062,8 +2062,7 @@ namespace Audio
 			}
 			else
 			{
-				SourceInfo.bIsLastBuffer = !SourceInfo.SubCallbackDelayLengthInFrames;
-				SourceInfo.SubCallbackDelayLengthInFrames = 0;
+				SourceInfo.bIsLastBuffer = true;
 				return;
 			}
 
@@ -3356,6 +3355,13 @@ namespace Audio
 			// Check for the stopping condition to "turn the sound off"
 			if (SourceInfo.bIsLastBuffer)
 			{
+				// since we started with a delay, give the sound one more callback to finish rendering the delayline
+				if (SourceInfo.SubCallbackDelayLengthInFrames > 0) 
+				{
+					SourceInfo.SubCallbackDelayLengthInFrames = 0;
+					continue;
+				}
+				
 				if (!SourceInfo.bIsDone)
 				{
 					SourceInfo.bIsDone = true;
