Fixing static analysis warning

diff --git a/Engine/Source/Runtime/AudioMixer/Private/AudioMixerSubmix.cpp b/Engine/Source/Runtime/AudioMixer/Private/AudioMixerSubmix.cpp
--- a/Engine/Source/Runtime/AudioMixer/Private/AudioMixerSubmix.cpp
+++ b/Engine/Source/Runtime/AudioMixer/Private/AudioMixerSubmix.cpp
@@ -1662,7 +1662,7 @@ namespace Audio
 			}
 
 #if UE_AUDIO_PROFILERTRACE_ENABLED
-			if (MixerDevice && OwningSubmixObject.IsValid())
+			if (OwningSubmixObject.IsValid())
 			{
 				UE_TRACE_LOG(Audio, SoundSubmixHasActivity, AudioChannel)
 					<< SoundSubmixHasActivity.DeviceId(MixerDevice->DeviceID)
