[AudioMixer] Removed leftover UE_LOG warning in FMixerSubmix::Init

diff --git a/Engine/Source/Runtime/AudioMixer/Private/AudioMixerSubmix.cpp b/Engine/Source/Runtime/AudioMixer/Private/AudioMixerSubmix.cpp
--- a/Engine/Source/Runtime/AudioMixer/Private/AudioMixerSubmix.cpp
+++ b/Engine/Source/Runtime/AudioMixer/Private/AudioMixerSubmix.cpp
@@ -285,7 +285,6 @@ namespace Audio
 #if UE_AUDIO_PROFILERTRACE_ENABLED
 			if (MixerDevice)
 			{
-				UE_LOG(LogAudioMixer, Warning, TEXT("SoundSubmixActivate: %s"), *(OwningSubmixObject->GetName()));
 				UE_TRACE_LOG(Audio, SoundSubmixActivate, AudioChannel)
 					<< SoundSubmixActivate.DeviceId(MixerDevice->DeviceID)
 					<< SoundSubmixActivate.SubmixId(OwningSubmixObject->GetUniqueID())
