[Audio Mixer] AudioMixerSourceManager: Added StartAudioBus override (this is needed in order to be able to send trace messages with an audio bus name to Audio Insights)

diff --git a/Engine/Plugins/Runtime/AudioWidgets/Source/AudioWidgets/Private/WaveformAudioSamplesDataProvider.cpp b/Engine/Plugins/Runtime/AudioWidgets/Source/AudioWidgets/Private/WaveformAudioSamplesDataProvider.cpp
--- a/Engine/Plugins/Runtime/AudioWidgets/Source/AudioWidgets/Private/WaveformAudioSamplesDataProvider.cpp
+++ b/Engine/Plugins/Runtime/AudioWidgets/Source/AudioWidgets/Private/WaveformAudioSamplesDataProvider.cpp
@@ -131,7 +131,8 @@ namespace AudioWidgets
 
 		const Audio::FAudioBusKey AudioBusKey = Audio::FAudioBusKey(AudioBusId);
 
-		AudioBusSubsystem->StartAudioBus(AudioBusKey, NumChannels, false);
+		const FString WaveformAudioSamplesDataProviderAudioBusName = FString::Format(TEXT("_WaveformAudioSamplesDataProvider_AudioBusId_{0}"), { AudioBusId });
+		AudioBusSubsystem->StartAudioBus(AudioBusKey, WaveformAudioSamplesDataProviderAudioBusName, NumChannels, false);
 
 		// Get an output patch for the audio bus
 		const uint32 NumFramesPerBufferToAnalyze = MixerDevice->GetNumOutputFrames();
diff --git a/Engine/Plugins/Runtime/Metasound/Source/MetasoundEngine/Private/MetasoundAudioBusReaderNode.cpp b/Engine/Plugins/Runtime/Metasound/Source/MetasoundEngine/Private/MetasoundAudioBusReaderNode.cpp
--- a/Engine/Plugins/Runtime/Metasound/Source/MetasoundEngine/Private/MetasoundAudioBusReaderNode.cpp
+++ b/Engine/Plugins/Runtime/Metasound/Source/MetasoundEngine/Private/MetasoundAudioBusReaderNode.cpp
@@ -231,7 +231,9 @@ namespace Metasound
 				AudioBusId = AudioBusProxy->AudioBusId;
 
 				Audio::FAudioBusKey AudioBusKey(AudioBusId);
-				AudioBusSubsystem->StartAudioBus(AudioBusKey, AudioBusChannels, false);
+
+				const FString AudioBusReaderNodeBusName = FString::Format(TEXT("_AudioBusReaderNode_AudioBusId_{0}"), { AudioBusId });
+				AudioBusSubsystem->StartAudioBus(AudioBusKey, AudioBusReaderNodeBusName, AudioBusChannels, false);
 
 				AudioBusPatchOutput = AudioBusSubsystem->AddPatchOutputForSoundAndAudioBus(InstanceID, AudioBusKey, BlockSizeFrames, int32(AudioBusChannels));
 				PatchInput = AudioBusPatchOutput;
diff --git a/Engine/Plugins/Runtime/Metasound/Source/MetasoundEngine/Private/MetasoundAudioBusWriterNode.cpp b/Engine/Plugins/Runtime/Metasound/Source/MetasoundEngine/Private/MetasoundAudioBusWriterNode.cpp
--- a/Engine/Plugins/Runtime/Metasound/Source/MetasoundEngine/Private/MetasoundAudioBusWriterNode.cpp
+++ b/Engine/Plugins/Runtime/Metasound/Source/MetasoundEngine/Private/MetasoundAudioBusWriterNode.cpp
@@ -167,7 +167,9 @@ namespace Metasound
 				AudioBusId = AudioBusProxy->AudioBusId;
 
 				Audio::FAudioBusKey AudioBusKey(AudioBusId);
-				AudioBusSubsystem->StartAudioBus(AudioBusKey, AudioBusChannels, false);
+
+				const FString AudioBusWriterNodeBusName = FString::Format(TEXT("_AudioBusWriterNode_AudioBusId_{0}"), { AudioBusId });
+				AudioBusSubsystem->StartAudioBus(AudioBusKey, AudioBusWriterNodeBusName, AudioBusChannels, false);
 
 				AudioBusPatchInput = AudioBusSubsystem->AddPatchInputForSoundAndAudioBus(InstanceID, AudioBusKey, BlockSizeFrames, int32(AudioBusChannels));
 				int32 NumBlocksToPush = InitialNumBlocks();
diff --git a/Engine/Source/Runtime/AudioAnalyzer/Private/AudioAnalyzer.cpp b/Engine/Source/Runtime/AudioAnalyzer/Private/AudioAnalyzer.cpp
--- a/Engine/Source/Runtime/AudioAnalyzer/Private/AudioAnalyzer.cpp
+++ b/Engine/Source/Runtime/AudioAnalyzer/Private/AudioAnalyzer.cpp
@@ -98,14 +98,16 @@ void UAudioAnalyzer::StartAnalyzing(const Audio::FDeviceId InAudioDeviceId, UAud
 	}
 	checkf(AnalyzerFactory != nullptr, TEXT("Need to register the factory as a modular feature for the analyzer '%s'."), *GetAnalyzerFactoryName().ToString());
 
-	// Start the audio bus. This won't do anythign if the bus is already started elsewhere.
+	// Start the audio bus. This won't do anything if the bus is already started elsewhere.
 	uint32 AudioBusId = AudioBus->GetUniqueID();
 	int32 NumChannels = (int32)AudioBus->AudioBusChannels + 1;
 
 	UAudioBusSubsystem* AudioBusSubsystem = MixerDevice->GetSubsystem<UAudioBusSubsystem>();
 	check(AudioBusSubsystem);
 	Audio::FAudioBusKey AudioBusKey = Audio::FAudioBusKey(AudioBusId);
-	AudioBusSubsystem->StartAudioBus(AudioBusKey, NumChannels, false);
+
+	const FString AudioAnalyzerBusName = FString::Format(TEXT("_Analyzer_{0}_AudioBusId_{1}"), { *GetName(), AudioBusId });
+	AudioBusSubsystem->StartAudioBus(AudioBusKey, AudioAnalyzerBusName, NumChannels, false);
 
 	// Get an output patch for the audio bus
 	NumFramesPerBufferToAnalyze = MixerDevice->GetNumOutputFrames();
diff --git a/Engine/Source/Runtime/AudioMixer/Private/AudioBusSubsystem.cpp b/Engine/Source/Runtime/AudioMixer/Private/AudioBusSubsystem.cpp
--- a/Engine/Source/Runtime/AudioMixer/Private/AudioBusSubsystem.cpp
+++ b/Engine/Source/Runtime/AudioMixer/Private/AudioBusSubsystem.cpp
@@ -30,6 +30,11 @@ void UAudioBusSubsystem::Deinitialize()
 }
 
 void UAudioBusSubsystem::StartAudioBus(Audio::FAudioBusKey InAudioBusKey, int32 InNumChannels, bool bInIsAutomatic)
+{
+	StartAudioBus(InAudioBusKey, FString(), InNumChannels, bInIsAutomatic);
+}
+
+void UAudioBusSubsystem::StartAudioBus(Audio::FAudioBusKey InAudioBusKey, const FString& InAudioBusName, int32 InNumChannels, bool bInIsAutomatic)
 {
 	if (IsInGameThread())
 	{
@@ -45,11 +50,11 @@ void UAudioBusSubsystem::StartAudioBus(Audio::FAudioBusKey InAudioBusKey, int32
 
 		ActiveAudioBuses_GameThread.Add(InAudioBusKey, BusData);
 
-		FAudioThread::RunCommandOnAudioThread([this, InAudioBusKey, InNumChannels, bInIsAutomatic]()
+		FAudioThread::RunCommandOnAudioThread([this, InAudioBusKey, InAudioBusName, InNumChannels, bInIsAutomatic]()
 		{
 			if (Audio::FMixerSourceManager* MixerSourceManager = GetMutableSourceManager())
 			{
-				MixerSourceManager->StartAudioBus(InAudioBusKey, InNumChannels, bInIsAutomatic);
+				MixerSourceManager->StartAudioBus(InAudioBusKey, InAudioBusName, InNumChannels, bInIsAutomatic);
 			}
 		});
 	}
@@ -58,9 +63,9 @@ void UAudioBusSubsystem::StartAudioBus(Audio::FAudioBusKey InAudioBusKey, int32
 		// If we're not the game thread, this needs to be on the game thread, so queue up a command to execute it on the game thread
 		if (Audio::FMixerDevice* MixerDevice = GetMutableMixerDevice())
 		{
-			MixerDevice->GameThreadMPSCCommand([this, InAudioBusKey, InNumChannels, bInIsAutomatic]
+			MixerDevice->GameThreadMPSCCommand([this, InAudioBusKey, InAudioBusName, InNumChannels, bInIsAutomatic]
 			{
-				StartAudioBus(InAudioBusKey, InNumChannels, bInIsAutomatic);
+				StartAudioBus(InAudioBusKey, InAudioBusName, InNumChannels, bInIsAutomatic);
 			});
 		}
 	}
@@ -246,7 +251,7 @@ void UAudioBusSubsystem::InitDefaultAudioBuses()
 				if (UAudioBus* AudioBus = Cast<UAudioBus>(BusObject))
 				{
 					const int32 NumChannels = static_cast<int32>(AudioBus->AudioBusChannels) + 1;
-					StartAudioBus(Audio::FAudioBusKey(AudioBus->GetUniqueID()), NumChannels, false /* bInIsAutomatic */);
+					StartAudioBus(Audio::FAudioBusKey(AudioBus->GetUniqueID()), AudioBus->GetPathName(), NumChannels, false /* bInIsAutomatic */);
 
 					TStrongObjectPtr<UAudioBus>AddedBus(AudioBus);
 					DefaultAudioBuses.AddUnique(AddedBus);
diff --git a/Engine/Source/Runtime/AudioMixer/Private/AudioMixerBlueprintLibrary.cpp b/Engine/Source/Runtime/AudioMixer/Private/AudioMixerBlueprintLibrary.cpp
--- a/Engine/Source/Runtime/AudioMixer/Private/AudioMixerBlueprintLibrary.cpp
+++ b/Engine/Source/Runtime/AudioMixer/Private/AudioMixerBlueprintLibrary.cpp
@@ -655,7 +655,7 @@ void UAudioMixerBlueprintLibrary::StartAudioBus(const UObject* WorldContextObjec
 
 		UAudioBusSubsystem* AudioBusSubsystem = MixerDevice->GetSubsystem<UAudioBusSubsystem>();
 		check(AudioBusSubsystem);
-		AudioBusSubsystem->StartAudioBus(Audio::FAudioBusKey(AudioBusId), NumChannels, false);
+		AudioBusSubsystem->StartAudioBus(Audio::FAudioBusKey(AudioBusId), AudioBus->GetPathName(), NumChannels, false);
 	}
 	else
 	{
diff --git a/Engine/Source/Runtime/AudioMixer/Private/AudioMixerSourceManager.cpp b/Engine/Source/Runtime/AudioMixer/Private/AudioMixerSourceManager.cpp
--- a/Engine/Source/Runtime/AudioMixer/Private/AudioMixerSourceManager.cpp
+++ b/Engine/Source/Runtime/AudioMixer/Private/AudioMixerSourceManager.cpp
@@ -1203,6 +1203,11 @@ namespace Audio
 	}
 
 	void FMixerSourceManager::StartAudioBus(FAudioBusKey InAudioBusKey, int32 InNumChannels, bool bInIsAutomatic)
+	{
+		StartAudioBus(InAudioBusKey, FString(), InNumChannels, bInIsAutomatic);
+	}
+
+	void FMixerSourceManager::StartAudioBus(FAudioBusKey InAudioBusKey, const FString& InAudioBusName, int32 InNumChannels, bool bInIsAutomatic)
 	{
 		if (AudioBusKeys_AudioThread.Contains(InAudioBusKey))
 		{
diff --git a/Engine/Source/Runtime/AudioMixer/Private/AudioMixerSourceManager.h b/Engine/Source/Runtime/AudioMixer/Private/AudioMixerSourceManager.h
--- a/Engine/Source/Runtime/AudioMixer/Private/AudioMixerSourceManager.h
+++ b/Engine/Source/Runtime/AudioMixer/Private/AudioMixerSourceManager.h
@@ -206,8 +206,11 @@ namespace Audio
 		void InitSource(const int32 SourceId, const FMixerSourceVoiceInitParams& InitParams);
 
 		// Creates and starts an audio bus manually.
+		UE_DEPRECATED(5.6, "Use the StartAudioBus version that requires an AudioBus name.")
 		void StartAudioBus(FAudioBusKey InAudioBusKey, int32 InNumChannels, bool bInIsAutomatic);
 
+		void StartAudioBus(FAudioBusKey InAudioBusKey, const FString& InAudioBusName, int32 InNumChannels, bool bInIsAutomatic);
+
 		// Stops an audio bus manually
 		void StopAudioBus(FAudioBusKey InAudioBusKey);
 
diff --git a/Engine/Source/Runtime/AudioMixer/Private/Effects/AudioMixerSubmixEffectDynamicsProcessor.cpp b/Engine/Source/Runtime/AudioMixer/Private/Effects/AudioMixerSubmixEffectDynamicsProcessor.cpp
--- a/Engine/Source/Runtime/AudioMixer/Private/Effects/AudioMixerSubmixEffectDynamicsProcessor.cpp
+++ b/Engine/Source/Runtime/AudioMixer/Private/Effects/AudioMixerSubmixEffectDynamicsProcessor.cpp
@@ -217,7 +217,8 @@ bool FSubmixEffectDynamicsProcessor::UpdateKeySourcePatch()
 				const FAudioBusKey BusKey(ObjectId);
 
 				// Make sure the bus is started. (if it's started already this does nothing).
-				AudioBusSubsystem->StartAudioBus(BusKey, NumChannels, /*bInIsAutomatic=*/false);
+				const FString SubmixEffectDynamicsKeySourceAudioBusName = FString::Format(TEXT("_SubmixEffectDynamicsKeySourceId_{0}"), { ObjectId });
+				AudioBusSubsystem->StartAudioBus(BusKey, SubmixEffectDynamicsKeySourceAudioBusName, NumChannels, /*bInIsAutomatic=*/false);
 
 				// Add/Recreate a new patch to this bus
 				KeySource.Patch = AudioBusSubsystem->AddPatchOutputForAudioBus(BusKey, NumFrames, NumChannels);
diff --git a/Engine/Source/Runtime/AudioMixer/Public/AudioBusSubsystem.h b/Engine/Source/Runtime/AudioMixer/Public/AudioBusSubsystem.h
--- a/Engine/Source/Runtime/AudioMixer/Public/AudioBusSubsystem.h
+++ b/Engine/Source/Runtime/AudioMixer/Public/AudioBusSubsystem.h
@@ -79,7 +79,10 @@ public:
 	//~ End USubsystem interface
 
 	// Audio bus API from FMixerDevice
+	UE_DEPRECATED(5.6, "Use the StartAudioBus version that requires an AudioBus name.")
 	AUDIOMIXER_API void StartAudioBus(Audio::FAudioBusKey InAudioBusKey, int32 InNumChannels, bool bInIsAutomatic);
+
+	AUDIOMIXER_API void StartAudioBus(Audio::FAudioBusKey InAudioBusKey, const FString& InAudioBusName, int32 InNumChannels, bool bInIsAutomatic);
 	AUDIOMIXER_API void StopAudioBus(Audio::FAudioBusKey InAudioBusKey);
 	AUDIOMIXER_API bool IsAudioBusActive(Audio::FAudioBusKey InAudioBusKey) const;
 	
diff --git a/Engine/Source/Runtime/Engine/Private/AudioBus.cpp b/Engine/Source/Runtime/Engine/Private/AudioBus.cpp
--- a/Engine/Source/Runtime/Engine/Private/AudioBus.cpp
+++ b/Engine/Source/Runtime/Engine/Private/AudioBus.cpp
@@ -62,13 +62,13 @@ void UAudioBus::PostEditChangeProperty(FPropertyChangedEvent& PropertyChangedEve
 	{
 		if (FAudioDeviceManager* DeviceManager = FAudioDeviceManager::Get())
 		{
-			DeviceManager->IterateOverAllDevices([BusId = GetUniqueID(), NumChannels = AudioBusChannels](Audio::FDeviceId, FAudioDevice* InDevice)
+			DeviceManager->IterateOverAllDevices([BusId = GetUniqueID(), Name = GetPathName(), NumChannels = AudioBusChannels](Audio::FDeviceId, FAudioDevice* InDevice)
 			{
 				UAudioBusSubsystem* AudioBusSubsystem = InDevice->GetSubsystem<UAudioBusSubsystem>();
 				check(AudioBusSubsystem);
 				Audio::FAudioBusKey AudioBusKey = Audio::FAudioBusKey(BusId);
 				AudioBusSubsystem->StopAudioBus(AudioBusKey);
-				AudioBusSubsystem->StartAudioBus(AudioBusKey, (int32)NumChannels + 1, false /* bInIsAutomatic */);
+				AudioBusSubsystem->StartAudioBus(AudioBusKey, Name, (int32)NumChannels + 1, false /* bInIsAutomatic */);
 			});
 		}
 	}
