Fix behavior when Submix sends to an Audio Bus with different Channel Configuration

diff --git a/Engine/Source/Runtime/AudioMixer/Private/AudioMixerSubmix.cpp b/Engine/Source/Runtime/AudioMixer/Private/AudioMixerSubmix.cpp
--- a/Engine/Source/Runtime/AudioMixer/Private/AudioMixerSubmix.cpp
+++ b/Engine/Source/Runtime/AudioMixer/Private/AudioMixerSubmix.cpp
@@ -555,6 +555,9 @@ namespace Audio
 			{
 				AudioBuses.Emplace(InAudioBusKey, InPatchInput);
 			}
+
+			// If we are sending to any audio buses, set up this buffer for potential downmixing
+			AudioBusSendBuffer.SetNumUninitialized(NumSamples * static_cast<int32>(EAudioBusChannels::MaxChannelCount));
 		});
 	}
 
@@ -1703,9 +1706,16 @@ namespace Audio
 	{
 		TRACE_CPUPROFILER_EVENT_SCOPE(FMixerSubmix::SendAudioToRegisteredAudioBuses);
 
-		for (auto& [AudioBusKey, PatchInput] : AudioBuses)
+		if (FMixerSourceManager* AudioMixerSourceManager = MixerDevice->GetSourceManager())
 		{
-			PatchInput.PushAudio(OutAudioBuffer.GetData(), OutAudioBuffer.Num());
+			for (auto& [AudioBusKey, PatchInput] : AudioBuses)
+			{
+				// Submix and Audio Bus do not necessarily have the same channel configuration, so we must downmix first
+				int32 BusChannelCount = AudioMixerSourceManager->GetAudioBusNumChannels(AudioBusKey);
+				DownmixBuffer(NumChannels, OutAudioBuffer, BusChannelCount, AudioBusSendBuffer);
+				
+				PatchInput.PushAudio(AudioBusSendBuffer.GetData(), NumSamples * BusChannelCount);
+			}
 		}
 	}
 
diff --git a/Engine/Source/Runtime/AudioMixer/Public/AudioMixerSubmix.h b/Engine/Source/Runtime/AudioMixer/Public/AudioMixerSubmix.h
--- a/Engine/Source/Runtime/AudioMixer/Public/AudioMixerSubmix.h
+++ b/Engine/Source/Runtime/AudioMixer/Public/AudioMixerSubmix.h
@@ -690,5 +690,6 @@ namespace Audio
 
 		// Registered audio buses
 		TMap<Audio::FAudioBusKey, Audio::FPatchInput> AudioBuses;
+		FAlignedFloatBuffer AudioBusSendBuffer;
 	};
 }
