Fix nans in Submix Audio Bus Send

diff --git a/Engine/Source/Runtime/AudioMixer/Private/AudioMixerSubmix.cpp b/Engine/Source/Runtime/AudioMixer/Private/AudioMixerSubmix.cpp
--- a/Engine/Source/Runtime/AudioMixer/Private/AudioMixerSubmix.cpp
+++ b/Engine/Source/Runtime/AudioMixer/Private/AudioMixerSubmix.cpp
@@ -1708,13 +1708,15 @@ namespace Audio
 
 		if (FMixerSourceManager* AudioMixerSourceManager = MixerDevice->GetSourceManager())
 		{
+			int32 NumFrames = OutAudioBuffer.Num() / NumChannels;
+
 			for (auto& [AudioBusKey, PatchInput] : AudioBuses)
 			{
 				// Submix and Audio Bus do not necessarily have the same channel configuration, so we must downmix first
 				int32 BusChannelCount = AudioMixerSourceManager->GetAudioBusNumChannels(AudioBusKey);
 				DownmixBuffer(NumChannels, OutAudioBuffer, BusChannelCount, AudioBusSendBuffer);
 				
-				PatchInput.PushAudio(AudioBusSendBuffer.GetData(), NumSamples * BusChannelCount);
+				PatchInput.PushAudio(AudioBusSendBuffer.GetData(), NumFrames * BusChannelCount);
 			}
 		}
 	}
