Stop Logspam when Invalid audio bus is registered to submix

diff --git a/Engine/Source/Runtime/AudioMixer/Private/AudioMixerSubmix.cpp b/Engine/Source/Runtime/AudioMixer/Private/AudioMixerSubmix.cpp
--- a/Engine/Source/Runtime/AudioMixer/Private/AudioMixerSubmix.cpp
+++ b/Engine/Source/Runtime/AudioMixer/Private/AudioMixerSubmix.cpp
@@ -1680,9 +1680,13 @@ namespace Audio
 			{
 				// Submix and Audio Bus do not necessarily have the same channel configuration, so we must downmix first
 				int32 BusChannelCount = AudioMixerSourceManager->GetAudioBusNumChannels(AudioBusKey);
-				DownmixBuffer(NumChannels, OutAudioBuffer, BusChannelCount, AudioBusSendBuffer);
-				
-				PatchInput.PushAudio(AudioBusSendBuffer.GetData(), NumFrames * BusChannelCount);
+				// It's possible for the buses referenced by this submix to be GC'd/destroyed, and in those cases the channel count will return 0.
+				// Todo: remove buses from the list of registered buses when they are destroyed.
+				if (BusChannelCount > 0)
+				{
+					DownmixBuffer(NumChannels, OutAudioBuffer, BusChannelCount, AudioBusSendBuffer);
+					PatchInput.PushAudio(AudioBusSendBuffer.GetData(), NumFrames * BusChannelCount);
+				}
 			}
 		}
 	}
