[AudioMixer] AudioMixerSourceManager, IsAudioBufferSilent function: changed silent threshold from -60db to -160db to match silence detection already implemented in FMixerSubmix.

diff --git a/Engine/Source/Runtime/AudioMixer/Private/AudioMixerSourceManager.cpp b/Engine/Source/Runtime/AudioMixer/Private/AudioMixerSourceManager.cpp
--- a/Engine/Source/Runtime/AudioMixer/Private/AudioMixerSourceManager.cpp
+++ b/Engine/Source/Runtime/AudioMixer/Private/AudioMixerSourceManager.cpp
@@ -270,17 +270,15 @@ namespace Audio
 		int32 Index = 0;
 
 #if PLATFORM_ENABLE_VECTORINTRINSICS
-		constexpr VectorRegister4Float Minus60dBLinear = MakeVectorRegisterFloatConstant(0.001f, 0.001f, 0.001f, 0.001f);
-
 		const int32 SimdNum = NumSamples & 0xFFFFFFF0;
 		for (; Index < SimdNum; Index += 16)
 		{
 			const VectorRegister4x4Float Samples = VectorLoad16(&AudioBuffer[Index]);
 
-			if (VectorAnyGreaterThan(VectorAbs(Samples.val[0]), Minus60dBLinear) ||
-			    VectorAnyGreaterThan(VectorAbs(Samples.val[1]), Minus60dBLinear) ||
-			    VectorAnyGreaterThan(VectorAbs(Samples.val[2]), Minus60dBLinear) ||
-			    VectorAnyGreaterThan(VectorAbs(Samples.val[3]), Minus60dBLinear))
+			if (VectorAnyGreaterThan(VectorAbs(Samples.val[0]), GlobalVectorConstants::SmallNumber) ||
+			    VectorAnyGreaterThan(VectorAbs(Samples.val[1]), GlobalVectorConstants::SmallNumber) ||
+			    VectorAnyGreaterThan(VectorAbs(Samples.val[2]), GlobalVectorConstants::SmallNumber) ||
+			    VectorAnyGreaterThan(VectorAbs(Samples.val[3]), GlobalVectorConstants::SmallNumber))
 			{
 				bIsSilent = false;
 				Index = INT_MAX;
