Move Modulation Inheritance evaluations to ActiveSound and WaveInstance

diff --git a/Engine/Source/Runtime/AudioMixer/Private/AudioMixerSource.cpp b/Engine/Source/Runtime/AudioMixer/Private/AudioMixerSource.cpp
--- a/Engine/Source/Runtime/AudioMixer/Private/AudioMixerSource.cpp
+++ b/Engine/Source/Runtime/AudioMixer/Private/AudioMixerSource.cpp
@@ -97,370 +97,47 @@ namespace Audio
 
 	namespace ModulationUtils
 	{
-		void MixInRoutedValue(const FModulationParameter& InParam, float& InOutValueA, float InValueB)
-		{
-			if (InParam.bRequiresConversion)
-			{
-				InParam.NormalizedFunction(InOutValueA);
-				InParam.NormalizedFunction(InValueB);
-			}
-			InParam.MixFunction(InOutValueA, InValueB);
-			if (InParam.bRequiresConversion)
-			{
-				InParam.UnitFunction(InOutValueA);
-			}
-		}
-
-		FSoundModulationDestinationSettings InitRoutedDestinationSettings(
-			const EModulationRouting& InActiveSoundRouting,
-			const FSoundModulationDestinationSettings& InActiveSoundSettings,
-			const EModulationRouting& InWaveRouting,
-			const FSoundModulationDestinationSettings& InWaveSettings,
-			const USoundClass* InSoundClass,
-			const FModulationParameter& InParam,
-			TFunctionRef<const FSoundModulationDestinationSettings* (const USoundClass&)> InGetSoundClassDestinationFunction)
-		{
-			auto UnionSoundClassSettings = [&](FSoundModulationDestinationSettings& InOutSettings)
-			{
-				if (InSoundClass)
-				{
-					const FSoundModulationDestinationSettings& ClassSettings = *InGetSoundClassDestinationFunction(*InSoundClass);
-					MixInRoutedValue(InParam, InOutSettings.Value, ClassSettings.Value);
-					InOutSettings.Modulators = InOutSettings.Modulators.Union(ClassSettings.Modulators);
-				}
-			};
-
-			switch (InActiveSoundRouting)
-			{
-				case EModulationRouting::Union:
-				{
-					FSoundModulationDestinationSettings UnionSettings = InActiveSoundSettings;
-					switch (InWaveRouting)
-					{
-						case EModulationRouting::Union:
-						{
-							MixInRoutedValue(InParam, UnionSettings.Value, InWaveSettings.Value);
-							UnionSettings.Modulators = UnionSettings.Modulators.Union(InWaveSettings.Modulators);
-							UnionSoundClassSettings(UnionSettings);
-							return UnionSettings;
-						}
-
-						case EModulationRouting::Inherit:
-						{
-							UnionSoundClassSettings(UnionSettings);
-						}
-						break;
-
-						case EModulationRouting::Override:
-						{
-							MixInRoutedValue(InParam, UnionSettings.Value, InWaveSettings.Value);
-							UnionSettings.Modulators = UnionSettings.Modulators.Union(InWaveSettings.Modulators);
-						}
-						break;
-
-						case EModulationRouting::Disable:
-						default:
-						break;
-					}
-
-					return UnionSettings;
-				}
-
-				case EModulationRouting::Inherit:
-				{
-					switch (InWaveRouting)
-					{
-						case EModulationRouting::Union:
-						{
-							FSoundModulationDestinationSettings UnionSettings = InWaveSettings;
-							UnionSoundClassSettings(UnionSettings);
-							return UnionSettings;
-						}
-
-						case EModulationRouting::Inherit:
-						{
-							if (InSoundClass)
-							{
-								return *InGetSoundClassDestinationFunction(*InSoundClass);
-							}
-						}
-						break;
-
-						case EModulationRouting::Override:
-						{
-							return InWaveSettings;
-						}
-
-						case EModulationRouting::Disable:
-						default:
-						break;
-					}
-				}
-				break;
-
-				case EModulationRouting::Override:
-				{
-					return InActiveSoundSettings;
-				}
-
-				case EModulationRouting::Disable:
-				default:
-				break;
-			}
-
-			return { };
-		}
-
-		float GetRoutedDestinationValue(
-			const EModulationRouting& InActiveSoundRouting,
-			const FSoundModulationDestinationSettings& InActiveSoundSettings,
-			const EModulationRouting& InWaveRouting,
-			const FSoundModulationDestinationSettings& InWaveSettings,
-			const USoundClass* InSoundClass,
-			const FModulationParameter& InParam,
-			TFunctionRef<const FSoundModulationDestinationSettings* (const USoundClass&)> InGetSoundClassDestinationFunction)
-		{
-			auto MixInSoundClassValue = [&](float& InOutValue)
-			{
-				if (InSoundClass)
-				{
-					const FSoundModulationDestinationSettings& ClassSettings = *InGetSoundClassDestinationFunction(*InSoundClass);
-					MixInRoutedValue(InParam, InOutValue, ClassSettings.Value);
-				}
-			};
-
-			switch (InActiveSoundRouting)
-			{
-				case EModulationRouting::Union:
-				{
-					float UnionValue = InActiveSoundSettings.Value;
-					switch (InWaveRouting)
-					{
-						case EModulationRouting::Union:
-						{
-							MixInRoutedValue(InParam, UnionValue, InWaveSettings.Value);
-							MixInSoundClassValue(UnionValue);
-							return UnionValue;
-						}
-
-						case EModulationRouting::Inherit:
-						{
-							MixInSoundClassValue(UnionValue);
-						}
-						break;
-
-						case EModulationRouting::Override:
-						{
-							MixInRoutedValue(InParam, UnionValue, InWaveSettings.Value);
-						}
-						break;
-
-						case EModulationRouting::Disable:
-						default:
-						break;
-					}
-
-					return UnionValue;
-				}
-
-				case EModulationRouting::Inherit:
-				{
-					switch (InWaveRouting)
-					{
-						case EModulationRouting::Union:
-						{
-							float UnionValue = InWaveSettings.Value;
-							MixInSoundClassValue(UnionValue);
-							return UnionValue;
-						}
-
-						case EModulationRouting::Inherit:
-						{
-							if (InSoundClass)
-							{
-								return InGetSoundClassDestinationFunction(*InSoundClass)->Value;
-							}
-						}
-						break;
-
-						case EModulationRouting::Override:
-						{
-							return InWaveSettings.Value;
-						}
-
-						case EModulationRouting::Disable:
-						default:
-						break;
-					}
-				}
-				break;
-
-				case EModulationRouting::Override:
-				{
-					return InActiveSoundSettings.Value;
-				}
-
-				case EModulationRouting::Disable:
-				default:
-				break;
-			}
-
-			return 1.0f;
-		}
-
 		FSoundModulationDestinationSettings InitRoutedVolumeModulation(const FWaveInstance& InWaveInstance, const USoundWave& InWaveData, const FActiveSound& InActiveSound)
 		{
-			const EModulationRouting& ActiveSoundRouting = InActiveSound.ModulationRouting.VolumeRouting;
-			const FSoundModulationDestinationSettings& ActiveSoundSettings = InActiveSound.ModulationRouting.VolumeModulationDestination;
-
-			const EModulationRouting& WaveRouting = InWaveData.ModulationSettings.VolumeRouting;
-			const FSoundModulationDestinationSettings& WaveSettings = InWaveData.ModulationSettings.VolumeModulationDestination;
-
-			return InitRoutedDestinationSettings(
-				ActiveSoundRouting,
-				ActiveSoundSettings,
-				WaveRouting,
-				WaveSettings,
-				MixerSourcePrivate::GetFallbackSoundClass(InActiveSound, InWaveInstance),
-				Audio::GetModulationParameter("Volume"),
-				[](const USoundClass& InSoundClass) { return &InSoundClass.Properties.ModulationSettings.VolumeModulationDestination; }
-			);
+			return InWaveInstance.GetEffectiveModulationSettings(EModulationDestination::Volume);
 		}
 
 		float GetRoutedVolume(const FWaveInstance& InWaveInstance, const USoundWave& InWaveData, const FActiveSound& InActiveSound)
 		{
-			const EModulationRouting& ActiveSoundRouting = InActiveSound.ModulationRouting.VolumeRouting;
-			const FSoundModulationDestinationSettings& ActiveSoundSettings = InActiveSound.ModulationRouting.VolumeModulationDestination;
-
-			const EModulationRouting& WaveRouting = InWaveData.ModulationSettings.VolumeRouting;
-			const FSoundModulationDestinationSettings& WaveSettings = InWaveData.ModulationSettings.VolumeModulationDestination;
-
-			return GetRoutedDestinationValue(
-				ActiveSoundRouting,
-				ActiveSoundSettings,
-				WaveRouting,
-				WaveSettings,
-				MixerSourcePrivate::GetFallbackSoundClass(InActiveSound, InWaveInstance),
-				Audio::GetModulationParameter("Volume"),
-				[](const USoundClass& InSoundClass) { return &InSoundClass.Properties.ModulationSettings.VolumeModulationDestination; }
-			);
+			return InWaveInstance.GetEffectiveModulationValue(EModulationDestination::Volume);
 		}
 
 		FSoundModulationDestinationSettings InitRoutedPitchModulation(const FWaveInstance& InWaveInstance, const USoundWave& InWaveData, const FActiveSound& InActiveSound)
 		{
-			const EModulationRouting& ActiveSoundRouting = InActiveSound.ModulationRouting.PitchRouting;
-			const FSoundModulationDestinationSettings& ActiveSoundSettings = InActiveSound.ModulationRouting.PitchModulationDestination;
-
-			const EModulationRouting& WaveRouting = InWaveData.ModulationSettings.PitchRouting;
-			const FSoundModulationDestinationSettings& WaveSettings = InWaveData.ModulationSettings.PitchModulationDestination;
-
-			return InitRoutedDestinationSettings(
-				ActiveSoundRouting,
-				ActiveSoundSettings,
-				WaveRouting,
-				WaveSettings,
-				MixerSourcePrivate::GetFallbackSoundClass(InActiveSound, InWaveInstance),
-				Audio::GetModulationParameter("Pitch"),
-				[](const USoundClass& InSoundClass) { return &InSoundClass.Properties.ModulationSettings.PitchModulationDestination; }
-			);
+			return InWaveInstance.GetEffectiveModulationSettings(EModulationDestination::Pitch);
 		}
 
 		float GetRoutedPitch(const FWaveInstance& InWaveInstance, const USoundWave& InWaveData, const FActiveSound& InActiveSound)
 		{
-			const EModulationRouting& ActiveSoundRouting = InActiveSound.ModulationRouting.PitchRouting;
-			const FSoundModulationDestinationSettings& ActiveSoundSettings = InActiveSound.ModulationRouting.PitchModulationDestination;
-
-			const EModulationRouting& WaveRouting = InWaveData.ModulationSettings.PitchRouting;
-			const FSoundModulationDestinationSettings& WaveSettings = InWaveData.ModulationSettings.PitchModulationDestination;
-
-			return GetRoutedDestinationValue(
-				ActiveSoundRouting,
-				ActiveSoundSettings,
-				WaveRouting,
-				WaveSettings,
-				MixerSourcePrivate::GetFallbackSoundClass(InActiveSound, InWaveInstance),
-				Audio::GetModulationParameter("Pitch"),
-				[](const USoundClass& InSoundClass) { return &InSoundClass.Properties.ModulationSettings.PitchModulationDestination; }
-			);
+			return InWaveInstance.GetEffectiveModulationValue(EModulationDestination::Pitch);
 		}
 
 		FSoundModulationDestinationSettings InitRoutedHighpassModulation(const FWaveInstance& InWaveInstance, const USoundWave& InWaveData, const FActiveSound& InActiveSound)
 		{
-			const EModulationRouting& ActiveSoundRouting = InActiveSound.ModulationRouting.HighpassRouting;
-			const FSoundModulationDestinationSettings& ActiveSoundSettings = InActiveSound.ModulationRouting.HighpassModulationDestination;
-
-			const EModulationRouting& WaveRouting = InWaveData.ModulationSettings.HighpassRouting;
-			const FSoundModulationDestinationSettings& WaveSettings = InWaveData.ModulationSettings.HighpassModulationDestination;
-
-			return InitRoutedDestinationSettings(
-				ActiveSoundRouting,
-				ActiveSoundSettings,
-				WaveRouting,
-				WaveSettings,
-				MixerSourcePrivate::GetFallbackSoundClass(InActiveSound, InWaveInstance),
-				Audio::GetModulationParameter("HPFCutoffFrequency"),
-				[](const USoundClass& InSoundClass) { return &InSoundClass.Properties.ModulationSettings.HighpassModulationDestination; }
-			);
+			return InWaveInstance.GetEffectiveModulationSettings(EModulationDestination::Highpass);
 		}
 
 		float GetRoutedHighpass(const FWaveInstance& InWaveInstance, const USoundWave& InWaveData, const FActiveSound& InActiveSound)
 		{
-			const EModulationRouting& ActiveSoundRouting = InActiveSound.ModulationRouting.HighpassRouting;
-			const FSoundModulationDestinationSettings& ActiveSoundSettings = InActiveSound.ModulationRouting.HighpassModulationDestination;
-
-			const EModulationRouting& WaveRouting = InWaveData.ModulationSettings.HighpassRouting;
-			const FSoundModulationDestinationSettings& WaveSettings = InWaveData.ModulationSettings.HighpassModulationDestination;
-
-			return GetRoutedDestinationValue(
-				ActiveSoundRouting,
-				ActiveSoundSettings,
-				WaveRouting,
-				WaveSettings,
-				MixerSourcePrivate::GetFallbackSoundClass(InActiveSound, InWaveInstance),
-				Audio::GetModulationParameter("HPFCutoffFrequency"),
-				[](const USoundClass& InSoundClass) { return &InSoundClass.Properties.ModulationSettings.HighpassModulationDestination; }
-			);
+			return InWaveInstance.GetEffectiveModulationValue(EModulationDestination::Highpass);
 		}
 
 		FSoundModulationDestinationSettings InitRoutedLowpassModulation(const FWaveInstance& InWaveInstance, const USoundWave& InWaveData, const FActiveSound& InActiveSound)
 		{
-			const EModulationRouting& ActiveSoundRouting = InActiveSound.ModulationRouting.LowpassRouting;
-			const FSoundModulationDestinationSettings& ActiveSoundSettings = InActiveSound.ModulationRouting.LowpassModulationDestination;
-
-			const EModulationRouting& WaveRouting = InWaveData.ModulationSettings.LowpassRouting;
-			const FSoundModulationDestinationSettings& WaveSettings = InWaveData.ModulationSettings.LowpassModulationDestination;
-
-			return InitRoutedDestinationSettings(
-				ActiveSoundRouting,
-				ActiveSoundSettings,
-				WaveRouting,
-				WaveSettings,
-				MixerSourcePrivate::GetFallbackSoundClass(InActiveSound, InWaveInstance),
-				Audio::GetModulationParameter("LPFCutoffFrequency"),
-				[](const USoundClass& InSoundClass) { return &InSoundClass.Properties.ModulationSettings.LowpassModulationDestination; }
-			);
+			return InWaveInstance.GetEffectiveModulationSettings(EModulationDestination::Lowpass);
 		}
 
 		float GetRoutedLowpass(const FWaveInstance& InWaveInstance, const USoundWave& InWaveData, const FActiveSound& InActiveSound)
 		{
-			const EModulationRouting& ActiveSoundRouting = InActiveSound.ModulationRouting.LowpassRouting;
-			const FSoundModulationDestinationSettings& ActiveSoundSettings = InActiveSound.ModulationRouting.LowpassModulationDestination;
-
-			const EModulationRouting& WaveRouting = InWaveData.ModulationSettings.LowpassRouting;
-			const FSoundModulationDestinationSettings& WaveSettings = InWaveData.ModulationSettings.LowpassModulationDestination;
-
-			return GetRoutedDestinationValue(
-				ActiveSoundRouting,
-				ActiveSoundSettings,
-				WaveRouting,
-				WaveSettings,
-				MixerSourcePrivate::GetFallbackSoundClass(InActiveSound, InWaveInstance),
-				Audio::GetModulationParameter("LPFCutoffFrequency"),
-				[](const USoundClass& InSoundClass) { return &InSoundClass.Properties.ModulationSettings.LowpassModulationDestination; }
-			);
+			return InWaveInstance.GetEffectiveModulationValue(EModulationDestination::Lowpass);
 		}
 
-		FSoundModulationDefaultSettings InitRoutedModulation(const FWaveInstance& InWaveInstance, const USoundWave& InWaveData, FActiveSound* InActiveSound)
+		FSoundModulationDefaultSettings InitRoutedModulation(const FWaveInstance& InWaveInstance, const USoundWave& InWaveData, const FActiveSound* InActiveSound)
 		{
 			FSoundModulationDefaultSettings Settings;
 			if (InActiveSound)
@@ -474,7 +151,7 @@ namespace Audio
 			return Settings;
 		}
 
-		FSoundModulationDefaultRoutingSettings UpdateRoutedModulation(const FWaveInstance& InWaveInstance, const USoundWave& InWaveData, FActiveSound* InActiveSound)
+		FSoundModulationDefaultRoutingSettings UpdateRoutedModulation(const FWaveInstance& InWaveInstance, const USoundWave& InWaveData, const FActiveSound* InActiveSound)
 		{
 			FSoundModulationDefaultRoutingSettings NewRouting;
 
diff --git a/Engine/Source/Runtime/Engine/Classes/Sound/SoundModulationDestination.h b/Engine/Source/Runtime/Engine/Classes/Sound/SoundModulationDestination.h
--- a/Engine/Source/Runtime/Engine/Classes/Sound/SoundModulationDestination.h
+++ b/Engine/Source/Runtime/Engine/Classes/Sound/SoundModulationDestination.h
@@ -121,6 +121,8 @@ struct FSoundModulationDefaultSettings
 	UPROPERTY(EditAnywhere, BlueprintReadWrite, Category = "Modulation", meta = (DisplayName = "Lowpass", AudioParam = "LPFCutoffFrequency", AudioParamClass = "SoundModulationParameterLPFFrequency"))
 	FSoundModulationDestinationSettings LowpassModulationDestination;
 
+	FSoundModulationDestinationSettings GetSettingsForDestination(EModulationDestination Destination) const;
+	
 #if WITH_EDITORONLY_DATA
 	ENGINE_API void VersionModulators();
 #endif // WITH_EDITORONLY_DATA
@@ -149,6 +151,8 @@ struct FSoundModulationDefaultRoutingSettings : public FSoundModulationDefaultSe
 	/** What low-pass modulation settings to use */
 	UPROPERTY(EditAnywhere, BlueprintReadWrite, Category = "Modulation")
 	EModulationRouting LowpassRouting = EModulationRouting::Inherit;
+	
+	EModulationRouting GetRoutingForDestination(EModulationDestination Destination) const;
 };
 
 namespace Audio
diff --git a/Engine/Source/Runtime/Engine/Private/ActiveSound.cpp b/Engine/Source/Runtime/Engine/Private/ActiveSound.cpp
--- a/Engine/Source/Runtime/Engine/Private/ActiveSound.cpp
+++ b/Engine/Source/Runtime/Engine/Private/ActiveSound.cpp
@@ -629,6 +629,24 @@ void FActiveSound::RemoveModulationRouting(const TSet<TObjectPtr<USoundModulator
 	bModulationRoutingUpdated = true;
 }
 
+TArray<TSet<TObjectPtr<USoundModulatorBase>>> FActiveSound::GetEffectiveVolumeModulators() const
+{
+	TArray<TSet<TObjectPtr<USoundModulatorBase>>> ModulatorSets;
+	for (const TPair<UPTRINT, FWaveInstance*>& Pair : GetWaveInstances())
+	{
+		if (const FWaveInstance* WaveInstance = Pair.Value)
+		{
+			TSet<TObjectPtr<USoundModulatorBase>> VolumeModulators = WaveInstance->GetEffectiveModulationSettings(EModulationDestination::Volume).Modulators;
+			if (!VolumeModulators.IsEmpty())
+			{
+				ModulatorSets.Add(VolumeModulators);
+			}
+		}
+	}
+
+	return ModulatorSets;
+}
+
 void FActiveSound::Stop()
 {
 	if (AudioDevice)
diff --git a/Engine/Source/Runtime/Engine/Private/Audio.cpp b/Engine/Source/Runtime/Engine/Private/Audio.cpp
--- a/Engine/Source/Runtime/Engine/Private/Audio.cpp
+++ b/Engine/Source/Runtime/Engine/Private/Audio.cpp
@@ -177,6 +177,166 @@ bool IsSpatializationCVarEnabled()
 	return AllowAudioSpatializationCVar != 0;
 }
 
+float FWaveInstance::GetEffectiveModulationValue(const EModulationDestination InDestination) const
+{
+	FSoundModulationDestinationSettings Settings;
+	float OutValue;
+	CalculateInheritedModulationSettingsAndValue(InDestination, Settings, OutValue);
+
+	return OutValue;
+}
+
+FSoundModulationDestinationSettings FWaveInstance::GetEffectiveModulationSettings(const EModulationDestination InDestination) const
+{
+	FSoundModulationDestinationSettings Settings;
+	float OutValue;
+	CalculateInheritedModulationSettingsAndValue(InDestination, Settings, OutValue);
+
+	return Settings;
+}
+
+Audio::FModulationParameter FWaveInstance::GetParameterBasedOnDestination(const EModulationDestination InDestination)
+{
+	switch (InDestination)
+	{
+	case EModulationDestination::Volume:
+		return Audio::GetModulationParameter("Volume");
+	case EModulationDestination::Pitch:
+		return Audio::GetModulationParameter("Pitch");
+	case EModulationDestination::Highpass:
+		return Audio::GetModulationParameter("HPFCutoffFilter");
+	case EModulationDestination::Lowpass:
+		return Audio::GetModulationParameter("LPFCutoffFilter");
+	default:
+		break;
+	}
+	return Audio::FModulationParameter();
+}
+
+void FWaveInstance::CalculateInheritedModulationSettingsAndValue(const EModulationDestination InDestination, FSoundModulationDestinationSettings& InOutSettings, float& OutValue) const
+{
+	const Audio::FModulationParameter& InParam = GetParameterBasedOnDestination(InDestination);
+	
+	const FActiveSound& InActiveSound = *ActiveSound;
+	const FSoundModulationDestinationSettings& InWaveSettings = WaveData->ModulationSettings.GetSettingsForDestination(InDestination);
+	EModulationRouting InWaveRouting = WaveData->ModulationSettings.GetRoutingForDestination(InDestination);
+
+	const USoundClass* InSoundClass = InActiveSound.GetSoundClass();
+
+	auto MixInRoutedValue = [&InParam](float& InOutValueA, float InValueB)
+	{
+		if (InParam.bRequiresConversion)
+		{
+			InParam.NormalizedFunction(InOutValueA);
+			InParam.NormalizedFunction(InValueB);
+		}
+		InParam.MixFunction(InOutValueA, InValueB);
+		if (InParam.bRequiresConversion)
+		{
+			InParam.UnitFunction(InOutValueA);
+		}
+	};
+	
+	auto UnionSoundClassSettings = [&](FSoundModulationDestinationSettings& Settings)
+		{
+			if (InSoundClass)
+			{
+				const FSoundModulationDestinationSettings& ClassSettings = InSoundClass->Properties.ModulationSettings.GetSettingsForDestination(InDestination);
+				OutValue = Settings.Value;
+				MixInRoutedValue(OutValue, ClassSettings.Value);
+				Settings.Modulators = Settings.Modulators.Union(ClassSettings.Modulators);
+			}
+		};
+
+	switch (InActiveSound.ModulationRouting.VolumeRouting)
+	{
+	case EModulationRouting::Union:
+	{
+		FSoundModulationDestinationSettings UnionSettings = InActiveSound.ModulationRouting.GetSettingsForDestination(InDestination);
+		switch (InWaveRouting)
+		{
+		case EModulationRouting::Union:
+		{
+			OutValue = UnionSettings.Value;
+			MixInRoutedValue(OutValue, InWaveSettings.Value);
+			UnionSettings.Modulators = UnionSettings.Modulators.Union(InWaveSettings.Modulators);
+			UnionSoundClassSettings(UnionSettings);
+		}
+		break;
+
+		case EModulationRouting::Inherit:
+		{
+			UnionSoundClassSettings(UnionSettings);
+		}
+		break;
+
+		case EModulationRouting::Override:
+		{
+			OutValue = UnionSettings.Value;
+			MixInRoutedValue(OutValue, InWaveSettings.Value);
+			UnionSettings.Modulators = UnionSettings.Modulators.Union(InWaveSettings.Modulators);
+		}
+		break;
+
+		case EModulationRouting::Disable:
+		default:
+			break;
+		}
+
+		InOutSettings = UnionSettings;
+		OutValue = UnionSettings.Value;
+		return;
+	}
+
+	case EModulationRouting::Inherit:
+	{
+		switch (InWaveRouting)
+		{
+		case EModulationRouting::Union:
+		{
+			FSoundModulationDestinationSettings UnionSettings = InWaveSettings;
+			UnionSoundClassSettings(UnionSettings);
+			InOutSettings = UnionSettings;
+			OutValue = UnionSettings.Value;
+		}
+
+		case EModulationRouting::Inherit:
+		{
+			if (InSoundClass)
+			{
+				FSoundModulationDestinationSettings InheritSettings = InSoundClass->Properties.ModulationSettings.GetSettingsForDestination(InDestination);
+				InOutSettings = InheritSettings;
+				OutValue = InheritSettings.Value;
+			}
+		}
+		break;
+
+		case EModulationRouting::Override:
+		{
+			InOutSettings = InWaveSettings;
+			OutValue = InWaveSettings.Value;
+		}
+
+		case EModulationRouting::Disable:
+		default:
+			break;
+		}
+	}
+	break;
+
+	case EModulationRouting::Override:
+	{
+		FSoundModulationDestinationSettings OverrideSettings = InActiveSound.ModulationRouting.GetSettingsForDestination(InDestination);
+		InOutSettings = OverrideSettings;
+		OutValue = OverrideSettings.Value;	
+	}
+
+	case EModulationRouting::Disable:
+	default:
+		break;
+	}
+}
+
 /*-----------------------------------------------------------------------------
 	FSoundBuffer implementation.
 -----------------------------------------------------------------------------*/
diff --git a/Engine/Source/Runtime/Engine/Private/Audio/AudioModulationInheritance.cpp b/Engine/Source/Runtime/Engine/Private/Audio/AudioModulationInheritance.cpp
deleted file mode 100644
diff --git a/Engine/Source/Runtime/Engine/Private/Audio/AudioModulationInheritance.h b/Engine/Source/Runtime/Engine/Private/Audio/AudioModulationInheritance.h
deleted file mode 100644
diff --git a/Engine/Source/Runtime/Engine/Private/AudioVirtualLoop.cpp b/Engine/Source/Runtime/Engine/Private/AudioVirtualLoop.cpp
--- a/Engine/Source/Runtime/Engine/Private/AudioVirtualLoop.cpp
+++ b/Engine/Source/Runtime/Engine/Private/AudioVirtualLoop.cpp
@@ -4,7 +4,6 @@
 #include "Audio/AudioDebug.h"
 #include "AudioDevice.h"
 #include "AudioMixerTrace.h"
-#include "Audio/AudioModulationInheritance.h"
 
 
 static int32 bVirtualLoopsEnabledCVar = 1;
@@ -134,7 +133,7 @@ bool FAudioVirtualLoop::Virtualize(const FActiveSound& InActiveSound, FAudioDevi
 	// Modulation data now so that Volume modulation can be used to determine if a sound is loud enough to realize.
 	// It's possible that the ActiveSound contains multiple WaveInstances.
 	// Because of this, we need to store one Modulation Destination for each WaveInstance.
-	TArray<TSet<TObjectPtr<USoundModulatorBase>>> ModulatorSets = AudioModulationInheritance::GetEffectiveVolumeModulatorsFromActiveSound(InActiveSound);
+	TArray<TSet<TObjectPtr<USoundModulatorBase>>> ModulatorSets = ActiveSound->GetEffectiveVolumeModulators();
 	for (int32 DestinationIndex = 0; DestinationIndex < ModulatorSets.Num(); ++DestinationIndex)
 	{
 		OutVirtualLoop.VolumeConcurrencyDestinations.Add(Audio::FModulationDestination());
diff --git a/Engine/Source/Runtime/Engine/Private/SoundModulationDestination.cpp b/Engine/Source/Runtime/Engine/Private/SoundModulationDestination.cpp
--- a/Engine/Source/Runtime/Engine/Private/SoundModulationDestination.cpp
+++ b/Engine/Source/Runtime/Engine/Private/SoundModulationDestination.cpp
@@ -46,6 +46,43 @@ FSoundModulationDefaultRoutingSettings::FSoundModulationDefaultRoutingSettings()
 {
 }
 
+FSoundModulationDestinationSettings FSoundModulationDefaultSettings::GetSettingsForDestination(EModulationDestination Destination) const
+{
+	switch (Destination)
+	{
+	case EModulationDestination::Volume:
+		return VolumeModulationDestination;
+	case EModulationDestination::Pitch:
+		return PitchModulationDestination;
+	case EModulationDestination::Lowpass:
+		return LowpassModulationDestination;
+	case EModulationDestination::Highpass:
+		return HighpassModulationDestination;
+	default:
+		break;
+	}
+
+	return FSoundModulationDestinationSettings();
+}
+
+EModulationRouting FSoundModulationDefaultRoutingSettings::GetRoutingForDestination(EModulationDestination Destination) const
+{
+	switch (Destination)
+	{
+		case EModulationDestination::Volume:
+			return VolumeRouting;
+		case EModulationDestination::Pitch:
+			return PitchRouting;
+		case EModulationDestination::Lowpass:
+			return LowpassRouting;
+		case EModulationDestination::Highpass:
+			return HighpassRouting;
+		default:
+			break;
+	}
+	return EModulationRouting::Disable;
+}
+
 
 namespace Audio
 {
diff --git a/Engine/Source/Runtime/Engine/Public/ActiveSound.h b/Engine/Source/Runtime/Engine/Public/ActiveSound.h
--- a/Engine/Source/Runtime/Engine/Public/ActiveSound.h
+++ b/Engine/Source/Runtime/Engine/Public/ActiveSound.h
@@ -754,6 +754,10 @@ public:
 	/* Removes given Modulators from the ActiveSound, if possible. */
 	ENGINE_API void RemoveModulationRouting(const TSet<TObjectPtr<USoundModulatorBase>>& NewModulators, EModulationDestination Destination);
 
+	/** Creates a copy of the modulation configuration for all WaveInstances on this ActiveSound.
+	 * Each element in the array corresponds to a Wave Instance's modulation settings, factoring in the routing settings on the wave instances. */
+	ENGINE_API TArray<TSet<TObjectPtr<USoundModulatorBase>>> GetEffectiveVolumeModulators() const;
+	
 	/* Determines which of the provided listeners is the closest to the sound */
 	ENGINE_API int32 FindClosestListener( const TArray<struct FListener>& InListeners ) const;
 
diff --git a/Engine/Source/Runtime/Engine/Public/Audio.h b/Engine/Source/Runtime/Engine/Public/Audio.h
--- a/Engine/Source/Runtime/Engine/Public/Audio.h
+++ b/Engine/Source/Runtime/Engine/Public/Audio.h
@@ -14,6 +14,7 @@
 #include "Sound/QuartzQuantizationUtilities.h"
 #include "Sound/SoundAttenuation.h"
 #include "Sound/SoundEffectSource.h"
+#include "Sound/SoundModulationDestination.h"
 #include "Sound/SoundSubmixSend.h"
 #include "Sound/SoundSourceBusSend.h"
 #include "IAudioExtensionPlugin.h"
@@ -526,6 +527,17 @@ public:
 
 	/** Retrieves the relative render cost of wave instance. */
 	float GetRelativeRenderCost() const { return RelativeRenderCost; }
+
+	ENGINE_API float GetEffectiveModulationValue(const EModulationDestination InDestination) const;
+
+	ENGINE_API FSoundModulationDestinationSettings GetEffectiveModulationSettings(const EModulationDestination InDestination) const;
+
+private:
+	static Audio::FModulationParameter GetParameterBasedOnDestination(const EModulationDestination InDestination);
+
+	/** This function calculates the effective modulation settings and current value, based on the currently configured modulation routing settings
+	 * (e.g. it will use inherited modulators if that destination's routing is set to Inherit). */
+	void CalculateInheritedModulationSettingsAndValue(const EModulationDestination InDestination, FSoundModulationDestinationSettings& InOutSettings, float& OutValue) const;
 };
 
 /*-----------------------------------------------------------------------------
