AudioLink cvar to force source link rendering on
[RN] AudioLink cvar to force source link rendering on

diff --git a/Engine/Source/Runtime/AudioMixer/Private/AudioMixerSource.cpp b/Engine/Source/Runtime/AudioMixer/Private/AudioMixerSource.cpp
--- a/Engine/Source/Runtime/AudioMixer/Private/AudioMixerSource.cpp
+++ b/Engine/Source/Runtime/AudioMixer/Private/AudioMixerSource.cpp
@@ -49,6 +49,13 @@ FAutoConsoleVariableRef CVarUseListenerOverrideForSpread(
 	TEXT("0: Use actual distance, 1: use listener override"),
 	ECVF_Default);
 
+static int32 bForceAudioLinkOnAllSourcesCVAr = 0;
+FAutoConsoleVariableRef CVarForceAudioLinkOnAllSources(
+	TEXT("au.AudioLink.ForceOnAllSource"),
+	bForceAudioLinkOnAllSourcesCVAr,
+	TEXT("0 (off), 1 (enabled). Will force AudioLink on all Sources (if the plugin is enabled)"),
+	ECVF_Default);
+
 static uint32 AudioMixerSourceFadeMinCVar = 512;
 static FAutoConsoleCommand GSetAudioMixerSourceFadeMin(
 	TEXT("au.SourceFadeMin"),
@@ -325,7 +332,7 @@ namespace Audio
 			InitParams.SourceBufferListener = WaveInstance->SourceBufferListener;
 			InitParams.bShouldSourceBufferListenerZeroBuffer = WaveInstance->bShouldSourceBufferListenerZeroBuffer;
 
-			if (WaveInstance->bShouldUseAudioLink)
+			if (WaveInstance->bShouldUseAudioLink || bForceAudioLinkOnAllSourcesCVAr)
 			{
 				if (IAudioLinkFactory* LinkFactory = MixerDevice->GetAudioLinkFactory())
 				{				
@@ -407,7 +414,7 @@ namespace Audio
 			// Toggle muting the source if sending only to output bus.
 			// This can get set even if the source doesn't have bus sends since bus sends can be dynamically enabled.
 			InitParams.bEnableBusSends = WaveInstance->bEnableBusSends;
-			InitParams.bEnableBaseSubmix = WaveInstance->bEnableBaseSubmix;
+			InitParams.bEnableBaseSubmix = WaveInstance->bEnableBaseSubmix && !bForceAudioLinkOnAllSourcesCVAr;
 			InitParams.bEnableSubmixSends = WaveInstance->bEnableSubmixSends;
 			InitParams.PlayOrder = WaveInstance->GetPlayOrder();
 			bPreviousBusEnablement = WaveInstance->bEnableBusSends;
