AudioLink source links will automatically mute submix sends to avoid duping audio
[RN] AudioLink source links will automatically mute base submix sends to avoid duping audio

diff --git a/Engine/Source/Runtime/AudioMixer/Private/AudioMixerSourceManager.cpp b/Engine/Source/Runtime/AudioMixer/Private/AudioMixerSourceManager.cpp
--- a/Engine/Source/Runtime/AudioMixer/Private/AudioMixerSourceManager.cpp
+++ b/Engine/Source/Runtime/AudioMixer/Private/AudioMixerSourceManager.cpp
@@ -2594,6 +2594,8 @@ namespace Audio
 			Params.SourceId = SourceId;
 			Params.Buffer = SourceInfo.PreDistanceAttenuationBuffer;
 			SourceInfo.AudioLink->OnNewBuffer(Params);
+			FMemory::Memzero(SourceInfo.PreDistanceAttenuationBuffer.GetData(), SourceInfo.PreDistanceAttenuationBuffer.Num() * sizeof(float));
+			FMemory::Memzero(SourceInfo.SourceBuffer.GetData(), SourceInfo.SourceBuffer.Num() * sizeof(float));
 		}
 
 		// If we have Source Buffer Listener
