Granular Scrubbing in Sequencer

[RN] Scrubbing audio in sequencer now implements a real-time granular method.

diff --git a/Engine/Source/Runtime/AudioMixer/Classes/Generators/SoundWaveScrubber.h b/Engine/Source/Runtime/AudioMixer/Classes/Generators/SoundWaveScrubber.h
new file mode 100644
--- /dev/null
+++ b/Engine/Source/Runtime/AudioMixer/Classes/Generators/SoundWaveScrubber.h
@@ -0,0 +1,226 @@
+// Copyright Epic Games, Inc. All Rights Reserved.
+
+#pragma once
+
+#include "DSP/AlignedBuffer.h"
+#include "DSP/GrainEnvelope.h"
+#include "DSP/SampleRateConverter.h"
+#include "Sound/SoundBase.h"
+#include "Sound/SoundGenerator.h"
+#include "Sound/SoundWave.h"
+#include "Sound/SoundWaveProxyReader.h"
+#include "AudioDynamicParameter.h"
+
+#include "SoundWaveScrubber.generated.h"
+
+namespace Audio
+{
+
+// Class manages the runtime generation of scrubbed audio from a reference sound wave
+// It does this using a simple granulation technique of overlapping grains using a Hanning window.
+class AUDIOMIXER_API FSoundWaveScrubber
+{
+public:
+	FSoundWaveScrubber();
+	~FSoundWaveScrubber();
+
+	// Initialize the sound wave scrubber
+	void Init(FSoundWaveProxyPtr InSoundWaveProxyPtr, float InSampleRate, int32 InNumChannels, float InPlayheadTimeSeconds = 0.0f);
+
+	// Sets whether or not to scrub the audio file. If false, then the output of the scrubber will be silence.
+	void SetIsScrubbing(bool bInIsScrubbing);
+
+	// Sets the scrubbing location in seconds
+	void SetPlayheadTime(float InPlayheadTimeSeconds);
+
+	// Sets the grain duration range in seconds (clamped in range of of 0.01 to 0.5 seconds)
+	// The grain duration used during grain spawn is based the scrubbing speed. Sloweer the speed, the longer the grain.
+	void SetGrainDurationRange(const FVector2D& InGrainDurationRange);
+
+	// Renders the audio from the sound wave scrubber into the output audio buffer.
+	int32 RenderAudio(TArrayView<float>& OutAudio);
+
+private:
+
+	// ----------------
+	// General state and/or settings
+	float AudioMixerSampleRate = 0.0f;
+	float SourceFileSampleRate = 0.0f;
+	float SourceFileDurationSeconds = 0.0f;
+	int32 NumChannels = 0;
+
+	// How much audio to decode per decode block
+	static constexpr float DecodedAudioSizeInSeconds = 1.0f;
+
+	// The current target playhead time in seconds
+	// The scrubber will interpolate to this if the current playhead time is less than a given threshold,
+	// otherwise the scrubber will "jump" to the playhead time.
+	float TargetPlayheadTimeSeconds = 0.0f;
+
+	// Interpolated playhead time
+	FDynamicParameter CurrentPlayheadTimeSeconds;
+
+	// ----------------
+	// Data for managing the decoded audio
+	struct FDecodedDataChunk
+	{
+		// The start frame of the decoded audio chunk
+		int32 FrameStart = INDEX_NONE;
+
+		// The actual decoded audio. 
+		// Size is NumFrames * NumChannels * AudioMixerSampleRate (after SRC)
+		Audio::FAlignedFloatBuffer PCMAudio;
+
+		// Count of the number of grains actively using this chunk
+		int32 NumGrainsUsingChunk = 0;
+	};
+	// Place to store decoded chunks. We need 2 chunks for actively playing grains plus a free chunk.
+	TArray<FDecodedDataChunk> DecodedChunks;
+
+	// Sound wave proxy to safely query and use the referenced sound wave asset
+	FSoundWaveProxyPtr SoundWaveProxyPtr;
+	// Sound wave proxy reader is a decoder we can use to decode audio into chunks
+	TUniquePtr<FSoundWaveProxyReader> SoundWaveProxyReaderPtr;
+	// Simple SRC interface. Decoded audio is SRC'd to match the audio mixer sample rate
+	TUniquePtr<Audio::ISampleRateConverter> SRC;
+
+	// Function utility to decode audio from the current sate to the given decoded data chunk
+	void DecodeToDataChunk(FDecodedDataChunk& InOutDataChunk, float InDecoderSeekTimeSeconds);
+
+	// This retrieves decoded data for the given read frame index. 
+	// Returns INDEX_NONE if no decoded data was available for the given read index
+	int32 GetDecodedDataChunkIndexForCurrentReadIndex(int32 InReadFrameIndex);
+
+	// This function will decode the audio and cache it.
+	// Returns the index for the decoded chunk
+	int32 DecodeDataChunkIndexForCurrentReadIndex(int32 InReadFrameIndex);
+
+	// ----------------
+	// Data for granulation
+	// Envelope shared across grains
+	Audio::Grain::FEnvelope GrainEnvelope;
+
+	// The grain duration in frames
+	FVector2D TargetGrainDurationRange;
+	FVector2D GrainDurationRange;
+	float GrainDurationSeconds = 0.0f;
+
+	// Grain duration as computed by scrub velocity based on the grain duration ranges
+	int32 CurrentGrainDurationFrames = 0;
+	int32 CurrentHalfGrainDurationFrames = 0;
+
+	// Utility data useful for debugging the scrubber.
+	// The grain count is a running tally of the numebr of grains rendered
+	int32 GrainCount = 0;
+	// This should never go above 2
+	int32 NumActiveGrains = 0;
+
+	// Active grain data for a single grain
+	struct FGrain
+	{
+		// The number of frames this grain has rendered
+		// If this is larger than GrainDurationFrames it's an inactive grain
+		// this is used to look up the grain envelope to find the amplitude of the grain per frame
+		int32 CurrentRenderedFramesCount = INDEX_NONE;
+
+		// Index into the decoded data array
+		int32 DecodedDataChunkIndex = 0;
+
+		// The current read frame of the grain
+		int32 CurrentReadFrame = INDEX_NONE;
+
+		// The duration of this grain. Set when the grain spawns. 
+		int32 GrainDurationFrames = 0;
+	};
+	// The grain data used to render the granular audio
+	TArray<FGrain> ActiveGrains;
+
+	// This is a frame count until we need to spawn another grain
+	int32 NumFramesTillNextGrainSpawn = 0;
+
+	// Spawns a new grain at the value of the current playhead
+	FGrain SpawnGrain();
+
+	// Makes sure the grain is referencing valid decoded data
+	void UpdateGrainDecodeData(FGrain& InGrain);
+
+	// Renders the grain, mixing it into the output buffer. Returns the frame index that is half
+	void RenderActiveGrains(TArrayView<float>& OutAudio, int32 InStartFrame, int32 InNumFramesToRender);
+
+	// Used for setting user parameters
+	FCriticalSection CritSect;
+
+	// Whether or not we're actively scrubbing audio playback. 
+	bool bIsScrubbing = false;
+};
+
+
+class AUDIOMIXER_API FSoundWaveScrubberGenerator : public ISoundGenerator
+{
+public:
+	void Init(FSoundWaveProxyPtr InSoundWaveProxyPtr, float InSampleRate, int32 InNumChannels, float InPlayheadTimeSeconds = 0.0f);
+
+	// ~ISoundGenerator Begin
+	virtual int32 OnGenerateAudio(float* OutAudio, int32 NumSamples) override;
+	virtual int32 GetDesiredNumSamplesToRenderPerCallback() const override;
+	virtual bool IsFinished() const override;
+	// ~ISoundGenerator End
+
+	// Sets whether or not to scrub the audio file. If false, then the output of the scrubber will be silence.
+	void SetIsScrubbing(bool bInIsScrubbing);
+	// Sets the scrubbing location in seconds
+	void SetPlayheadTime(float InPlayheadTimeSeconds);
+
+	// Sets the grain duration range in seconds (clamped in range of of 0.01 to 0.5 seconds)
+	// The grain duration used during grain spawn is based the scrubbing speed. Sloweer the speed, the longer the grain.
+	void SetGrainDurationRange(const FVector2D& InGrainDurationRange);
+
+private:
+	int32 NumChannels = 0;
+	FSoundWaveScrubber SoundWaveScrubber;
+
+};
+
+}
+
+
+UCLASS(MinimalAPI)
+class UScrubbedSound : public USoundWave
+{
+	GENERATED_UCLASS_BODY()
+public:
+
+	// ~USoundBase
+	virtual ISoundGeneratorPtr CreateSoundGenerator(const FSoundGeneratorInitParams& InParams) override;
+	// ~USoundBase
+
+	ISoundGeneratorPtr GetSoundGenerator() { return SoundWaveScrubber; }
+	
+	// Sets the sound wave to scrub
+	UFUNCTION(BlueprintCallable, Category = "Scrubbing")
+	AUDIOMIXER_API void SetSoundWave(USoundWave* InSoundWave);
+	
+	// Sets the scrub time in seconds
+	UFUNCTION(BlueprintCallable, Category = "Scrubbing")
+	AUDIOMIXER_API void SetPlayheadTime(float InPlayheadTimeSeconds);
+
+	// Sets the scrub grain duration range.
+	UFUNCTION(BlueprintCallable, Category = "Scrubbing")
+	AUDIOMIXER_API void SetGrainDurationRange(const FVector2D& InGrainDurationRangeSeconds);
+
+	// Sets if the scrubber is actively scrubbing or not
+	UFUNCTION(BlueprintCallable, Category = "Scrubbing")
+	AUDIOMIXER_API void SetIsScrubbing(bool bInIsScrubbing);
+
+private:
+	float PlayheadTimeSeconds = 0.0f;
+	FVector2D GrainDurationRange = { 0.4f, 0.05f };
+
+	bool bIsScrubbing = false;
+
+	ISoundGeneratorPtr SoundWaveScrubber;
+
+	UPROPERTY(Transient)
+	TObjectPtr<USoundWave> SoundWaveToScrub;
+	
+};
diff --git a/Engine/Source/Runtime/AudioMixer/Private/AudioMixerSourceManager.h b/Engine/Source/Runtime/AudioMixer/Private/AudioMixerSourceManager.h
--- a/Engine/Source/Runtime/AudioMixer/Private/AudioMixerSourceManager.h
+++ b/Engine/Source/Runtime/AudioMixer/Private/AudioMixerSourceManager.h
@@ -43,7 +43,7 @@ DECLARE_CYCLE_STAT_EXTERN(TEXT("HRTF"), STAT_AudioMixerHRTF, STATGROUP_AudioMixe
 
 // For diagnostics, keep track of what phase of updating the Source manager is in currently.
 UENUM()
-enum ESourceManagerRenderThreadPhase: uint8
+enum class ESourceManagerRenderThreadPhase: uint8
 {
 	Begin,
 	
diff --git a/Engine/Source/Runtime/AudioMixer/Private/Generators/SoundWaveScrubber.cpp b/Engine/Source/Runtime/AudioMixer/Private/Generators/SoundWaveScrubber.cpp
new file mode 100644
--- /dev/null
+++ b/Engine/Source/Runtime/AudioMixer/Private/Generators/SoundWaveScrubber.cpp
@@ -0,0 +1,509 @@
+// Copyright Epic Games, Inc. All Rights Reserved.
+
+#include "Generators/SoundWaveScrubber.h"
+#include "ProfilingDebugging/CpuProfilerTrace.h"
+
+namespace Audio
+{ 
+
+FSoundWaveScrubber::FSoundWaveScrubber()
+	: CurrentPlayheadTimeSeconds(0.0f)
+	, SRC(Audio::ISampleRateConverter::CreateSampleRateConverter())
+	, GrainDurationSeconds(0.0f)
+{
+	// Generate the grain envelope data
+	Audio::Grain::GenerateEnvelopeData(GrainEnvelope, 512, Audio::Grain::EEnvelope::Hann);
+}
+
+FSoundWaveScrubber::~FSoundWaveScrubber()
+{
+}
+
+void FSoundWaveScrubber::Init(FSoundWaveProxyPtr InSoundWaveProxyPtr, float InSampleRate, int32 InNumChannels, float InPlayheadTimeSeconds)
+{
+	check(InSoundWaveProxyPtr.IsValid());
+	SoundWaveProxyPtr = InSoundWaveProxyPtr;
+
+	AudioMixerSampleRate = InSampleRate;
+	SourceFileSampleRate = SoundWaveProxyPtr->GetSampleRate();
+	SourceFileDurationSeconds = SoundWaveProxyPtr->GetDuration();
+	NumChannels = InNumChannels;
+
+	TargetGrainDurationRange = { 0.4f, 0.05f };
+	GrainDurationRange = TargetGrainDurationRange;
+	GrainDurationSeconds = TargetGrainDurationRange.X;
+
+	CurrentPlayheadTimeSeconds.Set(InPlayheadTimeSeconds, 0.0f);
+	TargetPlayheadTimeSeconds = CurrentPlayheadTimeSeconds.GetValue();
+
+	// We need 3 slots for decoded data. 2 for the grain playback, 1 to decode new chunks while grains are playing.
+	DecodedChunks.Reset();
+	DecodedChunks.AddDefaulted(3);
+
+	float DecoderSeekTimeSeconds = FMath::Max(CurrentPlayheadTimeSeconds.GetValue() - 0.5f * DecodedAudioSizeInSeconds, 0.0f);
+
+	DecodeToDataChunk(DecodedChunks[0], DecoderSeekTimeSeconds);
+}
+
+int32 FSoundWaveScrubber::GetDecodedDataChunkIndexForCurrentReadIndex(int32 InReadFrameIndex)
+{
+	for (int32 i = 0; i < DecodedChunks.Num(); ++i)
+	{
+		const FDecodedDataChunk& DecodedChunk = DecodedChunks[i];
+		if (DecodedChunk.PCMAudio.Num() > 0)
+		{
+			int32 DecodedFrameCount = DecodedChunk.PCMAudio.Num() / NumChannels;
+			if (InReadFrameIndex >= DecodedChunk.FrameStart && InReadFrameIndex < DecodedChunk.FrameStart + DecodedFrameCount)
+			{
+				// We found a decoded audio chunk that contains the desired read frame index
+				return i;
+			}
+		}
+	}
+	return INDEX_NONE;
+}
+
+int32 FSoundWaveScrubber::DecodeDataChunkIndexForCurrentReadIndex(int32 InReadFrameIndex)
+{
+	// No decoded chunk was found for desired read frame index. 
+	// This indicates that we need to decode more audio.
+	for (int32 i = 0; i < DecodedChunks.Num(); ++i)
+	{
+		FDecodedDataChunk& DecodedChunk = DecodedChunks[i];
+		if (!DecodedChunk.NumGrainsUsingChunk)
+		{
+			float DecoderSeekTimeSeconds = InReadFrameIndex / AudioMixerSampleRate;
+			DecodeToDataChunk(DecodedChunk, DecoderSeekTimeSeconds);
+			check(DecodedChunk.PCMAudio.Num() > 0);
+			return i;
+		}
+	}
+
+	FDecodedDataChunk NewChunk;
+	float DecoderSeekTimeSeconds = InReadFrameIndex / AudioMixerSampleRate;
+	DecodeToDataChunk(NewChunk, DecoderSeekTimeSeconds);
+	check(NewChunk.PCMAudio.Num() > 0);
+	DecodedChunks.Add(MoveTemp(NewChunk));
+	return DecodedChunks.Num() - 1;
+}
+
+
+void FSoundWaveScrubber::DecodeToDataChunk(FDecodedDataChunk& InOutDataChunk, float InDecoderSeekTimeSeconds)
+{
+	check(InOutDataChunk.NumGrainsUsingChunk == 0);
+	check(DecodedAudioSizeInSeconds > 0.0f);
+	check(SourceFileSampleRate > 0.0f);
+	check(InDecoderSeekTimeSeconds >= 0.0f);
+	check(NumChannels > 0);
+	check(SoundWaveProxyPtr.IsValid());
+
+	if (!SoundWaveProxyReaderPtr.IsValid())
+	{
+		// Create the proxy reader, which is our decoder. Initialize it at the initial playhead time
+		FSoundWaveProxyReader::FSettings ProxyReaderSettings;
+		ProxyReaderSettings.MaxDecodeSizeInFrames = DecodedAudioSizeInSeconds * SourceFileSampleRate;
+		ProxyReaderSettings.StartTimeInSeconds = InDecoderSeekTimeSeconds;
+		SoundWaveProxyReaderPtr = FSoundWaveProxyReader::Create(SoundWaveProxyPtr.ToSharedRef(), ProxyReaderSettings);
+	}
+	else
+	{
+		check(SoundWaveProxyReaderPtr.IsValid());
+		float DecoderSeekTimeSeconds = InDecoderSeekTimeSeconds;
+		// If proxy reader already exists, then simply seek the decoder to the desired location
+		SoundWaveProxyReaderPtr->SeekToTime(DecoderSeekTimeSeconds);
+	}
+
+	InOutDataChunk.FrameStart = InDecoderSeekTimeSeconds * AudioMixerSampleRate;
+
+	// We allocate the size of buffer pre-SRC based on source file sample rate
+	int32 DecodedAudioSize = DecodedAudioSizeInSeconds * SourceFileSampleRate * NumChannels;
+	check(DecodedAudioSize > 0);
+	InOutDataChunk.PCMAudio.Reset();
+	InOutDataChunk.PCMAudio.AddUninitialized(DecodedAudioSize);
+
+	// This does the actual decoding of the audio to match the size of the input buffer
+	SoundWaveProxyReaderPtr->PopAudio(InOutDataChunk.PCMAudio);
+
+	// Check if we need to do SRC. If we do, this will change the allocated 
+	// size (potentially expanding or shrinking) to match 
+	// the audio mixer sample rate.
+	if (!FMath::IsNearlyEqual(SourceFileSampleRate, AudioMixerSampleRate))
+	{
+		SRC->Init((float)SourceFileSampleRate / AudioMixerSampleRate, NumChannels);
+		TArray<float> SampleRateConvertedPCM;
+		SRC->ProcessFullbuffer(InOutDataChunk.PCMAudio.GetData(), InOutDataChunk.PCMAudio.Num(), SampleRateConvertedPCM);
+		InOutDataChunk.PCMAudio = MoveTemp(SampleRateConvertedPCM);
+	}
+}
+
+FSoundWaveScrubber::FGrain FSoundWaveScrubber::SpawnGrain()
+{
+	// Try to retrieve a decoded data chunk for the current read frame based on the curret, interpolated playhead time seconds value
+	int32 CurrentReadFrame = CurrentPlayheadTimeSeconds.GetValue() * AudioMixerSampleRate;
+	int32 DecodedDataChunkIndex = GetDecodedDataChunkIndexForCurrentReadIndex(CurrentReadFrame);
+	if (DecodedDataChunkIndex == INDEX_NONE)
+	{
+		DecodedDataChunkIndex = DecodeDataChunkIndexForCurrentReadIndex(CurrentReadFrame);
+	}
+	check(DecodedDataChunkIndex != INDEX_NONE);
+
+	// Get the grain runtime data for the new grain spawn
+	FGrain NewGrain;
+	NewGrain.CurrentRenderedFramesCount = 0;
+	NewGrain.DecodedDataChunkIndex = DecodedDataChunkIndex;
+	NewGrain.CurrentReadFrame = CurrentReadFrame;
+	NewGrain.GrainDurationFrames = CurrentGrainDurationFrames;
+	DecodedChunks[DecodedDataChunkIndex].NumGrainsUsingChunk++;
+
+	GrainCount++;
+	NumActiveGrains++;
+
+	return NewGrain;
+}
+
+void FSoundWaveScrubber::SetIsScrubbing(bool bInIsScrubbing)
+{
+	bIsScrubbing = bInIsScrubbing;
+}
+
+void FSoundWaveScrubber::SetPlayheadTime(float InPlayheadTimeSeconds)
+{
+	FScopeLock Lock(&CritSect);
+
+	TargetPlayheadTimeSeconds = FMath::Fmod(FMath::Max(InPlayheadTimeSeconds, 0.0f), SourceFileDurationSeconds);
+}
+
+void FSoundWaveScrubber::SetGrainDurationRange(const FVector2D& InGrainDurationRange)
+{
+	FVector2D GrainDurationRangeClamped =
+	{
+		FMath::Clamp(InGrainDurationRange.X, 0.05f, 0.5f),
+		FMath::Clamp(InGrainDurationRange.Y, 0.05f, 0.5f),
+	};
+
+	FScopeLock Lock(&CritSect);
+	TargetGrainDurationRange = GrainDurationRangeClamped;
+}
+
+int32 FSoundWaveScrubber::RenderAudio(TArrayView<float>& OutAudio)
+{
+	SCOPED_NAMED_EVENT_TEXT("FSoundWaveScrubber::RenderAudio", FColor::Emerald);
+
+	// Number of frames of this generate audio
+	int32 NumFrames = OutAudio.Num() / NumChannels;
+	float DeltaTimeSecond = (float)NumFrames / AudioMixerSampleRate;
+
+	float PlayheadTimeDistanceSeconds = 0.0f;
+
+	constexpr float PlayheadLerpTime = 0.2f;
+
+	// Update parameters
+	{
+		FScopeLock Lock(&CritSect);
+
+		// Update the current playhead time seconds
+		if (!FMath::IsNearlyEqual(TargetPlayheadTimeSeconds, CurrentPlayheadTimeSeconds.GetTargetValue()))
+		{
+			float PlayheadTimeDelta = FMath::Abs(CurrentPlayheadTimeSeconds.GetValue() - TargetPlayheadTimeSeconds);
+			// If the playhead time jumps suddenly, we'll instantly set the current playhead to the target
+			if (PlayheadTimeDelta > 0.5f)
+			{
+				CurrentPlayheadTimeSeconds.Set(TargetPlayheadTimeSeconds, 0.0f);
+			}
+			else
+			{
+				CurrentPlayheadTimeSeconds.Set(TargetPlayheadTimeSeconds, PlayheadLerpTime);
+			}
+		}
+		CurrentPlayheadTimeSeconds.Update(DeltaTimeSecond);
+
+		PlayheadTimeDistanceSeconds = FMath::Abs(CurrentPlayheadTimeSeconds.GetTargetValue() - CurrentPlayheadTimeSeconds.GetValue());
+
+		GrainDurationRange = TargetGrainDurationRange;
+	}
+
+	// This maps the playhead delta from the target playhead time (which is an indirect measure of playhead velocity) to the duration range
+	float MappedGrainDurationRange = FMath::GetMappedRangeValueClamped({ 0.0f, PlayheadLerpTime }, GrainDurationRange, PlayheadTimeDistanceSeconds);
+	GrainDurationSeconds = MappedGrainDurationRange;
+
+	// Update the grain duration based on a mapping between the duration range and scrub velocity
+	CurrentGrainDurationFrames = GrainDurationSeconds * AudioMixerSampleRate;
+	CurrentHalfGrainDurationFrames = 0.5f * CurrentGrainDurationFrames;
+
+	// If we're actively scrubbing we need to spawn grains and render the granular audio
+	if (bIsScrubbing)
+	{
+		if (!ActiveGrains.Num())
+		{
+			ActiveGrains.Add(SpawnGrain());
+
+			NumFramesTillNextGrainSpawn = CurrentHalfGrainDurationFrames;
+		}
+
+		int32 StartRenderFrame = 0;
+		int32 NumFramesToRender = FMath::Min(NumFramesTillNextGrainSpawn, NumFrames);
+		int32 NumFramesRendered = 0;
+		int32 RemainingFrames = NumFrames;
+
+		while (RemainingFrames > 0)
+		{
+			// This renders the currently active grains (which should only ever max 2 grains)
+			// starting from the given start render frame and for the  numbe of indicated frames
+			RenderActiveGrains(OutAudio, StartRenderFrame, NumFramesToRender);
+
+			// Update the number of frames rendered this render block
+			NumFramesRendered += NumFramesToRender;
+
+			NumFramesTillNextGrainSpawn -= NumFramesToRender;
+			check(NumFramesTillNextGrainSpawn >= 0);
+
+			// Determine how many more frames, this block, that we need to render
+			RemainingFrames = FMath::Max(NumFrames - NumFramesRendered, 0);
+
+			// Check if we need to spawn a new grain
+			if (NumFramesTillNextGrainSpawn == 0)
+			{
+				// Spawn a new grain with a frame start 
+				StartRenderFrame = NumFrames - RemainingFrames;
+				ActiveGrains.Add(SpawnGrain());
+
+				// If we still have frames remaining in the buffer, then 
+				NumFramesTillNextGrainSpawn = CurrentHalfGrainDurationFrames;
+
+				NumFramesToRender = FMath::Min(NumFramesTillNextGrainSpawn, RemainingFrames);
+			}
+		}
+	}
+
+	return OutAudio.Num();
+}
+
+void FSoundWaveScrubber::UpdateGrainDecodeData(FGrain& InGrain)
+{
+	FDecodedDataChunk& DecodedData = DecodedChunks[InGrain.DecodedDataChunkIndex];
+
+	// Total number of frames of decoded data in the chunk
+	int32 NumReadFramesInDecodedData = DecodedData.PCMAudio.Num() / NumChannels;
+
+	// The number of frames that this grain is offset from the decoded data
+	int32 NumFramesOffsetInDecodedData = InGrain.CurrentReadFrame - DecodedData.FrameStart;
+	int32 NumFramesPossibleToRenderInChunk = NumReadFramesInDecodedData - NumFramesOffsetInDecodedData;
+	check(NumFramesPossibleToRenderInChunk >= 0);
+
+	// If we've totally consumed this decoded audio chunk, we need to get a new decoded audio chunk
+	if (!NumFramesPossibleToRenderInChunk)
+	{
+		// We're no longer using this decoded audio chunk
+		DecodedData.NumGrainsUsingChunk--;
+		check(DecodedData.NumGrainsUsingChunk >= 0);
+
+		InGrain.DecodedDataChunkIndex = GetDecodedDataChunkIndexForCurrentReadIndex(InGrain.CurrentReadFrame);
+		if (InGrain.DecodedDataChunkIndex == INDEX_NONE)
+		{
+			InGrain.DecodedDataChunkIndex = DecodeDataChunkIndexForCurrentReadIndex(InGrain.CurrentReadFrame);
+		}
+		check(InGrain.DecodedDataChunkIndex != INDEX_NONE);
+
+		check(InGrain.CurrentReadFrame >= DecodedChunks[InGrain.DecodedDataChunkIndex].FrameStart);
+		check(InGrain.CurrentReadFrame < DecodedChunks[InGrain.DecodedDataChunkIndex].FrameStart + DecodedChunks[InGrain.DecodedDataChunkIndex].PCMAudio.Num() / NumChannels);
+
+		DecodedChunks[InGrain.DecodedDataChunkIndex].NumGrainsUsingChunk++;
+	}
+}
+
+void FSoundWaveScrubber::RenderActiveGrains(TArrayView<float>& OutAudio, int32 InStartFrame, int32 NumFramesToRender)
+{
+	check(NumChannels > 0);
+
+	// Reverse iterate so we can quickly remove grains when they're done
+	for (int32 GrainIndex = ActiveGrains.Num() - 1; GrainIndex >= 0; --GrainIndex)
+	{
+		FGrain& Grain = ActiveGrains[GrainIndex];
+
+		// This is the number of frames we have left to render for this grain
+		int32 NumFramesLeftInGrain = Grain.GrainDurationFrames - Grain.CurrentRenderedFramesCount;
+		int32 NumFramesLeftToRender = FMath::Min(NumFramesToRender, NumFramesLeftInGrain);
+		check(NumFramesLeftToRender > 0);
+
+		int32 GrainWriteIndex = InStartFrame;
+
+		bool bGrainFinished = false;
+		while (NumFramesLeftToRender > 0 && !bGrainFinished)
+		{
+			// Make sure we have a valid decode data chunk ready for rendering
+			UpdateGrainDecodeData(Grain);
+
+			// Retrieve the decoded data.
+			const FDecodedDataChunk& DecodedData = DecodedChunks[Grain.DecodedDataChunkIndex];
+
+			// Total number of frames of decoded data in the chunk
+			int32 NumReadFramesInDecodedData = DecodedData.PCMAudio.Num() / NumChannels;
+
+			// The number of frames that this grain is offset from the decoded data
+			int32 NumFramesOffsetInDecodedData = Grain.CurrentReadFrame - DecodedData.FrameStart;
+			check(NumFramesOffsetInDecodedData >= 0);
+
+			int32 NumFramesPossibleToRenderInChunk = NumReadFramesInDecodedData - NumFramesOffsetInDecodedData;
+			check(NumFramesPossibleToRenderInChunk >= 0);
+
+			// This may completely consume the decoded chunk here, so only render the maximum number of frames in the 
+			// chunk. if NumFramesToRender is smaller than that, then we won't completely consume the decoded data
+			int32 NumFramesToRenderInThisChunk = FMath::Min(NumFramesPossibleToRenderInChunk, NumFramesLeftToRender);
+			check(NumFramesToRenderInThisChunk >= 0);
+
+			int32 SampleWriteIndex = GrainWriteIndex * NumChannels;
+			int32 SampleReadIndex = NumFramesOffsetInDecodedData * NumChannels;
+
+			// For inner-loop audio rendering, we avoid the constant array-access checks that happen on our containers
+			const float* DecodedPCMDataPtr = DecodedData.PCMAudio.GetData();
+			float* OutDataPtr = OutAudio.GetData();
+
+			for (int32 FrameIndex = 0; FrameIndex < NumFramesToRenderInThisChunk; ++FrameIndex)
+			{
+				// Retrieve the grain amplitude from the envelope
+				// Note the envelope is intentionally sized to match the size of the grain so we don't have
+				// to do any interpolation math on look up.
+				float EnvelopeFraction = FMath::Clamp((float)Grain.CurrentRenderedFramesCount++ / Grain.GrainDurationFrames, 0.0f, 1.0f);
+				float GrainAmplitude = Audio::Grain::GetValue(GrainEnvelope, EnvelopeFraction);
+
+				for (int32 ChannelIndex = 0; ChannelIndex < NumChannels; ++ChannelIndex)
+				{
+					// Read the decoded sample of the audio at this channel index
+					float DecodedSampleValue = DecodedPCMDataPtr[SampleReadIndex++];
+
+					// Scale the sample value by the grain amplitude
+					DecodedSampleValue *= GrainAmplitude;
+
+					// Mix the grain audio into the output buffer
+					OutDataPtr[SampleWriteIndex++] += DecodedSampleValue;
+				}
+			}
+
+			Grain.CurrentReadFrame += NumFramesToRenderInThisChunk;
+
+			GrainWriteIndex += NumFramesToRenderInThisChunk;
+			NumFramesLeftToRender -= NumFramesToRenderInThisChunk;
+
+			bGrainFinished = (Grain.CurrentRenderedFramesCount >= Grain.GrainDurationFrames);
+		}
+
+		// If the grain has finished, then remove it from the active grain list
+		if (bGrainFinished)
+		{
+			DecodedChunks[Grain.DecodedDataChunkIndex].NumGrainsUsingChunk--;
+			check(DecodedChunks[Grain.DecodedDataChunkIndex].NumGrainsUsingChunk >= 0);
+			ActiveGrains.RemoveAtSwap(GrainIndex, EAllowShrinking::No);
+			NumActiveGrains--;
+		}
+	}
+}
+
+void FSoundWaveScrubberGenerator::Init(FSoundWaveProxyPtr InSoundWaveProxyPtr, float InSampleRate, int32 InNumChannels, float InPlayheadTimeSeconds)
+{
+	NumChannels = InNumChannels;
+	SoundWaveScrubber.Init(InSoundWaveProxyPtr, InSampleRate, InNumChannels, InPlayheadTimeSeconds);
+}
+
+void FSoundWaveScrubberGenerator::SetIsScrubbing(bool bInIsScrubbing)
+{
+	SoundWaveScrubber.SetIsScrubbing(bInIsScrubbing);
+}
+
+void FSoundWaveScrubberGenerator::SetPlayheadTime(float InPlayheadTimeSeconds)
+{
+	SoundWaveScrubber.SetPlayheadTime(InPlayheadTimeSeconds);
+}
+
+void FSoundWaveScrubberGenerator::SetGrainDurationRange(const FVector2D& InGrainDurationRange)
+{
+	SoundWaveScrubber.SetGrainDurationRange(InGrainDurationRange);
+}
+
+int32 FSoundWaveScrubberGenerator::OnGenerateAudio(float* OutAudio, int32 NumSamples)
+{
+	FMemory::Memzero(OutAudio, NumSamples*sizeof(float));
+	TArrayView<float> OutView = TArrayView<float>(OutAudio, NumSamples);
+	return SoundWaveScrubber.RenderAudio(OutView);
+}
+
+int32 FSoundWaveScrubberGenerator::GetDesiredNumSamplesToRenderPerCallback() const
+{
+	return 256 * NumChannels;
+}
+
+bool FSoundWaveScrubberGenerator::IsFinished() const
+{
+	// This is intended to be "always on" unless stopped by owning audio component, etc.
+	return false;
+}
+
+} // namespace Audio
+
+UScrubbedSound::UScrubbedSound(const FObjectInitializer& ObjectInitializer)
+	: Super(ObjectInitializer)
+{
+	bProcedural = true;
+}
+
+ISoundGeneratorPtr UScrubbedSound::CreateSoundGenerator(const FSoundGeneratorInitParams& InParams)
+{
+	using namespace Audio;
+
+	if (SoundWaveToScrub)
+	{
+		FSoundWaveProxyPtr SoundWaveProxyPtr = SoundWaveToScrub->CreateSoundWaveProxy();
+
+		SoundWaveScrubber = ISoundGeneratorPtr(new FSoundWaveScrubberGenerator());
+
+		FSoundWaveScrubberGenerator* Scrubber = static_cast<FSoundWaveScrubberGenerator*>(SoundWaveScrubber.Get());
+		Scrubber->Init(SoundWaveProxyPtr, InParams.SampleRate, NumChannels, PlayheadTimeSeconds);
+		Scrubber->SetIsScrubbing(bIsScrubbing);
+		Scrubber->SetGrainDurationRange(GrainDurationRange);
+
+		return SoundWaveScrubber;
+	}
+
+	return nullptr;
+}
+
+void UScrubbedSound::SetSoundWave(USoundWave* InSoundWave)
+{
+	 SoundWaveToScrub = InSoundWave;
+	 NumChannels = InSoundWave->NumChannels;
+}
+
+void UScrubbedSound::SetIsScrubbing(bool bInIsScrubbing)
+{
+	using namespace Audio;
+
+	bIsScrubbing = bInIsScrubbing;
+	if (SoundWaveScrubber.IsValid())
+	{
+		FSoundWaveScrubberGenerator* Scrubber = static_cast<FSoundWaveScrubberGenerator*>(SoundWaveScrubber.Get());
+		Scrubber->SetIsScrubbing(bIsScrubbing);
+	}
+}
+
+void UScrubbedSound::SetPlayheadTime(float InPlayheadTimeSeconds)
+{
+	using namespace Audio;
+
+	PlayheadTimeSeconds = InPlayheadTimeSeconds;
+	if (SoundWaveScrubber.IsValid())
+	{
+		FSoundWaveScrubberGenerator* Scrubber = static_cast<FSoundWaveScrubberGenerator*>(SoundWaveScrubber.Get());
+		Scrubber->SetPlayheadTime(InPlayheadTimeSeconds);
+	}
+}
+
+void UScrubbedSound::SetGrainDurationRange(const FVector2D& InGrainDurationRangeSeconds)
+{
+	using namespace Audio;
+
+	GrainDurationRange = InGrainDurationRangeSeconds;
+	if (SoundWaveScrubber.IsValid())
+	{
+		FSoundWaveScrubberGenerator* Scrubber = static_cast<FSoundWaveScrubberGenerator*>(SoundWaveScrubber.Get());
+		Scrubber->SetGrainDurationRange(GrainDurationRange);
+	}
+}
\ No newline at end of file
diff --git a/Engine/Source/Runtime/Engine/Classes/Sound/SoundWaveProxyReader.h b/Engine/Source/Runtime/Engine/Classes/Sound/SoundWaveProxyReader.h
--- a/Engine/Source/Runtime/Engine/Classes/Sound/SoundWaveProxyReader.h
+++ b/Engine/Source/Runtime/Engine/Classes/Sound/SoundWaveProxyReader.h
@@ -165,7 +165,7 @@ public:
 
 private:
 
-	enum EDecodeResult
+	enum class EDecodeResult
 	{
 		Fail,					// Decoder failed somehow.
 		MoreDataRemaining,		// Data has produced and there's more remaining
diff --git a/Engine/Source/Runtime/MovieSceneTracks/MovieSceneTracks.Build.cs b/Engine/Source/Runtime/MovieSceneTracks/MovieSceneTracks.Build.cs
--- a/Engine/Source/Runtime/MovieSceneTracks/MovieSceneTracks.Build.cs
+++ b/Engine/Source/Runtime/MovieSceneTracks/MovieSceneTracks.Build.cs
@@ -16,7 +16,7 @@ public class MovieSceneTracks : ModuleRules
 				"AnimationCore",
 				"AudioExtensions",
 				"Constraints",
-
+				"AudioMixer"
 			});
 
 		PrivateDependencyModuleNames.AddRange(new string[]
diff --git a/Engine/Source/Runtime/MovieSceneTracks/Private/Systems/MovieSceneAudioSystem.cpp b/Engine/Source/Runtime/MovieSceneTracks/Private/Systems/MovieSceneAudioSystem.cpp
--- a/Engine/Source/Runtime/MovieSceneTracks/Private/Systems/MovieSceneAudioSystem.cpp
+++ b/Engine/Source/Runtime/MovieSceneTracks/Private/Systems/MovieSceneAudioSystem.cpp
@@ -60,6 +60,14 @@ FAutoConsoleVariableRef CVarUseTimeDilationToAdjustPlayDuration(
 	TEXT("Use the effective time dilation to scale the current time of audio.\n"),
 	ECVF_Default);
 
+
+static int32 ScrubWidthMillisecondsCVar = 80;
+FAutoConsoleVariableRef CVarScrubWidthMilliseconds(
+	TEXT("Sequencer.Audio.ScrubWidthMilliseconds"),
+	ScrubWidthMillisecondsCVar,
+	TEXT("The time-width of grains (in milliseconds) while scrubbing an audio track.\n"),
+	ECVF_Default);
+
 namespace UE::MovieScene
 {
 
@@ -484,6 +492,7 @@ private:
 		
 		UWorld* World = PlaybackContext ? PlaybackContext->GetWorld() : nullptr;
 		AWorldSettings* WorldSettings = World ? World->GetWorldSettings() : nullptr;
+		UScrubbedSound* ScrubbedSound = EvaluationData.ScrubbedSound.Get();
 
 		AudioSystem->PreAnimatedStorage->CachePreAnimatedValue(
 				FCachePreAnimatedValueParams(), &AudioComponent,
@@ -673,20 +682,26 @@ private:
 				AudioComponent.AttenuationSettings = AudioSection.GetAttenuationSettings();
 			}
 
-			// Only call stop on the sound if it is actually playing. This prevents spamming
-			// stop calls when a sound cue with a duration of zero is played.
-			if (AudioComponent.IsPlaying() || bSoundNeedsTimeSync)
+			// If our sound is currently the scrubbed sound, that means we're actively scrubbing
+			// So we don't need to stop or set the sound again
+			if (Context.GetStatus() != EMovieScenePlayerStatus::Scrubbing)
 			{
-				UE_LOG(LogMovieScene, Verbose, TEXT("Audio Component stopped due to needing a state change bIsPlaying: %d bNeedsTimeSync: %d. Component: %s Sound: %s"), AudioComponent.IsPlaying(), bSoundNeedsTimeSync, *AudioComponent.GetName(), *GetNameSafe(AudioComponent.Sound));
-				AudioComponent.Stop();
-			}
+				// Only call stop on the sound if it is actually playing. This prevents spamming
+				// stop calls when a sound cue with a duration of zero is played.
+				if (AudioComponent.IsPlaying() || bSoundNeedsTimeSync)
+				{
+					UE_LOG(LogMovieScene, Verbose, TEXT("Audio Component stopped due to needing a state change bIsPlaying: %d bNeedsTimeSync: %d. Component: %s Sound: %s"), AudioComponent.IsPlaying(), bSoundNeedsTimeSync, *AudioComponent.GetName(), *GetNameSafe(AudioComponent.Sound));
+					AudioComponent.Stop();
+				}
 
-			// Only change the sound clip if it has actually changed. This calls Stop internally if needed.
-			if (AudioComponent.Sound != Sound)
-			{
-				UE_LOG(LogMovieScene, Verbose, TEXT("Audio Component calling SetSound due to new sound. Component: %s OldSound: %s NewSound: %s"), *AudioComponent.GetName(), *GetNameSafe(AudioComponent.Sound), *GetNameSafe(Sound));
-				AudioComponent.SetSound(Sound);
+				// Only change the sound clip if it has actually changed. This calls Stop internally if needed.
+				if (AudioComponent.Sound != Sound)
+				{
+					UE_LOG(LogMovieScene, Verbose, TEXT("Audio Component calling SetSound due to new sound. Component: %s OldSound: %s NewSound: %s"), *AudioComponent.GetName(), *GetNameSafe(AudioComponent.Sound), *GetNameSafe(Sound));
+					AudioComponent.SetSound(Sound);
+				}
 			}
+
 #if WITH_EDITOR
 			if (GIsEditor && World != nullptr && !World->IsPlayInEditor())
 			{
@@ -702,35 +717,68 @@ private:
 			if (AudioTime >= 0.f)
 			{
 				UE_LOG(LogMovieScene, Verbose, TEXT("Audio Component Play at Local Time: %6.2f CurrentTime: %6.2f(s) SectionStart: %6.2f(s), SoundDur: %6.2f OffsetIntoClip: %6.2f Component: %s Sound: %s"), AudioTime, (EvalTime / Context.GetFrameRate()), SectionStartTimeSeconds, AudioComponent.Sound ? AudioComponent.Sound->GetDuration() : 0.0f, (float)Context.GetFrameRate().AsSeconds(AudioStartOffset), *AudioComponent.GetName(), *GetNameSafe(AudioComponent.Sound));
-				AudioComponent.Play(AudioTime);
-
-				// Keep track of when we asked this audio clip to play (in game time) so that we can figure 
-				// out if there's a significant desync in the future.
-				//
-				// The goal is later to compare:
-				//   (NewAudioTime - PreviousAudioTime) and 
-				//   (NewGameTime - PreviousGameTime)
-				//
-				// If their difference is larger than some threshold, we have a desync. NewGameTime and 
-				// NewAudioTime will be known next update, but PreviousGameTime and PreviousAudioTime
-				// are known now. Let's store (-PreviousAudioTime + PreviousGameTime), and we will only 
-				// need to add (NewAudioTime - NewGameTime).
-				if (World)
+				
+				if (Context.GetStatus() == EMovieScenePlayerStatus::Scrubbing)
 				{
-					FAudioDevice* AudioDevice = World->GetAudioDeviceRaw();
-					if (UseAudioClockForSequencerDesyncCVar && AudioDevice)
+					if (USoundWave* SoundWave = Cast<USoundWave>(Sound))
 					{
-						EvaluationData.PartialDesyncComputation = AudioDevice->GetAudioClock() - AudioTime;
+						if (ScrubbedSound)
+						{
+							// If we're not playing the audio component, then set up the audio component to use the scrubbed sound wave and play it
+							if (!AudioComponent.IsPlaying())
+							{
+								ScrubbedSound->SetSoundWave(SoundWave);
+								ScrubbedSound->SetPlayheadTime(AudioTime);
+								float MaxScrubWidthSeconds = 0.001f * (float)ScrubWidthMillisecondsCVar;
+								ScrubbedSound->SetGrainDurationRange({MaxScrubWidthSeconds, 0.05f});
+
+								ScrubbedSound->SetIsScrubbing(true);
+
+								AudioComponent.SetSound(ScrubbedSound);
+								AudioComponent.Play();
+							}
+							else
+							{
+								// If we're already playing, then just update the playhead time on the scrubbed sound
+								// This will propagate the playhead time to the rendering ISoundGenerator
+								ScrubbedSound->SetPlayheadTime(AudioTime);
+							}
+						}
 					}
-					else
+				}
+				else
+				{
+					ScrubbedSound->SetIsScrubbing(false);
+					AudioComponent.Play(AudioTime);
+
+					// Keep track of when we asked this audio clip to play (in game time) so that we can figure 
+					// out if there's a significant desync in the future.
+					//
+					// The goal is later to compare:
+					//   (NewAudioTime - PreviousAudioTime) and 
+					//   (NewGameTime - PreviousGameTime)
+					//
+					// If their difference is larger than some threshold, we have a desync. NewGameTime and 
+					// NewAudioTime will be known next update, but PreviousGameTime and PreviousAudioTime
+					// are known now. Let's store (-PreviousAudioTime + PreviousGameTime), and we will only 
+					// need to add (NewAudioTime - NewGameTime).
+					if (World)
 					{
-						EvaluationData.PartialDesyncComputation = World->GetAudioTimeSeconds() - AudioTime;
+						FAudioDevice* AudioDevice = World->GetAudioDeviceRaw();
+						if (UseAudioClockForSequencerDesyncCVar && AudioDevice)
+						{
+							EvaluationData.PartialDesyncComputation = AudioDevice->GetInterpolatedAudioClock() - AudioTime;
+						}
+						else
+						{
+							EvaluationData.PartialDesyncComputation = World->GetAudioTimeSeconds() - AudioTime;
+						}
 					}
 				}
 			}
 		}
 
-		if (Context.GetStatus() != EMovieScenePlayerStatus::Playing)
+		if (Context.GetStatus() == EMovieScenePlayerStatus::Stepping || Context.GetStatus() == EMovieScenePlayerStatus::Jumping)
 		{
 			float ScrubDuration = AudioTrackConstants::ScrubDuration;
 			if (FAudioDevice* AudioDevice = AudioComponent.GetAudioDevice())
@@ -1062,8 +1110,10 @@ UMovieSceneAudioSystem::FAudioComponentEvaluationData* UMovieSceneAudioSystem::A
 
 		FAudioDevice::FCreateComponentParams Params(Actor->GetWorld(), Actor);
 		UAudioComponent* NewComponent = FAudioDevice::CreateComponent(TempPlaybackAudioCue, Params);
+		static int32 ScrubSoundCounter = 0;
+		UScrubbedSound* ScrubbedSound = NewObject<UScrubbedSound>(Actor, FName(*FString::Printf(TEXT("ScrubbedSound_Bound_%i"), ScrubSoundCounter++)));
 
-		if (!NewComponent)
+		if (!NewComponent || !ScrubbedSound)
 		{
 			const int32 RowIndex = Section->GetRowIndex();
 			UE_LOG(LogMovieScene, Warning, TEXT("Failed to create audio component for spatialized audio track (row %d on %s)."), RowIndex, *ObjectName);
@@ -1073,9 +1123,12 @@ UMovieSceneAudioSystem::FAudioComponentEvaluationData* UMovieSceneAudioSystem::A
 		NewComponent->SetFlags(RF_Transient);
 		NewComponent->AttachToComponent(SceneComponent, FAttachmentTransformRules::KeepRelativeTransform);
 
+		ScrubbedSound->SetFlags(RF_Transient);
+
 		FInstanceObjectKey DataKey{ InstanceHandle, SectionKey };
 		ExistingData = &ActorAudioComponentMap.Add(DataKey);
 		ExistingData->AudioComponent = NewComponent;
+		ExistingData->ScrubbedSound = ScrubbedSound;
 	}
 
 	return ExistingData;
@@ -1099,7 +1152,10 @@ UMovieSceneAudioSystem::FAudioComponentEvaluationData* UMovieSceneAudioSystem::A
 
 		UAudioComponent* NewComponent = FAudioDevice::CreateComponent(TempPlaybackAudioCue, Params);
 
-		if (!NewComponent)
+		static int32 ScrubSoundCounter = 0;
+		UScrubbedSound* ScrubbedSound = NewObject<UScrubbedSound>(World->GetCurrentLevel(), FName(*FString::Printf(TEXT("ScrubbedSound_Root_%i"), ScrubSoundCounter++)));
+
+		if (!NewComponent || !ScrubbedSound)
 		{
 			const int32 RowIndex = Section->GetRowIndex();
 			UE_LOG(LogMovieScene, Warning, TEXT("Failed to create audio component for root audio track (row %d)."), RowIndex);
@@ -1108,9 +1164,12 @@ UMovieSceneAudioSystem::FAudioComponentEvaluationData* UMovieSceneAudioSystem::A
 
 		NewComponent->SetFlags(RF_Transient);
 
+		ScrubbedSound->SetFlags(RF_Transient);
+
 		FInstanceObjectKey DataKey{ InstanceHandle, SectionKey };
 		ExistingData = &RootAudioComponentMap.Add(DataKey);
 		ExistingData->AudioComponent = NewComponent;
+		ExistingData->ScrubbedSound = ScrubbedSound;
 	}
 
 	return ExistingData;
diff --git a/Engine/Source/Runtime/MovieSceneTracks/Private/Systems/MovieSceneAudioSystem.h b/Engine/Source/Runtime/MovieSceneTracks/Private/Systems/MovieSceneAudioSystem.h
--- a/Engine/Source/Runtime/MovieSceneTracks/Private/Systems/MovieSceneAudioSystem.h
+++ b/Engine/Source/Runtime/MovieSceneTracks/Private/Systems/MovieSceneAudioSystem.h
@@ -9,6 +9,7 @@
 #include "UObject/ObjectMacros.h"
 #include "Channels/MovieSceneAudioTriggerChannel.h"
 #include "EntitySystem/MovieSceneEntitySystem.h"
+#include "Generators/SoundWaveScrubber.h"
 
 #include "MovieSceneAudioSystem.generated.h"
 
@@ -38,6 +39,7 @@ namespace UE::MovieScene
 	{
 		/** The audio component that was created to play audio */
 		TWeakObjectPtr<UAudioComponent> AudioComponent;
+		TWeakObjectPtr<UScrubbedSound> ScrubbedSound;
 
 		/** Volume multiplier to use this frame */
 		double VolumeMultiplier;
diff --git a/Engine/Source/Runtime/SignalProcessing/Public/DSP/GrainEnvelope.h b/Engine/Source/Runtime/SignalProcessing/Public/DSP/GrainEnvelope.h
--- a/Engine/Source/Runtime/SignalProcessing/Public/DSP/GrainEnvelope.h
+++ b/Engine/Source/Runtime/SignalProcessing/Public/DSP/GrainEnvelope.h
@@ -19,7 +19,8 @@ namespace Audio
 			DownwardTriangle,
 			UpwardTriangle,
 			ExponentialDecay,
-			ExponentialAttack
+			ExponentialAttack,
+			Hann
 		};
 		
 		// Utility function generates an envelope with the input array for the number of desired frames
@@ -89,6 +90,15 @@ namespace Audio
 					}
 				}
 				break;
+
+			case EEnvelope::Hann:
+				{
+					for (int32 i = 0; i < InNumFrames; ++i, n += 1.0f)
+					{
+						InData[i] = 0.5f - 0.5f * FMath::Cos(2.0f * PI * n / N_1);
+					}
+				}
+				break;
 			}
 		}
 
