Adding cvars to enable/disable scrubbing and a cvar to enable/disable scrubbing while stationary.

diff --git a/Engine/Source/Runtime/AudioMixer/Classes/Generators/SoundWaveScrubber.h b/Engine/Source/Runtime/AudioMixer/Classes/Generators/SoundWaveScrubber.h
--- a/Engine/Source/Runtime/AudioMixer/Classes/Generators/SoundWaveScrubber.h
+++ b/Engine/Source/Runtime/AudioMixer/Classes/Generators/SoundWaveScrubber.h
@@ -30,6 +30,9 @@ public:
 	// Sets whether or not to scrub the audio file. If false, then the output of the scrubber will be silence.
 	void SetIsScrubbing(bool bInIsScrubbing);
 
+	// Sets if the scrubber should scrub while playhead is stationary (after it hits it's target playhead)
+	void SetIsScrubbingWhileStationary(bool bInIsScrubWhileStationary);
+
 	// Sets the scrubbing location in seconds
 	void SetPlayheadTime(float InPlayheadTimeSeconds);
 
@@ -57,6 +60,9 @@ private:
 	// otherwise the scrubber will "jump" to the playhead time.
 	float TargetPlayheadTimeSeconds = 0.0f;
 
+	// How long we've been stationary
+	float TimeSincePlayheadHasNotChanged = 0.0f;
+
 	// Interpolated playhead time
 	FDynamicParameter CurrentPlayheadTimeSeconds;
 
@@ -152,6 +158,10 @@ private:
 
 	// Whether or not we're actively scrubbing audio playback. 
 	bool bIsScrubbing = false;
+	bool bIsScrubbingDueToBeingStationary = true;
+
+	// Whether or not we scrub while the playhead doesn't move
+	bool bIsScrubbingWhileStationary = true;
 };
 
 
@@ -168,6 +178,10 @@ public:
 
 	// Sets whether or not to scrub the audio file. If false, then the output of the scrubber will be silence.
 	void SetIsScrubbing(bool bInIsScrubbing);
+
+	// Sets if the scrubber should scrub while playhead is stationary (after it hits it's target playhead)
+	void SetIsScrubbingWhileStationary(bool bInScrubWhileStationary);
+
 	// Sets the scrubbing location in seconds
 	void SetPlayheadTime(float InPlayheadTimeSeconds);
 
@@ -204,6 +218,10 @@ public:
 	UFUNCTION(BlueprintCallable, Category = "Scrubbing")
 	AUDIOMIXER_API void SetPlayheadTime(float InPlayheadTimeSeconds);
 
+	// Returns the current playhead time
+	UFUNCTION(BlueprintCallable, Category = "Scrubbing")
+	AUDIOMIXER_API float GetPlayheadTime() const { return PlayheadTimeSeconds; }
+
 	// Sets the scrub grain duration range.
 	UFUNCTION(BlueprintCallable, Category = "Scrubbing")
 	AUDIOMIXER_API void SetGrainDurationRange(const FVector2D& InGrainDurationRangeSeconds);
@@ -212,11 +230,17 @@ public:
 	UFUNCTION(BlueprintCallable, Category = "Scrubbing")
 	AUDIOMIXER_API void SetIsScrubbing(bool bInIsScrubbing);
 
+	// Sets if the scrubber should scrub while playhead is stationary (after it hits it's target playhead)
+	UFUNCTION(BlueprintCallable, Category = "Scrubbing")
+	AUDIOMIXER_API void SetIsScrubbingWhileStationary(bool bInScrubWhileStationary);
+
 private:
 	float PlayheadTimeSeconds = 0.0f;
 	FVector2D GrainDurationRange = { 0.4f, 0.05f };
 
 	bool bIsScrubbing = false;
+	bool bScrubWhileStationary = true;
+	float StationaryTimeSeconds = 0.1f;
 
 	ISoundGeneratorPtr SoundWaveScrubber;
 
diff --git a/Engine/Source/Runtime/AudioMixer/Private/Generators/SoundWaveScrubber.cpp b/Engine/Source/Runtime/AudioMixer/Private/Generators/SoundWaveScrubber.cpp
--- a/Engine/Source/Runtime/AudioMixer/Private/Generators/SoundWaveScrubber.cpp
+++ b/Engine/Source/Runtime/AudioMixer/Private/Generators/SoundWaveScrubber.cpp
@@ -166,6 +166,11 @@ void FSoundWaveScrubber::SetIsScrubbing(bool bInIsScrubbing)
 	bIsScrubbing = bInIsScrubbing;
 }
 
+void FSoundWaveScrubber::SetIsScrubbingWhileStationary(bool bInIsScrubWhileStationary)
+{
+	bIsScrubbingWhileStationary = bInIsScrubWhileStationary;
+}
+
 void FSoundWaveScrubber::SetPlayheadTime(float InPlayheadTimeSeconds)
 {
 	FScopeLock Lock(&CritSect);
@@ -215,8 +220,28 @@ int32 FSoundWaveScrubber::RenderAudio(TArrayView<float>& OutAudio)
 				CurrentPlayheadTimeSeconds.Set(TargetPlayheadTimeSeconds, PlayheadLerpTime);
 			}
 		}
+		float PrevPlayHeadTime = CurrentPlayheadTimeSeconds.GetValue();
 		CurrentPlayheadTimeSeconds.Update(DeltaTimeSecond);
 
+		// Check if we're stationary
+		if (!bIsScrubbingWhileStationary)
+		{
+			if (FMath::IsNearlyEqual(PrevPlayHeadTime, CurrentPlayheadTimeSeconds.GetValue(), 0.001f))
+			{
+				TimeSincePlayheadHasNotChanged += DeltaTimeSecond;
+			}
+			else
+			{
+				TimeSincePlayheadHasNotChanged = 0.0f;
+			}
+
+			bIsScrubbingDueToBeingStationary = TimeSincePlayheadHasNotChanged < 0.1f;
+		}
+		else
+		{
+			bIsScrubbingDueToBeingStationary = true;
+		}
+
 		PlayheadTimeDistanceSeconds = FMath::Abs(CurrentPlayheadTimeSeconds.GetTargetValue() - CurrentPlayheadTimeSeconds.GetValue());
 
 		GrainDurationRange = TargetGrainDurationRange;
@@ -231,7 +256,7 @@ int32 FSoundWaveScrubber::RenderAudio(TArrayView<float>& OutAudio)
 	CurrentHalfGrainDurationFrames = 0.5f * CurrentGrainDurationFrames;
 
 	// If we're actively scrubbing we need to spawn grains and render the granular audio
-	if (bIsScrubbing)
+	if (bIsScrubbing && bIsScrubbingDueToBeingStationary)
 	{
 		if (!ActiveGrains.Num())
 		{
@@ -409,6 +434,11 @@ void FSoundWaveScrubberGenerator::SetIsScrubbing(bool bInIsScrubbing)
 	SoundWaveScrubber.SetIsScrubbing(bInIsScrubbing);
 }
 
+void FSoundWaveScrubberGenerator::SetIsScrubbingWhileStationary(bool bInScrubWhileStationary)
+{
+	SoundWaveScrubber.SetIsScrubbingWhileStationary(bInScrubWhileStationary);
+}
+
 void FSoundWaveScrubberGenerator::SetPlayheadTime(float InPlayheadTimeSeconds)
 {
 	SoundWaveScrubber.SetPlayheadTime(InPlayheadTimeSeconds);
@@ -458,6 +488,7 @@ ISoundGeneratorPtr UScrubbedSound::CreateSoundGenerator(const FSoundGeneratorIni
 		FSoundWaveScrubberGenerator* Scrubber = static_cast<FSoundWaveScrubberGenerator*>(SoundWaveScrubber.Get());
 		Scrubber->Init(SoundWaveProxyPtr, InParams.SampleRate, NumChannels, PlayheadTimeSeconds);
 		Scrubber->SetIsScrubbing(bIsScrubbing);
+		Scrubber->SetIsScrubbingWhileStationary(bScrubWhileStationary);
 		Scrubber->SetGrainDurationRange(GrainDurationRange);
 
 		return SoundWaveScrubber;
@@ -484,6 +515,18 @@ void UScrubbedSound::SetIsScrubbing(bool bInIsScrubbing)
 	}
 }
 
+void UScrubbedSound::SetIsScrubbingWhileStationary(bool bInScrubWhileStationary)
+{
+	using namespace Audio;
+
+	bScrubWhileStationary = bInScrubWhileStationary;
+	if (SoundWaveScrubber.IsValid())
+	{
+		FSoundWaveScrubberGenerator* Scrubber = static_cast<FSoundWaveScrubberGenerator*>(SoundWaveScrubber.Get());
+		Scrubber->SetIsScrubbingWhileStationary(bInScrubWhileStationary);
+	}
+}
+
 void UScrubbedSound::SetPlayheadTime(float InPlayheadTimeSeconds)
 {
 	using namespace Audio;
diff --git a/Engine/Source/Runtime/MovieSceneTracks/Private/Systems/MovieSceneAudioSystem.cpp b/Engine/Source/Runtime/MovieSceneTracks/Private/Systems/MovieSceneAudioSystem.cpp
--- a/Engine/Source/Runtime/MovieSceneTracks/Private/Systems/MovieSceneAudioSystem.cpp
+++ b/Engine/Source/Runtime/MovieSceneTracks/Private/Systems/MovieSceneAudioSystem.cpp
@@ -68,6 +68,21 @@ FAutoConsoleVariableRef CVarScrubWidthMilliseconds(
 	TEXT("The time-width of grains (in milliseconds) while scrubbing an audio track.\n"),
 	ECVF_Default);
 
+
+static bool bEnableGranularScrubbingCVar = true;
+FAutoConsoleVariableRef CVarEnableGranularScrubbing(
+	TEXT("Sequencer.Audio.EnableGranularScrubbing"),
+	bEnableGranularScrubbingCVar,
+	TEXT("Whether or not to use granular scrubbing.\n"),
+	ECVF_Default);
+
+static bool bEnableGranularScrubbingWhileStationaryCVar = true;
+FAutoConsoleVariableRef CVarEnableGranularScrubbingWhileStationary(
+	TEXT("Sequencer.Audio.EnableGranularScrubbingWhileStationary"),
+	bEnableGranularScrubbingWhileStationaryCVar,
+	TEXT("Whether or not to use granular scrubbing while holding the playhead still.\n"),
+	ECVF_Default);
+
 namespace UE::MovieScene
 {
 
@@ -718,7 +733,7 @@ private:
 			{
 				UE_LOG(LogMovieScene, Verbose, TEXT("Audio Component Play at Local Time: %6.2f CurrentTime: %6.2f(s) SectionStart: %6.2f(s), SoundDur: %6.2f OffsetIntoClip: %6.2f Component: %s Sound: %s"), AudioTime, (EvalTime / Context.GetFrameRate()), SectionStartTimeSeconds, AudioComponent.Sound ? AudioComponent.Sound->GetDuration() : 0.0f, (float)Context.GetFrameRate().AsSeconds(AudioStartOffset), *AudioComponent.GetName(), *GetNameSafe(AudioComponent.Sound));
 				
-				if (Context.GetStatus() == EMovieScenePlayerStatus::Scrubbing)
+				if (Context.GetStatus() == EMovieScenePlayerStatus::Scrubbing && bEnableGranularScrubbingCVar)
 				{
 					if (USoundWave* SoundWave = Cast<USoundWave>(Sound))
 					{
@@ -734,6 +749,7 @@ private:
 								ScrubbedSound->SetGrainDurationRange({MaxScrubWidthSeconds, 0.05f});
 
 								ScrubbedSound->SetIsScrubbing(true);
+								ScrubbedSound->SetIsScrubbingWhileStationary(bEnableGranularScrubbingWhileStationaryCVar);
 
 								AudioComponent.SetSound(ScrubbedSound);
 								AudioComponent.Play();
