Source Rendering Restructuring - Bypass current resampling code

Added the CVar "au.PerSourceResampling" to turn this new code path on (off by default).

#rnx

diff --git a/Engine/Source/Runtime/AudioMixer/Private/AudioMixerSourceBuffer.cpp b/Engine/Source/Runtime/AudioMixer/Private/AudioMixerSourceBuffer.cpp
--- a/Engine/Source/Runtime/AudioMixer/Private/AudioMixerSourceBuffer.cpp
+++ b/Engine/Source/Runtime/AudioMixer/Private/AudioMixerSourceBuffer.cpp
@@ -126,8 +126,6 @@ namespace Audio
 
 			// Prepare the memory to fit the max number of samples
 			SourceVoiceBuffers[BufferIndex]->AudioData.Reset(TotalSamples);
-			SourceVoiceBuffers[BufferIndex]->bRealTimeBuffer = true;
-			SourceVoiceBuffers[BufferIndex]->LoopCount = 0;
 		}
 	}
 
diff --git a/Engine/Source/Runtime/AudioMixer/Private/AudioMixerSourceManager.cpp b/Engine/Source/Runtime/AudioMixer/Private/AudioMixerSourceManager.cpp
--- a/Engine/Source/Runtime/AudioMixer/Private/AudioMixerSourceManager.cpp
+++ b/Engine/Source/Runtime/AudioMixer/Private/AudioMixerSourceManager.cpp
@@ -146,11 +146,19 @@ static FAutoConsoleCommand GModulationSetMaxPitchRange(
 
 			const float Range = FCString::Atof(*Args[0]);
 			MaxModulationPitchRangeFreqCVar = Audio::GetFrequencyMultiplier(Range * 0.5f);
-			MaxModulationPitchRangeFreqCVar = Audio::GetFrequencyMultiplier(Range * -0.5f);
+			MinModulationPitchRangeFreqCVar = Audio::GetFrequencyMultiplier(Range * -0.5f);
 		}
 	)
 );
 
+static int32 PerSourceResampling = 0;
+FAutoConsoleVariableRef CVarAudioMixerPerSourceResampling(
+	TEXT("au.PerSourceResampling"),
+	PerSourceResampling,
+	TEXT("Use new source rendering code path, with each source getting a dedicated resampler if needed.\n")
+	TEXT("0: Disabled, 1: Enabled"),
+	ECVF_Default);
+
 #define ENVELOPE_TAIL_THRESHOLD (1.58489e-5f) // -96 dB
 
 #define VALIDATE_SOURCE_MIXER_STATE 1
@@ -392,7 +400,6 @@ namespace Audio
 
 			SourceInfo.SourceListener = nullptr;
 			SourceInfo.CurrentPCMBuffer = nullptr;	
-			SourceInfo.CurrentAudioChunkNumFrames = 0;
 			SourceInfo.CurrentFrameAlpha = 0.0f;
 			SourceInfo.CurrentFrameIndex = 0;
 			SourceInfo.NumFramesPlayed = 0;
@@ -779,7 +786,6 @@ namespace Audio
 		SourceInfo.LowPassFilter.Reset();
 		SourceInfo.HighPassFilter.Reset();
 		SourceInfo.CurrentPCMBuffer = nullptr;
-		SourceInfo.CurrentAudioChunkNumFrames = 0;
 		SourceInfo.SourceBuffer.Reset();
 		SourceInfo.PreDistanceAttenuationBuffer.Reset();
 		SourceInfo.SourceEffectScratchBuffer.Reset();
@@ -2080,9 +2086,12 @@ namespace Audio
 
 		const int32 NumChannels = SourceInfo.NumInputChannels;
 
+		// Number of frames in the current PCM buffer. Gets updated when we fetch a new buffer.
+		int32 CurrentAudioChunkNumFrames = SourceInfo.GetCurrentAudioChunkNumFrames();
+
 		// Check if the next frame index is out of range of the total number of frames we have in our current audio buffer
-		bool bNextFrameOutOfRange = (SourceInfo.CurrentFrameIndex + 1) >= SourceInfo.CurrentAudioChunkNumFrames;
-		bool bCurrentFrameOutOfRange = SourceInfo.CurrentFrameIndex >= SourceInfo.CurrentAudioChunkNumFrames;
+		bool bNextFrameOutOfRange = (SourceInfo.CurrentFrameIndex + 1) >= CurrentAudioChunkNumFrames;
+		bool bCurrentFrameOutOfRange = SourceInfo.CurrentFrameIndex >= CurrentAudioChunkNumFrames;
 
 		bool bReadCurrentFrame = true;
 
@@ -2108,14 +2117,6 @@ namespace Audio
 			// If this is our first PCM buffer, we don't need to do a callback to get more audio
 			if (SourceInfo.CurrentPCMBuffer.IsValid())
 			{
-				if (SourceInfo.CurrentPCMBuffer->LoopCount == Audio::LOOP_FOREVER && !SourceInfo.CurrentPCMBuffer->bRealTimeBuffer)
-				{
-					AUDIO_MIXER_DEBUG_LOG(SourceId, TEXT("Hit Loop boundary, looping."));
-
-					SourceInfo.CurrentFrameIndex = FMath::Max(SourceInfo.CurrentFrameIndex - SourceInfo.CurrentAudioChunkNumFrames, 0);
-					break;
-				}
-
 				if (ensure(SourceInfo.MixerSourceBuffer.IsValid()))
 				{
 #if ENABLE_AUDIO_DEBUG
@@ -2134,13 +2135,13 @@ namespace Audio
 			if (ensure(SourceInfo.MixerSourceBuffer.IsValid()) && SourceInfo.MixerSourceBuffer->GetNumBuffersQueued() > 0 && NumChannels > 0)
 			{
 				SourceInfo.CurrentPCMBuffer = SourceInfo.MixerSourceBuffer->GetNextBuffer();
-				SourceInfo.CurrentAudioChunkNumFrames = SourceInfo.CurrentPCMBuffer->AudioData.Num() / NumChannels;
+				CurrentAudioChunkNumFrames = SourceInfo.GetCurrentAudioChunkNumFrames();
 
 				// Subtract the number of frames in the current buffer from our frame index.
 				// Note: if this is the first time we're playing, CurrentFrameIndex will be 0
 				if (bReadCurrentFrame)
 				{
-					SourceInfo.CurrentFrameIndex = FMath::Max(SourceInfo.CurrentFrameIndex - SourceInfo.CurrentAudioChunkNumFrames, 0);
+					SourceInfo.CurrentFrameIndex = FMath::Max(SourceInfo.CurrentFrameIndex - CurrentAudioChunkNumFrames, 0);
 				}
 				else
 				{
@@ -2155,8 +2156,8 @@ namespace Audio
 				return;
 			}
 
-			bNextFrameOutOfRange = (SourceInfo.CurrentFrameIndex + 1) >= SourceInfo.CurrentAudioChunkNumFrames;
-			bCurrentFrameOutOfRange = SourceInfo.CurrentFrameIndex >= SourceInfo.CurrentAudioChunkNumFrames;
+			bNextFrameOutOfRange = (SourceInfo.CurrentFrameIndex + 1) >= CurrentAudioChunkNumFrames;
+			bCurrentFrameOutOfRange = SourceInfo.CurrentFrameIndex >= CurrentAudioChunkNumFrames;
 		}
 
 		if (SourceInfo.CurrentPCMBuffer.IsValid())
@@ -2198,169 +2199,159 @@ namespace Audio
 		}
 	}
 
-	void FMixerSourceManager::ComputeSourceBuffersForIdRange(const bool bGenerateBuses, const int32 SourceIdStart, const int32 SourceIdEnd)
+	void FMixerSourceManager::ComputeSourceBuffer(const bool bGenerateBuses, const int32 SourceId)
 	{
 		CSV_SCOPED_TIMING_STAT(Audio, SourceBuffers);
-		CONDITIONAL_SCOPE_CYCLE_COUNTER(STAT_AudioMixerSourceBuffers, (SourceIdStart < SourceIdEnd));
+		SCOPE_CYCLE_COUNTER(STAT_AudioMixerSourceBuffers);
 
 		const double AudioRenderThreadTime = MixerDevice->GetAudioRenderThreadTime();
 		const double AudioClockDelta = MixerDevice->GetAudioClockDelta();
 
-		for (int32 SourceId = SourceIdStart; SourceId < SourceIdEnd; ++SourceId)
-		{
-			FSourceInfo& SourceInfo = SourceInfos[SourceId];
-
-			if (!SourceInfo.bIsBusy || !SourceInfo.bIsPlaying || SourceInfo.bIsPaused || SourceInfo.bIsPausedForQuantization)
-			{
-				continue;
-			}
-
-			// If this source is still playing at this point but technically done, zero the buffers. We haven't yet been removed by the FMixerSource owner.
-			// This should be rare but could happen due to thread timing since done-ness is queried on audio thread.
-			if (SourceInfo.bIsDone)
-			{
-				const int32 NumSamples = NumOutputFrames * SourceInfo.NumInputChannels;
+		FSourceInfo& SourceInfo = SourceInfos[SourceId];
 
-				SourceInfo.PreDistanceAttenuationBuffer.Reset();
-				SourceInfo.PreDistanceAttenuationBuffer.AddZeroed(NumSamples);
+		if (!SourceInfo.bIsBusy || !SourceInfo.bIsPlaying || SourceInfo.bIsPaused || SourceInfo.bIsPausedForQuantization)
+		{
+			return;
+		}
 
-				SourceInfo.SourceBuffer.Reset();
-				SourceInfo.SourceBuffer.AddZeroed(NumSamples);
+		const bool bIsSourceBus = SourceInfo.AudioBusId != INDEX_NONE;
+		if ((bGenerateBuses && !bIsSourceBus) || (!bGenerateBuses && bIsSourceBus))
+		{
+			return;
+		}
 
-				continue;
-			}
+		// Fill array with elements all at once to avoid sequential Add() operation overhead.
+		const int32 NumSamples = NumOutputFrames * SourceInfo.NumInputChannels;
 
-			const bool bIsSourceBus = SourceInfo.AudioBusId != INDEX_NONE;
-			if ((bGenerateBuses && !bIsSourceBus) || (!bGenerateBuses && bIsSourceBus))
-			{
-				continue;
-			}
+		UE_CLOG(Audio::MatchesLogFilter(*SourceInfo.GetDebugName()), LogAudioTiming, Verbose,
+			TEXT("ComputeSourceBuffer Name=%s, NumOutputFrames=%d, FramesPlayed=%lld, CurrentFrameIndex=%d "),
+			*SourceInfo.GetDebugName(), NumOutputFrames, SourceInfo.NumFramesPlayed, SourceInfo.CurrentFrameIndex);
 
-			// Fill array with elements all at once to avoid sequential Add() operation overhead.
-			const int32 NumSamples = NumOutputFrames * SourceInfo.NumInputChannels;
-			
-			UE_CLOG(Audio::MatchesLogFilter(*SourceInfo.GetDebugName()), LogAudioTiming, Verbose,
-				TEXT("ComputeSourceBuffersForIdRange Name=%s, NumOutputFrames=%d, FramesPlayed=%lld, CurrentFrameIndex=%d "),
-					*SourceInfo.GetDebugName(), NumOutputFrames, SourceInfo.NumFramesPlayed, SourceInfo.CurrentFrameIndex);
-			
-			// Initialize both the pre-distance attenuation buffer and the source buffer
-			SourceInfo.PreDistanceAttenuationBuffer.Reset();
-			SourceInfo.PreDistanceAttenuationBuffer.AddZeroed(NumSamples);
+		// Initialize both the pre-distance attenuation buffer and the source buffer
+		SourceInfo.PreDistanceAttenuationBuffer.Reset();
+		SourceInfo.PreDistanceAttenuationBuffer.AddZeroed(NumSamples);
 
+		SourceInfo.SourceEffectScratchBuffer.Reset();
+		SourceInfo.SourceEffectScratchBuffer.AddZeroed(NumSamples);
 
-			SourceInfo.SourceEffectScratchBuffer.Reset();
-			SourceInfo.SourceEffectScratchBuffer.AddZeroed(NumSamples);
+		SourceInfo.SourceBuffer.Reset();
+		SourceInfo.SourceBuffer.AddZeroed(NumSamples);
 
-			SourceInfo.SourceBuffer.Reset();
-			SourceInfo.SourceBuffer.AddZeroed(NumSamples);
+		// If this source is still playing at this point but technically done, return after zeroing the buffers. We haven't yet been removed by the FMixerSource owner.
+		// This should be rare but could happen due to thread timing since done-ness is queried on audio thread.
+		if (SourceInfo.bIsDone)
+		{
+			return;
+		}
 
-			if (SourceInfo.SubCallbackDelayLengthInFrames && !SourceInfo.bDelayLineSet)
-			{
-				SourceInfo.SourceBufferDelayLine.SetCapacity(SourceInfo.SubCallbackDelayLengthInFrames * SourceInfo.NumInputChannels + SourceInfo.NumInputChannels);
-				SourceInfo.SourceBufferDelayLine.PushZeros(SourceInfo.SubCallbackDelayLengthInFrames * SourceInfo.NumInputChannels);
-				SourceInfo.bDelayLineSet = true;
-			}
+		if (SourceInfo.SubCallbackDelayLengthInFrames && !SourceInfo.bDelayLineSet && !PerSourceResampling)
+		{
+			SourceInfo.SourceBufferDelayLine.SetCapacity(SourceInfo.SubCallbackDelayLengthInFrames * SourceInfo.NumInputChannels + SourceInfo.NumInputChannels);
+			SourceInfo.SourceBufferDelayLine.PushZeros(SourceInfo.SubCallbackDelayLengthInFrames * SourceInfo.NumInputChannels);
+			SourceInfo.bDelayLineSet = true;
+		}
 
-			float* PreDistanceAttenBufferPtr = SourceInfo.PreDistanceAttenuationBuffer.GetData();
+		float* PreDistanceAttenBufferPtr = SourceInfo.PreDistanceAttenuationBuffer.GetData();
 
-			// if this is a bus, we just want to copy the bus audio to this source's output audio
-			// Note we need to copy this since bus instances may have different audio via dynamic source effects, etc.
-			if (bIsSourceBus)
+		// if this is a bus, we just want to copy the bus audio to this source's output audio
+		// Note we need to copy this since bus instances may have different audio via dynamic source effects, etc.
+		if (bIsSourceBus)
+		{
+			// Get the source's rendered and mixed audio bus data
+			const TSharedPtr<FMixerAudioBus> AudioBusPtr = AudioBuses.FindRef(SourceInfo.AudioBusId);
+			if (AudioBusPtr.IsValid())
 			{
-				// Get the source's rendered and mixed audio bus data
-				const TSharedPtr<FMixerAudioBus> AudioBusPtr = AudioBuses.FindRef(SourceInfo.AudioBusId);
-				if (AudioBusPtr.IsValid())
+				int32 NumFramesPlayed = NumOutputFrames;
+				if (SourceInfo.SourceBusDurationFrames != INDEX_NONE)
 				{
-					int32 NumFramesPlayed = NumOutputFrames;
-					if (SourceInfo.SourceBusDurationFrames != INDEX_NONE)
+					// If we're now finishing, only copy over the real data
+					if ((SourceInfo.NumFramesPlayed + NumOutputFrames) >= SourceInfo.SourceBusDurationFrames)
 					{
-						// If we're now finishing, only copy over the real data
-						if ((SourceInfo.NumFramesPlayed + NumOutputFrames) >= SourceInfo.SourceBusDurationFrames)
-						{
-							NumFramesPlayed = SourceInfo.SourceBusDurationFrames - SourceInfo.NumFramesPlayed;
-							SourceInfo.bIsLastBuffer = true;
-						}
+						NumFramesPlayed = SourceInfo.SourceBusDurationFrames - SourceInfo.NumFramesPlayed;
+						SourceInfo.bIsLastBuffer = true;
 					}
+				}
 
-					SourceInfo.NumFramesPlayed += NumFramesPlayed;
+				SourceInfo.NumFramesPlayed += NumFramesPlayed;
 
-					// Retrieve the channel map of going from the audio bus channel count to the source channel count since they may not match
-					int32 NumAudioBusChannels = AudioBusPtr->GetNumChannels();
-					if (NumAudioBusChannels != SourceInfo.NumInputChannels)
-					{
-						Audio::FAlignedFloatBuffer ChannelMap;
-						MixerDevice->Get2DChannelMap(SourceInfo.bIsVorbis, AudioBusPtr->GetNumChannels(), SourceInfo.NumInputChannels, SourceInfo.bIsCenterChannelOnly, ChannelMap);
-						AudioBusPtr->CopyCurrentBuffer(ChannelMap, SourceInfo.NumInputChannels, SourceInfo.PreDistanceAttenuationBuffer, NumFramesPlayed);
-					}
-					else
-					{
-						AudioBusPtr->CopyCurrentBuffer(SourceInfo.NumInputChannels, SourceInfo.PreDistanceAttenuationBuffer, NumFramesPlayed);
-					}
+				// Retrieve the channel map of going from the audio bus channel count to the source channel count since they may not match
+				int32 NumAudioBusChannels = AudioBusPtr->GetNumChannels();
+				if (NumAudioBusChannels != SourceInfo.NumInputChannels)
+				{
+					Audio::FAlignedFloatBuffer ChannelMap;
+					MixerDevice->Get2DChannelMap(SourceInfo.bIsVorbis, AudioBusPtr->GetNumChannels(), SourceInfo.NumInputChannels, SourceInfo.bIsCenterChannelOnly, ChannelMap);
+					AudioBusPtr->CopyCurrentBuffer(ChannelMap, SourceInfo.NumInputChannels, SourceInfo.PreDistanceAttenuationBuffer, NumFramesPlayed);
+				}
+				else
+				{
+					AudioBusPtr->CopyCurrentBuffer(SourceInfo.NumInputChannels, SourceInfo.PreDistanceAttenuationBuffer, NumFramesPlayed);
 				}
 			}
-			else
-			{
+		}
+		else
+		{
 
 #if AUDIO_SUBFRAME_ENABLED
-				// If we're not going to start yet, just continue
-				double StartFraction = (SourceInfo.StartTime - AudioRenderThreadTime) / AudioClockDelta;
-				if (StartFraction >= 1.0)
-				{
-					// note this is already zero'd so no need to write zeroes
-					SourceInfo.PitchSourceParam.Reset();
-					continue;
-				}
+			// If we're not going to start yet, just continue
+			double StartFraction = (SourceInfo.StartTime - AudioRenderThreadTime) / AudioClockDelta;
+			if (StartFraction >= 1.0)
+			{
+				// note this is already zero'd so no need to write zeroes
+				SourceInfo.PitchSourceParam.Reset();
+				return;
+			}
 
-				// Init the frame index iterator to 0 (i.e. render whole buffer)
-				int32 StartFrame = 0;
+			// Init the frame index iterator to 0 (i.e. render whole buffer)
+			int32 StartFrame = 0;
 
-				// If the start fraction is greater than 0.0 (and is less than 1.0), we are starting on a sub-frame
-				// Otherwise, just start playing it right away
-				if (StartFraction > 0.0)
-				{
-					StartFrame = NumOutputFrames * StartFraction;
-				}
+			// If the start fraction is greater than 0.0 (and is less than 1.0), we are starting on a sub-frame
+			// Otherwise, just start playing it right away
+			if (StartFraction > 0.0)
+			{
+				StartFrame = NumOutputFrames * StartFraction;
+			}
 
-				// Update sample index to the frame we're starting on, accounting for source channels
-				int32 SampleIndex = StartFrame * SourceInfo.NumInputChannels;
-				bool bWriteOutZeros = true;
+			// Update sample index to the frame we're starting on, accounting for source channels
+			int32 SampleIndex = StartFrame * SourceInfo.NumInputChannels;
 #else
-				int32 SampleIndex = 0;
-				int32 StartFrame = 0;
+			int32 SampleIndex = 0;
+			int32 StartFrame = 0;
 #endif
 
-				// Modulate parameter target should modulation be active
-				// Due to managing two separate pitch values that are updated at different rates
-				// (game thread rate and copy set by SetPitch and buffer callback rate set by Modulation System),
-				// the PitchSourceParam's target is marshaled before processing by mult'ing in the modulation pitch,
-				// processing the buffer, and then resetting it back if modulation is active. 
+			// Modulate parameter target should modulation be active
+			// Due to managing two separate pitch values that are updated at different rates
+			// (game thread rate and copy set by SetPitch and buffer callback rate set by Modulation System),
+			// the PitchSourceParam's target is marshaled before processing by mult'ing in the modulation pitch,
+			// processing the buffer, and then resetting it back if modulation is active. 
 
-				const bool bModActive = MixerDevice->IsModulationPluginEnabled() && MixerDevice->ModulationInterface.IsValid();
-				if (bModActive)
-				{
-					SourceInfo.PitchModulation.ProcessControl(SourceInfo.PitchModulationBase);
-				}
+			const bool bModActive = MixerDevice->IsModulationPluginEnabled() && MixerDevice->ModulationInterface.IsValid();
+			if (bModActive)
+			{
+				SourceInfo.PitchModulation.ProcessControl(SourceInfo.PitchModulationBase);
+			}
 
-				const float TargetPitch = SourceInfo.PitchSourceParam.GetTarget();
-				// Convert from semitones to frequency multiplier
-				const float ModPitch = bModActive
-					? Audio::GetFrequencyMultiplier(SourceInfo.PitchModulation.GetValue())
-					: 1.0f;
-				const float FinalPitch = FMath::Clamp(TargetPitch * ModPitch, MinModulationPitchRangeFreqCVar, MaxModulationPitchRangeFreqCVar);
-				SourceInfo.PitchSourceParam.SetValue(FinalPitch, NumOutputFrames);
+			const float TargetPitch = SourceInfo.PitchSourceParam.GetTarget();
+			// Convert from semitones to frequency multiplier
+			const float ModPitch = bModActive
+				? Audio::GetFrequencyMultiplier(SourceInfo.PitchModulation.GetValue())
+				: 1.0f;
+			const float FinalPitch = FMath::Clamp(TargetPitch * ModPitch, MinModulationPitchRangeFreqCVar, MaxModulationPitchRangeFreqCVar);
+			SourceInfo.PitchSourceParam.SetValue(FinalPitch, NumOutputFrames);
 
 #if UE_AUDIO_PROFILERTRACE_ENABLED
-				const bool bChannelEnabled = UE_TRACE_CHANNELEXPR_IS_ENABLED(AudioMixerChannel);
-				if (bChannelEnabled)
-				{
-					UE_TRACE_LOG(Audio, MixerSourcePitch, AudioMixerChannel)
-						<< MixerSourcePitch.DeviceId(MixerDevice->DeviceID)
-						<< MixerSourcePitch.Timestamp(FPlatformTime::Cycles64())
-						<< MixerSourcePitch.PlayOrder(SourceInfo.PlayOrder)
-						<< MixerSourcePitch.Pitch(TargetPitch);
-				}
+			const bool bChannelEnabled = UE_TRACE_CHANNELEXPR_IS_ENABLED(AudioMixerChannel);
+			if (bChannelEnabled)
+			{
+				UE_TRACE_LOG(Audio, MixerSourcePitch, AudioMixerChannel)
+					<< MixerSourcePitch.DeviceId(MixerDevice->DeviceID)
+					<< MixerSourcePitch.Timestamp(FPlatformTime::Cycles64())
+					<< MixerSourcePitch.PlayOrder(SourceInfo.PlayOrder)
+					<< MixerSourcePitch.Pitch(TargetPitch);
+			}
 #endif // UE_AUDIO_PROFILERTRACE_ENABLED
 
+			if (!PerSourceResampling)
+			{
 				float CurrentAlpha = SourceInfo.CurrentFrameAlpha;
 
 				for (int32 Frame = StartFrame; Frame < NumOutputFrames; ++Frame)
@@ -2383,7 +2374,7 @@ namespace Audio
 					{
 						// Our inter-frame alpha lerping value is causing us to read new source frames
 						bReadNextSample = true;
-						
+
 						const float Delta = FMath::FloorToFloat(CurrentAlpha);
 						const int DeltaInt = (int)Delta;
 
@@ -2433,18 +2424,97 @@ namespace Audio
 				}
 
 				SourceInfo.CurrentFrameAlpha = CurrentAlpha;
+			}
+			else // PerSourceResampling
+			{
+				SourceInfo.bHasStarted = true;
 
-				// After processing the frames, reset the pitch param
-				SourceInfo.PitchSourceParam.Reset();
+				// Fill PreDistanceAttenuationBuffer with incoming source data
+				if (ensure(SourceInfo.MixerSourceBuffer.IsValid() && SourceInfo.NumInputChannels > 0))
+				{
+					int32 FramesToWrite = NumOutputFrames;
+
+					// Skip writing some frames in the first block if a delay is requested.
+					if (SourceInfo.SubCallbackDelayLengthInFrames > 0 && !SourceInfo.bDelayLineSet)
+					{
+						if (ensure(SourceInfo.SubCallbackDelayLengthInFrames <= FramesToWrite))
+						{
+							FramesToWrite -= SourceInfo.SubCallbackDelayLengthInFrames;
+						}
+						SourceInfo.bDelayLineSet = true;
+					}
 
-				// Reset target value as modulation may have modified prior to processing
-				// And source param should not store modulation value internally as its
-				// processed by the modulation plugin independently.
-				SourceInfo.PitchSourceParam.SetValue(TargetPitch, NumOutputFrames);
+					// Also skip some frames if subframe rendering is enabled
+					FramesToWrite = FMath::Max(FramesToWrite - StartFrame, 0);
+
+					while (FramesToWrite > 0)
+					{
+						int32 CurrentChunkFrames = SourceInfo.GetCurrentAudioChunkNumFrames();
+						if (SourceInfo.CurrentFrameIndex >= CurrentChunkFrames)
+						{
+							// Time to read another chunk
+
+							// If this is our first PCM buffer, we don't need to do a callback to get more audio
+							if (SourceInfo.CurrentPCMBuffer.IsValid())
+							{
+#if ENABLE_AUDIO_DEBUG
+								// Writing to this value is a read/write race condition on the CPUCoreUtilization value. Calling this
+								// out as an acceptable race condition given that it is utilized for debug purposes only. 
+								GameThreadInfo.CPUCoreUtilization[SourceId] = SourceInfo.MixerSourceBuffer->GetCPUCoreUtilization();
+#endif // if ENABLE_AUDIO_DEBUG
+								GameThreadInfo.RelativeRenderCost[SourceId] = SourceInfo.MixerSourceBuffer->GetRelativeRenderCost();
+								
+								SourceInfo.MixerSourceBuffer->OnBufferEnd();
+							}
+
+							// If we're out of audio in our queue, we're done
+							if (SourceInfo.MixerSourceBuffer->GetNumBuffersQueued() == 0)
+							{
+								SourceInfo.bIsLastBuffer = true;
+								break;
+							}
+
+							SourceInfo.CurrentPCMBuffer = SourceInfo.MixerSourceBuffer->GetNextBuffer();
+							SourceInfo.CurrentFrameIndex = 0;
+							CurrentChunkFrames = SourceInfo.GetCurrentAudioChunkNumFrames();
+
+							if (!ensure(SourceInfo.CurrentPCMBuffer.IsValid()))
+							{
+								break;
+							}
+						}
+
+						// Copy as much data as possible from CurrentPCMBuffer to PreDistanceAttenuationBuffer
+						check(SourceInfo.CurrentPCMBuffer.IsValid());
+
+						const int32 FramesToCopy = FMath::Min(FramesToWrite, CurrentChunkFrames - SourceInfo.CurrentFrameIndex);
+						const int32 SamplesToCopy = FramesToCopy * SourceInfo.NumInputChannels;
+
+						const int32 DestinationOffset = NumSamples - FramesToWrite * SourceInfo.NumInputChannels;
+						float* CopyDestination = SourceInfo.PreDistanceAttenuationBuffer.GetData() + DestinationOffset;
+						const float* CopySource = SourceInfo.CurrentPCMBuffer->AudioData.GetData() + SourceInfo.CurrentFrameIndex * SourceInfo.NumInputChannels;
+
+						FMemory::Memcpy(CopyDestination, CopySource, SamplesToCopy * sizeof(float));
+
+						// These will be scaled when resampling is being done.
+						SourceInfo.CurrentFrameIndex += FramesToCopy;
+						SourceInfo.NumFramesPlayed += FramesToCopy;
+
+						FramesToWrite -= FramesToCopy;
+					}
+				}
 			}
+
+			// After processing the frames, reset the pitch param
+			SourceInfo.PitchSourceParam.Reset();
+
+			// Reset target value as modulation may have modified prior to processing
+			// And source param should not store modulation value internally as its
+			// processed by the modulation plugin independently.
+			SourceInfo.PitchSourceParam.SetValue(TargetPitch, NumOutputFrames);
 		}
 	}
-
+	
 	void FMixerSourceManager::ConnectBusPatches()
 	{
 		while (TOptional<FPendingAudioBusConnection> PendingAudioBusConnection = PendingAudioBusConnections.Dequeue())
@@ -3107,7 +3177,10 @@ namespace Audio
 	{
 		// Buses generate their input buffers independently
 		// Get the next block of frames from the source buffers
-		ComputeSourceBuffersForIdRange(bGenerateBuses, SourceIdStart, SourceIdEnd);
+		for (int32 SourceId = SourceIdStart; SourceId < SourceIdEnd; ++SourceId)
+		{
+			ComputeSourceBuffer(bGenerateBuses, SourceId);
+		}
 
 		// Compute the audio source buffers after their individual effect chain processing
 		ComputePostSourceEffectBufferForIdRange(bGenerateBuses, SourceIdStart, SourceIdEnd);
diff --git a/Engine/Source/Runtime/AudioMixer/Private/AudioMixerSourceManager.h b/Engine/Source/Runtime/AudioMixer/Private/AudioMixerSourceManager.h
--- a/Engine/Source/Runtime/AudioMixer/Private/AudioMixerSourceManager.h
+++ b/Engine/Source/Runtime/AudioMixer/Private/AudioMixerSourceManager.h
@@ -76,12 +76,6 @@ namespace Audio
 	{
 		/** PCM float data. */
 		FAlignedFloatBuffer AudioData;
-
-		/** How many times this buffer will loop. */
-		int32 LoopCount = 0;
-
-		/** If this buffer is from real-time decoding and needs to make callbacks for more data. */
-		bool bRealTimeBuffer = false;
 	};
 
 
@@ -338,7 +332,7 @@ namespace Audio
 		void GenerateSourceAudio(const bool bGenerateBuses);
 		void GenerateSourceAudio(const bool bGenerateBuses, const int32 SourceIdStart, const int32 SourceIdEnd);
 
-		void ComputeSourceBuffersForIdRange(const bool bGenerateBuses, const int32 SourceIdStart, const int32 SourceIdEnd);
+		void ComputeSourceBuffer(const bool bGenerateBuses, const int32 SourceId);
 		void ComputePostSourceEffectBufferForIdRange(const bool bGenerateBuses, const int32 SourceIdStart, const int32 SourceIdEnd);
 		void ComputeOutputBuffersForIdRange(const bool bGenerateBuses, const int32 SourceIdStart, const int32 SourceIdEnd);
 
@@ -454,7 +448,6 @@ namespace Audio
 
 			// Data used for rendering sources
 			TSharedPtr<FMixerSourceVoiceBuffer, ESPMode::ThreadSafe> CurrentPCMBuffer;
-			int32 CurrentAudioChunkNumFrames;
 
 			// The post-attenuation source buffer, used to send audio to submixes
 			Audio::FAlignedFloatBuffer SourceBuffer;
@@ -608,6 +601,20 @@ namespace Audio
 #endif //AUDIO_MIXER_ENABLE_DEBUG_MODE
 				
 			}
+
+			// return the number of sample frames in the current input PCM buffer (if any)
+			int32 GetCurrentAudioChunkNumFrames() const
+			{
+				if (CurrentPCMBuffer)
+				{
+					check(NumInputChannels != 0);
+					return CurrentPCMBuffer->AudioData.Num() / NumInputChannels;
+				}
+				else
+				{
+					return 0;
+				}
+			}
 		};
 
 		static void ApplyDistanceAttenuation(FSourceInfo& InSourceInfo, int32 NumSamples);
diff --git a/Engine/Source/Runtime/AudioMixer/Private/SoundWaveDecoder.cpp b/Engine/Source/Runtime/AudioMixer/Private/SoundWaveDecoder.cpp
--- a/Engine/Source/Runtime/AudioMixer/Private/SoundWaveDecoder.cpp
+++ b/Engine/Source/Runtime/AudioMixer/Private/SoundWaveDecoder.cpp
@@ -216,12 +216,6 @@ namespace Audio
 					SourceInfo.CurrentFrameValues[Channel] = AudioData[CurrentSampleIndex + Channel];
 				}
 
-				if (SourceInfo.CurrentPCMBuffer->LoopCount == Audio::LOOP_FOREVER && !SourceInfo.CurrentPCMBuffer->bRealTimeBuffer)
-				{
-					SourceInfo.CurrentFrameIndex = FMath::Max(SourceInfo.CurrentFrameIndex - SourceInfo.CurrentAudioChunkNumFrames, 0);
-					break;
-				}
-
 				MixerSourceBufferLocal->OnBufferEnd();
 			}
 
