Exporting marker data to soundwave using the waveform editor and exporting marker data to .wav files using the right click context menu (now using libsndfile)

diff --git a/Engine/Plugins/Editor/WaveformEditor/Source/WaveformEditor/Private/WaveformEditorWaveWriter.cpp b/Engine/Plugins/Editor/WaveformEditor/Source/WaveformEditor/Private/WaveformEditorWaveWriter.cpp
--- a/Engine/Plugins/Editor/WaveformEditor/Source/WaveformEditor/Private/WaveformEditorWaveWriter.cpp
+++ b/Engine/Plugins/Editor/WaveformEditor/Source/WaveformEditor/Private/WaveformEditorWaveWriter.cpp
@@ -57,8 +57,11 @@ void FWaveformEditorWaveWriter::ExportTransformedWaveform()
 			}
 		}
 	};
+	
+	//Duplicate the SourceSoundWave so we do not overwrite the original file with the transformations on export
+	TObjectPtr<USoundWave> ExportSoundWave = DuplicateObject<USoundWave>(SourceSoundWave, SourceSoundWave->GetOuter());
 
-	if (!WaveWriter->BeginWriteToSoundWave(AssetName, BufferToWrite, AssetPath, OnSoundWaveWritten))
+	if (!WaveWriter->BeginWriteToSoundWave(AssetName, BufferToWrite, AssetPath, OnSoundWaveWritten, ExportSoundWave))
 	{
 		UE_LOG(LogWaveformEditor, Log, TEXT("Exporting edited soundwave to %s/%s failed"), *AssetPath, *AssetName);
 	}
diff --git a/Engine/Plugins/Editor/WaveformEditor/Source/WaveformTransformations/Private/WaveformTransformationMarkers.cpp b/Engine/Plugins/Editor/WaveformEditor/Source/WaveformTransformations/Private/WaveformTransformationMarkers.cpp
--- a/Engine/Plugins/Editor/WaveformEditor/Source/WaveformTransformations/Private/WaveformTransformationMarkers.cpp
+++ b/Engine/Plugins/Editor/WaveformEditor/Source/WaveformTransformations/Private/WaveformTransformationMarkers.cpp
@@ -1,6 +1,7 @@
 // Copyright Epic Games, Inc. All Rights Reserved.
 
 #include "WaveformTransformationMarkers.h"
+#include "Sound/SoundWave.h"
 
 UWaveformTransformationMarkers::UWaveformTransformationMarkers(const FObjectInitializer& ObjectInitializer) : Super(ObjectInitializer)
 {
@@ -28,6 +29,42 @@ void UWaveformTransformationMarkers::OverwriteTransformation()
 }
 
 #if WITH_EDITOR
+void UWaveformTransformationMarkers::OverwriteSoundWaveData(USoundWave& InOutSoundWave)
+{
+	uint32 TransformationMarkerCount = 0;
+	uint32 ThisTransformationNumber = 0;
+
+	for (TObjectPtr<UWaveformTransformationBase> Transformation : InOutSoundWave.Transformations)
+	{
+		TObjectPtr<UWaveformTransformationMarkers> MarkersTransformation = Cast<UWaveformTransformationMarkers>(Transformation);
+		if (MarkersTransformation)
+		{
+			TransformationMarkerCount++;
+			if (this == MarkersTransformation)
+			{
+				ThisTransformationNumber = TransformationMarkerCount;
+				break;
+			}
+		}
+	}
+
+	check(ThisTransformationNumber > 0);
+
+	// If this is the first Marker transformation, overwrite the soundwave original SoundWaveCuePoints
+	if (ThisTransformationNumber == 1)
+	{
+		InOutSoundWave.SetSoundWaveCuePoints(Markers->WaveCues);
+	}
+	// Handles the case if there are multiple marker transformations to preserve all marker "tracks"
+	else
+	{
+		TArray<FSoundWaveCuePoint> SoundWaveCuePoints(InOutSoundWave.GetSoundWaveCuePoints());
+		SoundWaveCuePoints.Append(Markers->WaveCues);
+		
+		InOutSoundWave.SetSoundWaveCuePoints(SoundWaveCuePoints);
+	}
+}
+
 void UWaveCueArray::PostEditChangeProperty(FPropertyChangedEvent& PropertyChangedEvent)
 {
 	Super::PostEditChangeProperty(PropertyChangedEvent);
@@ -72,4 +109,4 @@ void UWaveCueArray::Reset()
 	WaveCues.Empty();
 	bIsInitialized = false;
 }
-#endif
\ No newline at end of file
+#endif
diff --git a/Engine/Plugins/Editor/WaveformEditor/Source/WaveformTransformations/Public/WaveformTransformationMarkers.h b/Engine/Plugins/Editor/WaveformEditor/Source/WaveformTransformations/Public/WaveformTransformationMarkers.h
--- a/Engine/Plugins/Editor/WaveformEditor/Source/WaveformTransformations/Public/WaveformTransformationMarkers.h
+++ b/Engine/Plugins/Editor/WaveformEditor/Source/WaveformTransformations/Public/WaveformTransformationMarkers.h
@@ -6,6 +6,8 @@
 
 #include "WaveformTransformationMarkers.generated.h"
 
+class USoundWave;
+
 //Used to make CuePoints work with PropertyHandles
 UCLASS(BlueprintType, DefaultToInstanced)
 class WAVEFORMTRANSFORMATIONS_API UWaveCueArray : public UObject
@@ -23,6 +25,7 @@ public:
 	TArray<FSoundWaveCuePoint> WaveCues;
 
 private:
+	UPROPERTY()
 	bool bIsInitialized = false;
 };
 
@@ -46,4 +49,8 @@ public:
 
 	void UpdateConfiguration(FWaveTransformUObjectConfiguration& InOutConfiguration) override;
 	void OverwriteTransformation() override;
+
+#if WITH_EDITOR
+	void OverwriteSoundWaveData(USoundWave& InOutSoundWave) override;
+#endif //WITH_EDITOR
 };
diff --git a/Engine/Source/Editor/UnrealEd/Private/EditorExporters.cpp b/Engine/Source/Editor/UnrealEd/Private/EditorExporters.cpp
--- a/Engine/Source/Editor/UnrealEd/Private/EditorExporters.cpp
+++ b/Engine/Source/Editor/UnrealEd/Private/EditorExporters.cpp
@@ -190,7 +190,7 @@ bool USoundExporterWAV::ExportBinary( UObject* Object, const TCHAR* Type, FArchi
 	// Assert that's the case
 	if(ensure(Ar.IsSaving()))
 	{	
-		TFuture<FSharedBuffer> FutureBuffer = Sound->RawData.GetPayload(); // This will block, but for now that's ok.
+		TFuture<FSharedBuffer> FutureBuffer = Sound->RawData.GetPayload(USoundWave::FEditorAudioBulkData::EPayloadFlags::ReturnCopyWithTransformationData); // This will block, but for now that's ok.
 		void* RawWaveData = (void*)FutureBuffer.Get().GetData();
 		Ar.Serialize(RawWaveData, FutureBuffer.Get().GetSize());
 	}
@@ -298,7 +298,7 @@ bool USoundSurroundExporterWAV::ExportBinary( UObject* Object, const TCHAR* Type
 		if (ensure(Ar.IsSaving()))
 		{
 			// Non-const pointer for the sake of calling serialize, so assert that we are writing an archive.
-			TFuture<FSharedBuffer> FutureBuffer = Sound->RawData.GetPayload();
+			TFuture<FSharedBuffer> FutureBuffer = Sound->RawData.GetPayload(USoundWave::FEditorAudioBulkData::EPayloadFlags::ReturnCopyWithTransformationData);
 			uint8* RawWaveData = (uint8*) FutureBuffer.Get().GetData();
 
 			if( Sound->ChannelSizes[ FileIndex ] )
diff --git a/Engine/Source/Runtime/AudioExtensions/Public/IWaveformTransformation.h b/Engine/Source/Runtime/AudioExtensions/Public/IWaveformTransformation.h
--- a/Engine/Source/Runtime/AudioExtensions/Public/IWaveformTransformation.h
+++ b/Engine/Source/Runtime/AudioExtensions/Public/IWaveformTransformation.h
@@ -14,6 +14,8 @@
 
 #include "IWaveformTransformation.generated.h"
 
+class USoundWave;
+
 namespace Audio
 {
 	// information about the current state of the wave file we are transforming
@@ -111,6 +113,10 @@ public:
 	virtual void UpdateConfiguration(FWaveTransformUObjectConfiguration& InOutConfiguration) {};
 	virtual void OverwriteTransformation() {};
 	
+#if WITH_EDITOR
+	virtual void OverwriteSoundWaveData(USoundWave& InOutSoundWave) {};
+#endif //WITH_EDITOR
+
 	virtual bool IsEditorOnly() const override { return true; }
 };
 
diff --git a/Engine/Source/Runtime/AudioMixer/Private/SoundFileIO/SoundFile.h b/Engine/Source/Runtime/AudioMixer/Private/SoundFileIO/SoundFile.h
--- a/Engine/Source/Runtime/AudioMixer/Private/SoundFileIO/SoundFile.h
+++ b/Engine/Source/Runtime/AudioMixer/Private/SoundFileIO/SoundFile.h
@@ -83,6 +83,90 @@ namespace Audio
 		void* DataPtr = nullptr;
 	};
 
+	/**
+	 * FSoundFileCuePoint which maps to libsndfile SF_CUE_POINT struct.
+	 */
+	struct FSoundFileCuePoint
+	{
+		int32 CueId = 0;        
+		uint32 Position = 0;    
+		int32 ChunkId = 0;      
+		int32 ChunkStart = 0;   
+		int32 BlockStart = 0;   
+		uint32 SampleOffset = 0;
+		char Name[256];         
+	};
+
+	/**
+	 * FSoundFileCues which maps to libsndfile SF_CUES struct.
+	 */
+	struct FSoundFileCues
+	{
+		uint32 CueCount = 0;
+		FSoundFileCuePoint CuePoints[100];
+	};
+	
+	// ChunkID: 'cue ' 
+	// A cue chunk specifies one or more sample offsets which are often used to mark noteworthy sections of audio. For example, 
+	// the beginning and end of a verse in a song may have cue points to make them easier to find. The cue chunk is optional and 
+	// if included, a single cue chunk should specify all cue points for the "WAVE" chunk. 
+	// No more than one cue chunk is allowed in a "WAVE" chunk.
+	struct FRiffCueChunk
+	{
+		uint32 ChunkID;			// 'cue '
+		uint32 ChunkDataSize;	// Depends on the number of cue points
+		uint32 NumCuePoints;	// Number of cue points in the list
+	};
+	
+	struct FRiffCuePointChunk
+	{
+		uint32 CueID;			// Unique ID value for the cue point
+		uint32 Position;		// Play order position
+		uint32 DataChunkID;		// RIFF ID of corresponding data chunk
+		uint32 ChunkStart;		// Byte offset of data chunk
+		uint32 BlockStart;		// Byte offset of sample of first channel
+		uint32 SampleOffset;	// Byte offset to sample byte of first channel
+	};
+	
+	struct FRiffSampleChunk
+	{
+		uint32 ChunkID;				// 'smpl'
+		uint32 ChunkDataSize;		// Depends on the number of sample loops
+		uint32 ManufacturerCode;	// The MIDI Manufacturers Association manufacturer code
+		uint32 Product;				// The Product / Model ID of the target device, specific to the manufacturer
+		uint32 SamplePeriod;		// The period of one sample in nanoseconds.
+		uint32 MidiUnityNote;		// The MIDI note that will play when this sample is played at its current pitch
+		uint32 MidiPitchFraction;	// The fraction of a semitone up from the specified note. 
+		uint32 SmpteFormat;			// The SMPTE format. Possible values: 0, 24, 25, 29, 30
+		uint32 SmpteOffset;			// Specifies a time offset for the sample, if the sample should start at a later time and not immediately.
+		uint32 NumSampleLoops;		// Number of sample loops contained in this chunks data
+		uint32 NumSampleDataBytes;	// Number of bytes of optional sampler specific data that follows the sample loops. zero if there is no such data.
+	};
+	
+	struct FRiffSampleLoopChunk
+	{
+		uint32 LoopID;			// A unique ID of the loop, which could be a cue point
+		uint32 LoopType;		// The loop type. 0: Forward Looping, 1: Ping-Pong, 2: Backward, 3-31: future standard types. >=32: manufacturer specific types
+		uint32 StartFrame;		// Start point of the loop in samples
+		uint32 EndFrame;		// End point of the loop in samples. The end sample is also played.
+		uint32 Fraction;		// The resolution at which this loop should be fine tuned.
+		uint32 NumPlayTimes;	// The number of times to play the loop. A value of zero means inifity. In a Midi sampler, that may mean infinite sustain.
+	};
+	
+	struct FRiffListChunk
+	{
+		uint32 ChunkID;			// 'list'
+		uint32 ChunkDataSize;	// Depends on contained text
+		uint32 TypeID;			// always 'adtl'
+	};
+	
+	struct FRiffLabelChunk
+	{
+		uint32 ChunkID;			// 'labl'
+		uint32 ChunkDataSize;	// depends on contained text
+		uint32 CuePointID;		// Cue Point ID associated with the label
+	};
+
 	/**
 	 * FSoundFileChunkInfoWrapper wraps FSoundFileChunkInfo and manages
 	 * chunk data memory.
@@ -180,7 +264,7 @@ namespace Audio
 		virtual ESoundFileError::Type ReadSamples(float* DataPtr, SoundFileCount NumSamples, SoundFileCount& OutNumSamplesRead) = 0;
 		virtual ESoundFileError::Type ReadSamples(double* DataPtr, SoundFileCount NumSamples, SoundFileCount& OutNumSamplesRead) = 0;
 		virtual ESoundFileError::Type GetDescription(FSoundFileDescription& OutputDescription, TArray<ESoundFileChannelMap::Type>& OutChannelMap) = 0;
-		virtual ESoundFileError::Type GetOptionalChunks(FSoundFileChunkArray& OutChunkInfoArray) = 0;
+		virtual ESoundFileError::Type GetOptionalChunks(FSoundFileChunkArray& OutChunkInfoArray, const TSet<uint32>& ChunkIdsToSkip = {}) = 0;
 	};
 
 	class ISoundFileWriter
@@ -197,5 +281,7 @@ namespace Audio
 		virtual ESoundFileError::Type WriteSamples(const double* DataPtr, SoundFileCount NumSamples, SoundFileCount& OutNumSampleWritten) = 0;
 		virtual ESoundFileError::Type GetData(TArray<uint8>** OutData) = 0;
 		virtual ESoundFileError::Type WriteOptionalChunks(const FSoundFileChunkArray& ChunkInfoArray) = 0;
+		virtual ESoundFileError::Type WriteByteArrayChunk(TArray<uint8>& InSoundFileChunk) = 0;
+		virtual ESoundFileError::Type WriteCueCommandData(FSoundFileCues& InSoundFileData) = 0;
 	};
 } // namespace Audio
\ No newline at end of file
diff --git a/Engine/Source/Runtime/AudioMixer/Private/SoundFileIO/SoundFileIO.cpp b/Engine/Source/Runtime/AudioMixer/Private/SoundFileIO/SoundFileIO.cpp
--- a/Engine/Source/Runtime/AudioMixer/Private/SoundFileIO/SoundFileIO.cpp
+++ b/Engine/Source/Runtime/AudioMixer/Private/SoundFileIO/SoundFileIO.cpp
@@ -18,7 +18,7 @@
 
 namespace Audio::SoundFileUtils
 {
-	static void CopyOptionalWavChunks(TSharedPtr<ISoundFileReader>& InSoundDataReader, const int32 InInputFormat, TSharedPtr<ISoundFileWriter>& InSoundFileWriter, const int32 InOutputFormat);
+	static void CopyOptionalWavChunks(TSharedPtr<ISoundFileReader>& InSoundDataReader, const int32 InInputFormat, TSharedPtr<ISoundFileWriter>& InSoundFileWriter, const int32 InOutputFormat, const TSet<uint32>& ChunkIdsToSkip = {});
 
 	bool AUDIOMIXER_API InitSoundFileIOManager()
 	{
@@ -194,14 +194,14 @@ namespace Audio::SoundFileUtils
 		return true;
 	}
 
-	void CopyOptionalWavChunks(TSharedPtr<ISoundFileReader>& InSoundDataReader, const int32 InInputFormat, TSharedPtr<ISoundFileWriter>& InSoundFileWriter, const int32 InOutputFormat)
+	void CopyOptionalWavChunks(TSharedPtr<ISoundFileReader>& InSoundDataReader, const int32 InInputFormat, TSharedPtr<ISoundFileWriter>& InSoundFileWriter, const int32 InOutputFormat, const TSet<uint32>& ChunkIdsToSkip)
 	{
 		// libsndfile only supports chunk operations with wave file formats
 		if ((InInputFormat & ESoundFileFormat::WAV) && (InOutputFormat & ESoundFileFormat::WAV))
 		{
 			// Get the optional chunks from the input data
 			FSoundFileChunkArray OptionalChunks;
-			ESoundFileError::Type Error = InSoundDataReader->GetOptionalChunks(OptionalChunks);
+			ESoundFileError::Type Error = InSoundDataReader->GetOptionalChunks(OptionalChunks, ChunkIdsToSkip);
 			if (Error != ESoundFileError::Type::NONE)
 			{
 				UE_LOG(LogAudioMixer, Error, TEXT("Error encountered while reading optional chunk data...skipping"));
@@ -217,4 +217,352 @@ namespace Audio::SoundFileUtils
 			}
 		}
 	}
+
+	static void WriteUInt32ToByteArray(TArray<uint8>& InOutByteArray, int32& InOutIndex, const uint32 InValue)
+	{
+		check(InOutByteArray.Num() >= InOutIndex + 4);
+		
+		InOutByteArray[InOutIndex++] = (uint8)(InValue >> 0);
+		InOutByteArray[InOutIndex++] = (uint8)(InValue >> 8);
+		InOutByteArray[InOutIndex++] = (uint8)(InValue >> 16);
+		InOutByteArray[InOutIndex++] = (uint8)(InValue >> 24);
+	}
+
+	bool AUDIOMIXER_API CreateCueAndSampleChunks(const TArray<uint8>& InAudioData, TArray<uint8>& OutWaveData, const TArray<FWaveCue>& InWaveCues, const TArray<FWaveSampleLoop>& InSampleLoops)
+	{
+		const FSoundFileConvertFormat ConvertFormat = FSoundFileConvertFormat::CreateDefault();
+
+		FSoundFileIOManager SoundIOManager;
+		TSharedPtr<ISoundFileReader> InputSoundDataReader = SoundIOManager.CreateSoundDataReader();
+		check(InputSoundDataReader);
+
+		ESoundFileError::Type Error = InputSoundDataReader->Init(&InAudioData);
+		if (Error != ESoundFileError::Type::NONE)
+		{
+			UE_LOG(LogAudioMixer, Error, TEXT("Error encountered while initializing InputSoundDataReader... returning false"));
+	
+			return false;
+		}
+
+		TArray<ESoundFileChannelMap::Type> ChannelMap;
+
+		FSoundFileDescription InputDescription;
+		InputSoundDataReader->GetDescription(InputDescription, ChannelMap);
+
+		FSoundFileDescription NewSoundFileDescription;
+		NewSoundFileDescription.NumChannels = InputDescription.NumChannels;
+		NewSoundFileDescription.NumFrames = InputDescription.NumFrames;
+		NewSoundFileDescription.FormatFlags = ConvertFormat.Format;
+		NewSoundFileDescription.SampleRate = InputDescription.SampleRate;
+		NewSoundFileDescription.NumSections = InputDescription.NumSections;
+		NewSoundFileDescription.bIsSeekable = InputDescription.bIsSeekable;
+
+		TSharedPtr<ISoundFileWriter> SoundFileWriter = SoundIOManager.CreateSoundFileWriter();
+		check(SoundFileWriter);
+		
+		Error = SoundFileWriter->Init(NewSoundFileDescription, ChannelMap, ConvertFormat.EncodingQuality);
+		if (Error != ESoundFileError::Type::NONE)
+		{
+			UE_LOG(LogAudioMixer, Error, TEXT("Error encountered while initializing SoundFileWriter... returning false"));
+	
+			return false;
+		}
+
+		//Skip these chunkIds because we are overwriting them
+		TSet<uint32> ChunkIdsToSkip;
+		ChunkIdsToSkip.Reserve(3);
+
+		if(!InWaveCues.IsEmpty())
+		{
+			ChunkIdsToSkip.Add(FWaveModInfo::GetChunkId("cue "));
+			ChunkIdsToSkip.Add(FWaveModInfo::GetChunkId("LIST"));
+			ChunkIdsToSkip.Add(FWaveModInfo::GetChunkId("labl"));
+		}
+
+		if(!InSampleLoops.IsEmpty())
+		{
+			ChunkIdsToSkip.Add(FWaveModInfo::GetChunkId("smpl"));
+		}
+
+		// Copy optional chunks before writing data chunk which libsndfile assumes will be the last chunk
+		CopyOptionalWavChunks(InputSoundDataReader, InputDescription.FormatFlags, SoundFileWriter, NewSoundFileDescription.FormatFlags, ChunkIdsToSkip);
+		
+		if (!InWaveCues.IsEmpty())
+		{
+			FSoundFileCues Cues;
+			
+			// Limit number of cue points to 100 for libsndfile compatibility
+			Cues.CueCount = InWaveCues.Num() <= 100 ? InWaveCues.Num() : 100;
+			
+			for (uint32 Index = 0; Index < Cues.CueCount; ++Index)
+			{
+				const FWaveCue& WaveCue = InWaveCues[Index];
+				
+				FSoundFileCuePoint CuePoint;
+				
+				CuePoint.CueId = WaveCue.CuePointID;
+				CuePoint.ChunkId = FWaveModInfo::GetChunkId("data");
+				CuePoint.Position = WaveCue.Position;
+				CuePoint.ChunkStart = 0;
+				CuePoint.BlockStart = 0;
+				CuePoint.SampleOffset = WaveCue.Position;
+
+				// CuePoint.Name is currently unused by libsndfile, but retaining info in case it is updated
+				FCStringAnsi::Strncpy(CuePoint.Name, TCHAR_TO_ANSI(*WaveCue.Label), sizeof(CuePoint.Name) - 1);
+				
+				Cues.CuePoints[Index] = CuePoint;
+			}
+
+			Error = SoundFileWriter->WriteCueCommandData(Cues);
+			if (Error != ESoundFileError::Type::NONE)
+			{
+				UE_LOG(LogAudioMixer, Error, TEXT("Error encountered while writing cue data... returning false"));
+
+				return false;
+			}
+
+			TArray<uint8> ListChunkData;
+			int32 ChunkDataByteIndex = 0;
+			int32 ListChunkSize = sizeof(FRiffListChunk);
+
+			TMap<uint32, FRiffLabelChunk> LabelChunks;
+
+			for (const FWaveCue& WaveCue : InWaveCues)
+			{
+				// If no label, set label size to 2 so that the label chunk includes a blank label or some DAWs won't render the marker 
+				const uint32 LabelSize = WaveCue.Label.GetCharArray().Num() > 0 ? WaveCue.Label.GetCharArray().Num() : 2;
+
+				FRiffLabelChunk LabelChunk;
+				LabelChunk.ChunkID = FWaveModInfo::GetChunkId("labl");
+				LabelChunk.ChunkDataSize = sizeof(FRiffLabelChunk::CuePointID) + LabelSize;
+				LabelChunk.CuePointID = WaveCue.CuePointID;
+
+				LabelChunks.Emplace(WaveCue.CuePointID, LabelChunk);
+
+				ListChunkSize += LabelChunk.ChunkDataSize + sizeof(FRiffLabelChunk::ChunkID) + sizeof(FRiffLabelChunk::ChunkDataSize);
+
+				// Padding parent chunk size, libsndfile does this for us, but it does not pad the child chunk so we need to account for
+				// the size of the child chunk padding in the TArray.
+				if (LabelSize % 2 != 0)
+				{
+					ListChunkSize++;
+				}
+			}
+
+			ListChunkData.AddZeroed(ListChunkSize);
+
+			// FieldName: chunk ID
+			// FieldSize: 4 bytes
+			// FieldValue: "LIST" (big endian)
+			ListChunkData[ChunkDataByteIndex++] = 'L';
+			ListChunkData[ChunkDataByteIndex++] = 'I';
+			ListChunkData[ChunkDataByteIndex++] = 'S';
+			ListChunkData[ChunkDataByteIndex++] = 'T';
+
+			// FieldName: size
+			// FieldSize: 4 bytes
+			// FieldValue: The size of the LIST chunk(number of bytes) less 8 (less the "chunk ID" and the "size")
+			WriteUInt32ToByteArray(ListChunkData, ChunkDataByteIndex, ListChunkSize - sizeof(FRiffCueChunk::ChunkID)
+				- sizeof(FRiffCueChunk::ChunkDataSize));
+
+			// FieldName: list type ID
+			// FieldSize: 4 bytes
+			// FieldValue: The ASCII character string "adtl"
+			ListChunkData[ChunkDataByteIndex++] = 'a';
+			ListChunkData[ChunkDataByteIndex++] = 'd';
+			ListChunkData[ChunkDataByteIndex++] = 't';
+			ListChunkData[ChunkDataByteIndex++] = 'l';
+
+			for (const FWaveCue& WaveCue : InWaveCues)
+			{
+				FRiffLabelChunk& LabelChunk = LabelChunks[WaveCue.CuePointID];
+
+				WriteUInt32ToByteArray(ListChunkData, ChunkDataByteIndex, LabelChunk.ChunkID);
+				WriteUInt32ToByteArray(ListChunkData, ChunkDataByteIndex, LabelChunk.ChunkDataSize);
+				WriteUInt32ToByteArray(ListChunkData, ChunkDataByteIndex, LabelChunk.CuePointID);
+
+				TArray<TCHAR> Label = WaveCue.Label.GetCharArray();
+
+				// If Label is empty, add two null characters (to ensure character alignment)
+				// so that the label chunk includes a blank label or some DAWs won't render the marker
+				if (Label.IsEmpty())
+				{
+					Label.Append({'\0', '\0'});
+				}
+				
+				for (int32 Index = 0; Index < Label.Num(); Index++)
+				{
+					ListChunkData[ChunkDataByteIndex++] = static_cast<uint8>(Label[Index]);
+				}
+
+				// if the label with null terminator has an odd bytelength, pad with 1 more \0 to ensure 2 byte alignment
+				// libsndfile does not pad children chunks when writing a parent chunk.
+				if (Label.Num() % 2 != 0)
+				{
+					ListChunkData[ChunkDataByteIndex++] = '\0';
+				}
+			}
+
+			Error = SoundFileWriter->WriteByteArrayChunk(ListChunkData);
+			if (Error != ESoundFileError::Type::NONE)
+			{
+				UE_LOG(LogAudioMixer, Error, TEXT("Error encountered while writing list chunk data... returning false"));
+
+				return false;
+			}
+		}
+
+		if (!InSampleLoops.IsEmpty())
+		{
+			TArray<uint8> SampleChunkData;
+			int32 ChunkDataByteIndex = 0;
+			int32 SizeOfSampleChunk = sizeof(FRiffSampleChunk) + sizeof(FRiffSampleLoopChunk) * InSampleLoops.Num();
+			SampleChunkData.AddZeroed(SizeOfSampleChunk);
+
+			// FieldName: chunk ID
+			// FieldSize: 4 bytes
+			// FieldValue: "cue " (big endian)
+			SampleChunkData[ChunkDataByteIndex++] = 's';
+			SampleChunkData[ChunkDataByteIndex++] = 'm';
+			SampleChunkData[ChunkDataByteIndex++] = 'p';
+			SampleChunkData[ChunkDataByteIndex++] = 'l';
+
+			// FieldName: size
+			// FieldSize: 4 bytes
+			// FieldValue: The size of the chunk(number of bytes) less 8 (less the "chunk ID" and the "size")
+			WriteUInt32ToByteArray(SampleChunkData, ChunkDataByteIndex, SizeOfSampleChunk - sizeof(FRiffSampleChunk::ChunkID)
+				- sizeof(FRiffSampleChunk::ChunkDataSize));
+
+			// FieldName: manufacturer 
+			// FieldSize: 4 bytes
+			// FieldValue: The MIDI Manufacturers Association manufacturer code (see MIDI System Exclusive message). A value of zero implies that there is no specific manufacturer. The first byte of the four bytes specifies the number of bytes in the manufacturer code that are relevant (1 or 3). For example, Roland would be specified as 0x01000041 (0x41), where as Microsoft would be 0x03000041 (0x00 0x00 0x41)
+			WriteUInt32ToByteArray(SampleChunkData, ChunkDataByteIndex, 0);
+
+			// FieldName: product 
+			// FieldSize: 4 bytes
+			// FieldValue: The product / model ID of the target device, specific to the manufacturer. A value of zero means no specific product	
+			WriteUInt32ToByteArray(SampleChunkData, ChunkDataByteIndex, 0);
+
+			// FieldName: sample period 
+			// FieldSize: 4 bytes
+			// FieldValue: The period of one sample in nanoseconds. For example, at the sampling rate 44.1 KHz the size of one sample is (1 / 44100) * 1,000,000,000 = 22675 nanoseconds = 0x00005893	
+			check(NewSoundFileDescription.SampleRate > 0);
+			const int32 SamplePeriod = 1000000000 / NewSoundFileDescription.SampleRate;
+			WriteUInt32ToByteArray(SampleChunkData, ChunkDataByteIndex, SamplePeriod);
+
+			// FieldName: MIDI unity note 
+			// FieldSize: 4 bytes
+			// FieldValue: The MIDI note that will play when this sample is played at its current pitch. The values are between 0 and 127 (see MIDI Note On message)
+			WriteUInt32ToByteArray(SampleChunkData, ChunkDataByteIndex, 0);
+
+			// FieldName: MIDI pitch fraction 
+			// FieldSize: 4 bytes
+			// FieldValue: The fraction of a semitone up from the specified note. For example, one-half semitone is 50 cents and will be specified as 0x80
+			WriteUInt32ToByteArray(SampleChunkData, ChunkDataByteIndex, 0);
+
+			// FieldName: SMPTE format 
+			// FieldSize: 4 bytes
+			// FieldValue: The SMPTE format. Possible values are 0, 24, 25, 29, and 30 (see Time division (of a MIDI file))
+			WriteUInt32ToByteArray(SampleChunkData, ChunkDataByteIndex, 0);
+
+			// FieldName: SMPTE offset 
+			// FieldSize: 4 bytes
+			// FieldValue: Specifies a time offset for the sample, if the sample should start at a later time and not immediately. The first byte of this value specifies the number of hours and is in between -23 and 23. The second byte is the number of minutes and is between 0 and 59. The third byte is the number of seconds (0 to 59). The last byte is the number of frames and is between 0 and the frames specified by the SMPTE format. For example, if the SMPTE format is 24, then the number of frames is between 0 and 23
+			WriteUInt32ToByteArray(SampleChunkData, ChunkDataByteIndex, 0);
+
+			// FieldName: number of sample loops 
+			// FieldSize: 4 bytes
+			// FieldValue: Specifies the number of sample loops that are contained in this chunk's data
+			WriteUInt32ToByteArray(SampleChunkData, ChunkDataByteIndex, InSampleLoops.Num());
+
+			// FieldName: sample data 
+			// FieldSize: 4 bytes
+			// FieldValue: The number of bytes of optional sampler specific data that follows the sample loops. If there is no such data, the number of bytes is zero	
+			WriteUInt32ToByteArray(SampleChunkData, ChunkDataByteIndex, 0);
+
+			// FieldName: data
+			// FieldSize: variable
+			// FieldValue: A list of sample loops. Each sample loop uses 24 bytes of data
+			for (const FWaveSampleLoop& SampleLoop : InSampleLoops)
+			{
+				FRiffSampleLoopChunk SampleLoopChunk;
+				SampleLoopChunk.LoopID = SampleLoop.LoopID;
+				SampleLoopChunk.LoopType = 0;
+				SampleLoopChunk.StartFrame = SampleLoop.StartFrame;
+				SampleLoopChunk.EndFrame = SampleLoop.EndFrame;
+				SampleLoopChunk.Fraction = 0;
+				SampleLoopChunk.NumPlayTimes = 0;
+
+				WriteUInt32ToByteArray(SampleChunkData, ChunkDataByteIndex, SampleLoopChunk.LoopID);
+				WriteUInt32ToByteArray(SampleChunkData, ChunkDataByteIndex, SampleLoopChunk.LoopType);
+				WriteUInt32ToByteArray(SampleChunkData, ChunkDataByteIndex, SampleLoopChunk.StartFrame);
+				WriteUInt32ToByteArray(SampleChunkData, ChunkDataByteIndex, SampleLoopChunk.EndFrame);
+				WriteUInt32ToByteArray(SampleChunkData, ChunkDataByteIndex, SampleLoopChunk.Fraction);
+				WriteUInt32ToByteArray(SampleChunkData, ChunkDataByteIndex, SampleLoopChunk.NumPlayTimes);
+			}
+
+			Error = SoundFileWriter->WriteByteArrayChunk(SampleChunkData);
+			if (Error != ESoundFileError::Type::NONE)
+			{
+				UE_LOG(LogAudioMixer, Error, TEXT("Error encountered while writing sample chunk data... returning false"));
+
+				return false;
+			}
+		}
+
+		// Create a buffer to do the processing 
+		SoundFileCount InputSamplesRead = 0;
+		SoundFileCount ProcessBufferSamples = static_cast<SoundFileCount>(1024) * NewSoundFileDescription.NumChannels;
+		TArray<float> ProcessBuffer;
+		ProcessBuffer.Init(0.0f, ProcessBufferSamples);
+
+		// Read the first block of samples
+		Error = InputSoundDataReader->ReadSamples(ProcessBuffer.GetData(), ProcessBufferSamples, InputSamplesRead);
+		if (Error != ESoundFileError::Type::NONE)
+		{
+			UE_LOG(LogAudioMixer, Error, TEXT("Error encountered while ReadingSamples... returning false"));
+
+			return false;
+		}
+
+		SoundFileCount SamplesWritten = 0;
+
+		while (InputSamplesRead == ProcessBuffer.Num())
+		{
+			Error = SoundFileWriter->WriteSamples((const float*)ProcessBuffer.GetData(), InputSamplesRead, SamplesWritten);
+			check(Error == ESoundFileError::Type::NONE);
+			check(SamplesWritten == InputSamplesRead);
+
+			// read more samples
+			Error = InputSoundDataReader->ReadSamples(ProcessBuffer.GetData(), ProcessBufferSamples, InputSamplesRead);
+			if (Error != ESoundFileError::Type::NONE)
+			{
+				UE_LOG(LogAudioMixer, Error, TEXT("Error encountered while ReadingSamples... returning false"));
+
+				return false;
+			}
+
+			for (int32 Sample = 0; Sample < InputSamplesRead; ++Sample)
+			{
+				ProcessBuffer[Sample] = FMath::Clamp(ProcessBuffer[Sample], -1.0f, 1.0f);
+			}
+		}
+
+		// Write final block of samples
+		Error = SoundFileWriter->WriteSamples((const float*)ProcessBuffer.GetData(), InputSamplesRead, SamplesWritten);
+		check(Error == ESoundFileError::Type::NONE);
+
+		// Release the sound file handles as soon as we finished converting the file
+		InputSoundDataReader->Release();
+		SoundFileWriter->Release();
+
+		// Get the raw binary data.....
+		TArray<uint8>* Data = nullptr;
+		SoundFileWriter->GetData(&Data);
+
+		OutWaveData.Init(0, Data->Num());
+		FMemory::Memcpy(OutWaveData.GetData(), (const void*)&(*Data)[0], OutWaveData.Num());
+
+		return true;
+	}
 }
diff --git a/Engine/Source/Runtime/AudioMixer/Private/SoundFileIO/SoundFileIOManagerImpl.cpp b/Engine/Source/Runtime/AudioMixer/Private/SoundFileIO/SoundFileIOManagerImpl.cpp
--- a/Engine/Source/Runtime/AudioMixer/Private/SoundFileIO/SoundFileIOManagerImpl.cpp
+++ b/Engine/Source/Runtime/AudioMixer/Private/SoundFileIO/SoundFileIOManagerImpl.cpp
@@ -43,7 +43,9 @@ namespace Audio
 	static const int32 SET_ENCODING_QUALITY = 0x1300;
 	static const int32 SET_CHANNEL_MAP_INFO = 0x1101;
 	static const int32 GET_CHANNEL_MAP_INFO = 0x1100;
-	static const int32 UPDATE_HEADER_NOW = 0x1060;
+	static const int32 UPDATE_HEADER_NOW = 0x1060; // Always returns 0, ignore return value.
+	static const int32 SET_INSTRUMENT = 0x10D1;
+	static const int32 SET_CUE = 0x10CF; // Returns TRUE if the markers are written to file.
 
 	// Exported SoundFile Functions
 	typedef LibSoundFileHandle*(*SoundFileOpenFuncPtr)(const char* Path, int32 Mode, FSoundFileDescription* Description);
@@ -383,7 +385,7 @@ namespace Audio
 		return ESoundFileError::Type::NONE;
 	}
 
-	static ESoundFileError::Type GetOptionalChunksInternal(LibSoundFileHandle* FileHandle, FSoundFileChunkArray& OutChunkInfoArray)
+	static ESoundFileError::Type GetOptionalChunksInternal(LibSoundFileHandle* FileHandle, FSoundFileChunkArray& OutChunkInfoArray, const TSet<uint32>& ChunkIdsToSkip = {})
 	{
 		// Verify that the necessary library function pointers have been properly set
 		if (SoundFileGetChunkIterator != nullptr && SoundFileGetChunkSize != nullptr &&
@@ -393,6 +395,10 @@ namespace Audio
 
 			for (const uint32 Id : OptionalChunkIds)
 			{
+				if (ChunkIdsToSkip.Contains(Id))
+				{
+					continue;
+				}
 				FSoundFileChunkInfo ChunkLookup;				
 				// Copy chunk ID over. DWORD (4 bytes, each is ANSI char)
 				*reinterpret_cast<uint32*>(ChunkLookup.ChunkId) = Id;
@@ -458,13 +464,8 @@ namespace Audio
 				{
 					UE_LOG(LogAudioMixer, VeryVerbose, TEXT("Wrote ChunkId: %s, chunk DataLength: %d"), *FString((ANSICHAR*)ChunkInfo.GetPtr()->ChunkId), ChunkInfo.GetPtr()->DataLength);
 
-					// Update file header after adding a new chunk
-					Result = SoundFileCommand(FileHandle, UPDATE_HEADER_NOW, nullptr, 0);
-					if (Result)
-					{
-						UE_LOG(LogAudioMixer, Error, TEXT("LibSoundFile failed to update file header; Result = %d"), Result);
-						return ESoundFileError::Type::INVALID_STATE;
-					}
+					// Update file header after adding a new chunk - UPDATE_HEADER_NOW command always returns 0
+					SoundFileCommand(FileHandle, UPDATE_HEADER_NOW, nullptr, 0);
 				}
 			}
 		}
@@ -475,6 +476,92 @@ namespace Audio
 
 		return ESoundFileError::Type::NONE;
 	}
+
+	static ESoundFileError::Type WriteByteArrayChunkInternal(LibSoundFileHandle* FileHandle, TArray<uint8>& InSoundFileChunk)
+	{
+		// Verify that the necessary library function pointers have been properly set
+		if (SoundFileCommand != nullptr && SoundFileSetChunk != nullptr)
+		{
+			check(InSoundFileChunk.Num() > 8)
+			uint8* SoundFileChunk = InSoundFileChunk.GetData();
+			check(SoundFileChunk);
+
+			FSoundFileChunkArray ChunkInfoArray;
+			FSoundFileChunkInfo ChunkInfo;
+
+			ChunkInfo.ChunkId[0] = InSoundFileChunk[0];
+			ChunkInfo.ChunkId[1] = InSoundFileChunk[1];
+			ChunkInfo.ChunkId[2] = InSoundFileChunk[2];
+			ChunkInfo.ChunkId[3] = InSoundFileChunk[3];
+			ChunkInfo.ChunkIdSize = sizeof(uint32);
+
+			ChunkInfo.DataLength = 
+				((uint32)InSoundFileChunk[4] << 0) |
+				((uint32)InSoundFileChunk[5] << 8) |
+				((uint32)InSoundFileChunk[6] << 16)|
+				((uint32)InSoundFileChunk[7] << 24);
+
+			ChunkInfo.DataPtr = SoundFileChunk + 8;
+			
+			// Note, libsndfile uses 4-byte pad when writing chunk data
+			int32 Result = SoundFileSetChunk(FileHandle, &ChunkInfo);
+			if (Result)
+			{
+				UE_LOG(LogAudioMixer, Error, TEXT("LibSoundFile failed to write chunk data; Result = %d"), Result);
+
+				return ESoundFileError::Type::INVALID_STATE;
+			}
+			else
+			{
+				int32 DataLength = ChunkInfo.DataLength;
+				UE_LOG(LogAudioMixer, VeryVerbose, TEXT("Wrote ChunkId: %s, chunk DataLength: %d"), *FString((ANSICHAR*)&ChunkInfo.ChunkId), &DataLength);
+
+				// Update file header after adding a new chunk - UPDATE_HEADER_NOW command always returns 0
+				SoundFileCommand(FileHandle, UPDATE_HEADER_NOW, nullptr, 0);
+			}
+		}
+		else
+		{
+			UE_LOG(LogAudioMixer, Error, TEXT("LibSoundFile wasn't properly loaded with symbols for accessing wav data."));
+
+			return ESoundFileError::Type::INVALID_STATE;
+		}
+
+		return ESoundFileError::Type::NONE;
+	}
+
+	// Only works for SoundFile Commands that return a bool (ex. SET_CUE)
+	static ESoundFileError::Type WriteCommandDataInternal(LibSoundFileHandle* FileHandle, const int32 Command, void* InSoundFileData, size_t InSoundFileDataLength)
+	{
+		// Verify that the necessary library function pointers have been properly set
+		if (SoundFileCommand != nullptr)
+		{
+			check(InSoundFileData);
+
+			if (SoundFileCommand(FileHandle, Command, InSoundFileData, sizeof(FSoundFileCues)))
+			{
+				int32 DataLength = 0;
+				UE_LOG(LogAudioMixer, VeryVerbose, TEXT("SoundFileCommand Completed... Updating Header Now"));
+
+				// Update file header after adding a new chunk - UPDATE_HEADER_NOW command always returns 0
+				SoundFileCommand(FileHandle, UPDATE_HEADER_NOW, nullptr, 0);
+			}
+			else
+			{
+				UE_LOG(LogAudioMixer, Error, TEXT("LibSoundFile failed to write command data; Result = %hs"), SoundFileStrError(FileHandle));
+
+				return ESoundFileError::Type::INVALID_STATE;
+			}
+		}
+		else
+		{
+			UE_LOG(LogAudioMixer, Error, TEXT("LibSoundFile wasn't properly loaded with symbols for accessing wav data."));
+
+			return ESoundFileError::Type::INVALID_STATE;
+		}
+
+		return ESoundFileError::Type::NONE;
+	}
 	
 	static SoundFileCount OnSoundFileGetLengthBytes(void* UserData)
 	{
@@ -769,9 +856,9 @@ namespace Audio
 			return ESoundFileError::Type::NONE;
 		}
 
-		ESoundFileError::Type GetOptionalChunks(FSoundFileChunkArray& OutChunkInfoArray) override
+		ESoundFileError::Type GetOptionalChunks(FSoundFileChunkArray& OutChunkInfoArray, const TSet<uint32>& ChunkIdsToSkip = {}) override
 		{
-			return GetOptionalChunksInternal(FileHandle, OutChunkInfoArray);
+			return GetOptionalChunksInternal(FileHandle, OutChunkInfoArray, ChunkIdsToSkip);
 		}
 
 	private:
@@ -1232,9 +1319,9 @@ namespace Audio
 			return ESoundFileError::Type::NONE;
 		}
 
-		ESoundFileError::Type GetOptionalChunks(FSoundFileChunkArray& OutChunkInfoArray) override
+		ESoundFileError::Type GetOptionalChunks(FSoundFileChunkArray& OutChunkInfoArray, const TSet<uint32>& ChunkIdsToSkip = {}) override
 		{
-			return GetOptionalChunksInternal(FileHandle, OutChunkInfoArray);
+			return GetOptionalChunksInternal(FileHandle, OutChunkInfoArray, ChunkIdsToSkip);
 		}
 
 	private:
@@ -1577,6 +1664,16 @@ namespace Audio
 			return WriteOptionalChunksInternal(FileHandle, ChunkInfoArray);
 		}
 
+		ESoundFileError::Type WriteByteArrayChunk(TArray<uint8>& InSoundFileChunk) override
+		{
+			return WriteByteArrayChunkInternal(FileHandle, InSoundFileChunk);
+		}
+
+		ESoundFileError::Type WriteCueCommandData(FSoundFileCues& InSoundFileData) override
+		{
+			return WriteCommandDataInternal(FileHandle, SET_CUE, &InSoundFileData, sizeof(InSoundFileData));
+		}
+
 	private:
 
 		ESoundFileError::Type SetError(ESoundFileError::Type InError)
diff --git a/Engine/Source/Runtime/AudioMixer/Public/SoundFileIO/SoundFileIO.h b/Engine/Source/Runtime/AudioMixer/Public/SoundFileIO/SoundFileIO.h
--- a/Engine/Source/Runtime/AudioMixer/Public/SoundFileIO/SoundFileIO.h
+++ b/Engine/Source/Runtime/AudioMixer/Public/SoundFileIO/SoundFileIO.h
@@ -2,6 +2,7 @@
 
 #pragma once 
 
+#include "Audio.h"
 #include "CoreMinimal.h"
 
 namespace Audio::SoundFileUtils
@@ -10,4 +11,5 @@ namespace Audio::SoundFileUtils
 	bool AUDIOMIXER_API ShutdownSoundFileIOManager();
 	uint32 AUDIOMIXER_API GetNumSamples(const TArray<uint8>& InAudioData);
 	bool AUDIOMIXER_API ConvertAudioToWav(const TArray<uint8>& InAudioData, TArray<uint8>& OutWaveData);
+	bool AUDIOMIXER_API CreateCueAndSampleChunks(const TArray<uint8>& InAudioData, TArray<uint8>& OutWaveData, const TArray<FWaveCue>& InWaveCues, const TArray<FWaveSampleLoop>& InSampleLoops);
 }
diff --git a/Engine/Source/Runtime/Engine/Classes/Sound/SampleBufferIO.h b/Engine/Source/Runtime/Engine/Classes/Sound/SampleBufferIO.h
--- a/Engine/Source/Runtime/Engine/Classes/Sound/SampleBufferIO.h
+++ b/Engine/Source/Runtime/Engine/Classes/Sound/SampleBufferIO.h
@@ -168,7 +168,8 @@ namespace Audio
 		// at the specified file path relative to the project directory.
 		// This function should only be used in the editor.
 		// If a USoundWave asset already exists 
-		ENGINE_API bool BeginWriteToSoundWave(const FString& FileName, const TSampleBuffer<>& InSampleBuffer, FString InPath, TFunction<void(const USoundWave*)> OnSuccess = [](const USoundWave* ResultingWave) {});
+		// InSoundWaveTransformationTarget is used to overwrite SoundWave data with its own transformations. (ex. marker data)
+		ENGINE_API bool BeginWriteToSoundWave(const FString& FileName, const TSampleBuffer<>& InSampleBuffer, FString InPath, TFunction<void(const USoundWave*)> OnSuccess = [](const USoundWave* ResultingWave) {}, TObjectPtr<USoundWave> InSoundWaveTransformationTarget = nullptr);
 	
 		// This writes out the InSampleBuffer as a wav file at the path specified by FilePath and FileName.
 		// If FilePath is a relative path, it will be relative to the /Saved/BouncedWavFiles folder, otherwise specified absolute path will be used.
diff --git a/Engine/Source/Runtime/Engine/Classes/Sound/SoundWave.h b/Engine/Source/Runtime/Engine/Classes/Sound/SoundWave.h
--- a/Engine/Source/Runtime/Engine/Classes/Sound/SoundWave.h
+++ b/Engine/Source/Runtime/Engine/Classes/Sound/SoundWave.h
@@ -477,6 +477,12 @@ public:
 	/** Filters for the cue points that are _not_ loop regions and returns those as a new array */
 	UFUNCTION(BlueprintCallable, Category = "Audio")
 	ENGINE_API TArray<FSoundWaveCuePoint> GetCuePoints() const;
+	ENGINE_API void SetAllCuePoints(TArray<FSoundWaveCuePoint> InCuePoints) const;
+
+#if WITH_EDITOR
+	ENGINE_API TArray<FSoundWaveCuePoint> GetSoundWaveCuePoints() const;
+	ENGINE_API void SetSoundWaveCuePoints(const TArray<FSoundWaveCuePoint>& InCuePoints);
+#endif //WITH_EDITOR
 
 	/** Filters for the cue points that _are_ loop regions and returns those as a new array */
 	UFUNCTION(BlueprintCallable, Category = "Audio")
@@ -598,6 +604,9 @@ public:
 #if WITH_EDITOR
 	ENGINE_API const FWaveTransformUObjectConfiguration& GetTransformationChainConfig() const;
 	ENGINE_API const FWaveTransformUObjectConfiguration& UpdateTransformations();
+	ENGINE_API TArray<uint8> CreateNewCueData(const char ChunkName[4]);
+	ENGINE_API void OverwriteSoundWaveDataWithTransformationData();
+
 #endif
 
 private:
@@ -896,6 +905,12 @@ public:
 		// resulting in the bulk data not being correlated with our asset.
 		USoundWave* SoundWave;
 
+		enum class EPayloadFlags : uint8
+		{
+			None,
+			ReturnCopyWithTransformationData = 0x01
+		};
+
 		FEditorAudioBulkData()
 			: RawData()
 			, SoundWave()
@@ -909,7 +924,7 @@ public:
 
 		ENGINE_API void CreateFromBulkData(FBulkData& InBulkData, const FGuid& InGuid, UObject* Owner);
 		ENGINE_API void Serialize(FArchive& Ar, UObject* Owner, bool bAllowRegister=true);
-		ENGINE_API TFuture<FSharedBuffer> GetPayload() const;
+		ENGINE_API TFuture<FSharedBuffer> GetPayload(const EPayloadFlags Flags = EPayloadFlags::None) const;
 		ENGINE_API bool HasPayloadData() const;
 		ENGINE_API void UpdatePayload(FSharedBuffer InPayload, UObject* Owner = nullptr);
 
diff --git a/Engine/Source/Runtime/Engine/Private/Audio.cpp b/Engine/Source/Runtime/Engine/Private/Audio.cpp
--- a/Engine/Source/Runtime/Engine/Private/Audio.cpp
+++ b/Engine/Source/Runtime/Engine/Private/Audio.cpp
@@ -1588,6 +1588,11 @@ const TArray<uint32>& FWaveModInfo::GetOptionalWaveChunkIds()
 	return OptionalChunkIds;
 }
 
+const uint32 FWaveModInfo::GetChunkId(const char ChunkID[4])
+{
+	return UE_mmioFOURCC(ChunkID[0], ChunkID[1], ChunkID[2], ChunkID[3]);
+}
+
 bool IsKnownChunkId(uint32 ChunkId)
 {
 	bool bIsKnown = FWaveModInfo::GetRequiredWaveChunkIds().Contains(ChunkId) ||
diff --git a/Engine/Source/Runtime/Engine/Private/SampleBufferIO.cpp b/Engine/Source/Runtime/Engine/Private/SampleBufferIO.cpp
--- a/Engine/Source/Runtime/Engine/Private/SampleBufferIO.cpp
+++ b/Engine/Source/Runtime/Engine/Private/SampleBufferIO.cpp
@@ -160,7 +160,7 @@ namespace Audio
 		return true;
 	}
 
-	bool FSoundWavePCMWriter::BeginWriteToSoundWave(const FString& FileName, const TSampleBuffer<>& InSampleBuffer, FString InPath, TFunction<void(const USoundWave*)> OnSuccess)
+	bool FSoundWavePCMWriter::BeginWriteToSoundWave(const FString& FileName, const TSampleBuffer<>& InSampleBuffer, FString InPath, TFunction<void(const USoundWave*)> OnSuccess, TObjectPtr<USoundWave> InSoundWaveTransformationTarget)
 	{
 		if (!IsDone())
 		{
@@ -199,6 +199,14 @@ namespace Audio
 		CurrentSoundWave = NewObject<USoundWave>(CurrentPackage, *FileName, RF_Public | RF_Standalone);
 		CurrentSoundWave->AddToRoot();
 
+#if WITH_EDITOR
+		if (InSoundWaveTransformationTarget != nullptr)
+		{
+			InSoundWaveTransformationTarget->OverwriteSoundWaveDataWithTransformationData();
+			CurrentSoundWave->SetSoundWaveCuePoints(InSoundWaveTransformationTarget->GetSoundWaveCuePoints());
+		}
+#endif //WITH_EDITOR
+
 		CurrentBuffer = InSampleBuffer;
 
 		// TODO: Support writing to multi-channel sound waves. For now, we mix down sound waves to stereo:
diff --git a/Engine/Source/Runtime/Engine/Private/SoundWave.cpp b/Engine/Source/Runtime/Engine/Private/SoundWave.cpp
--- a/Engine/Source/Runtime/Engine/Private/SoundWave.cpp
+++ b/Engine/Source/Runtime/Engine/Private/SoundWave.cpp
@@ -37,6 +37,8 @@
 #include "Templates/UnrealTemplate.h"
 #include "UObject/ObjectSaveContext.h"
 #include "ISoundWaveCloudStreaming.h"
+#include "SoundFileIO/SoundFileIO.h"
+
 #if WITH_EDITORONLY_DATA
 #include "uewav.h"
 #endif
@@ -1452,6 +1454,24 @@ TArray<FSoundWaveCuePoint> USoundWave::GetCuePoints() const
 	return SoundWaveDataPtr->GetCuePoints();
 }
 
+void USoundWave::SetAllCuePoints(TArray<FSoundWaveCuePoint> InCuePoints) const
+{
+	SoundWaveDataPtr->SetAllCuePoints(InCuePoints);
+}
+
+#if WITH_EDITOR
+TArray<FSoundWaveCuePoint> USoundWave::GetSoundWaveCuePoints() const
+{
+	return CuePoints;
+}
+
+void USoundWave::SetSoundWaveCuePoints(const TArray<FSoundWaveCuePoint>& InCuePoints)
+{
+	CuePoints = InCuePoints;
+	SetAllCuePoints(InCuePoints);
+}
+#endif //WITH_EDITOR
+
 TArray<FSoundWaveCuePoint> USoundWave::GetLoopRegions() const
 {
 	return SoundWaveDataPtr->GetLoopRegions();
@@ -4601,6 +4621,63 @@ const FWaveTransformUObjectConfiguration& USoundWave::UpdateTransformations()
 	return TransformationChainConfig;
 }
 
+void USoundWave::OverwriteSoundWaveDataWithTransformationData()
+{
+	for (TObjectPtr<UWaveformTransformationBase>& Transformation : Transformations)
+	{
+		if (Transformation)
+		{
+			Transformation->OverwriteSoundWaveData(*this);
+		}
+	}
+
+	//Overwrite the wave data for cue points and sample loops
+	TArray<FWaveCue> WaveCues;
+	TArray<FWaveSampleLoop> WaveSampleLoops;
+
+	for (const FSoundWaveCuePoint& CuePoint : CuePoints)
+	{
+		FWaveCue WaveCue;
+		WaveCue.CuePointID = CuePoint.CuePointID;
+		WaveCue.Position = CuePoint.FramePosition;
+		WaveCue.Label = CuePoint.Label;
+		WaveCue.SampleLength = CuePoint.FrameLength;
+
+		WaveCues.Emplace(WaveCue);
+		
+		if (CuePoint.IsLoopRegion())
+		{
+			FWaveSampleLoop WaveSampleLoop;
+			WaveSampleLoop.LoopID = CuePoint.CuePointID;
+			WaveSampleLoop.StartFrame = CuePoint.FramePosition;
+			WaveSampleLoop.EndFrame = CuePoint.FramePosition + CuePoint.FrameLength;
+
+			WaveSampleLoops.Emplace(WaveSampleLoop);
+		}
+	}
+
+	// Some DAWs require the Cues and SampleLoops to be sorted by frame position and not CueID.
+	Algo::Sort(WaveCues, [](const FWaveCue& WaveCueA, const FWaveCue& WaveCueB)
+		{
+			return WaveCueA.Position < WaveCueB.Position;
+		});
+	Algo::Sort(WaveSampleLoops, [](const FWaveSampleLoop& LoopA, const FWaveSampleLoop& LoopB)
+		{
+			return LoopA.StartFrame < LoopB.StartFrame;
+		});
+
+	TArray<uint8> WaveFileData;
+	TFuture<FSharedBuffer> FutureBuffer = RawData.GetPayload();
+	TArray<uint8> InWaveFileData;
+	InWaveFileData.Append((uint8*)FutureBuffer.Get().GetData(), FutureBuffer.Get().GetSize());
+
+	Audio::SoundFileUtils::CreateCueAndSampleChunks(InWaveFileData, WaveFileData, WaveCues, WaveSampleLoops);
+
+	// Update the RawData with the new payload
+	FSharedBuffer Buffer = FSharedBuffer::Clone(WaveFileData.GetData(), WaveFileData.Num());
+	RawData.UpdatePayload(Buffer);
+}
+
 #endif // #if WITH_EDITOR
 
 void USoundWave::CacheInheritedLoadingBehavior() const
@@ -5070,7 +5147,7 @@ bool USoundWave::HasError() const
 }
 
 #if WITH_EDITORONLY_DATA
-TFuture<FSharedBuffer> USoundWave::FEditorAudioBulkData::GetPayload() const
+TFuture<FSharedBuffer> USoundWave::FEditorAudioBulkData::GetPayload(const EPayloadFlags Flags) const
 {
 	if (SoundWave == nullptr)
 	{
@@ -5085,7 +5162,18 @@ TFuture<FSharedBuffer> USoundWave::FEditorAudioBulkData::GetPayload() const
 	// sound wave wanders off and then drop all of this is a task.
 
 	// Here's our blockage on a potential rehydration.
-	FSharedBuffer BulkDataBuffer = RawData.GetPayload().Get();
+	FSharedBuffer BulkDataBuffer;
+	if (Flags == EPayloadFlags::ReturnCopyWithTransformationData)
+	{
+		TObjectPtr<USoundWave> DuplicateSoundWave = DuplicateObject<USoundWave>(SoundWave, SoundWave->GetOuter());
+
+		DuplicateSoundWave->OverwriteSoundWaveDataWithTransformationData();
+		BulkDataBuffer = DuplicateSoundWave->RawData.RawData.GetPayload().Get();
+	}
+	else
+	{
+		BulkDataBuffer = RawData.GetPayload().Get();
+	}
 
 	// We need it mutable. Unless the bulk data system holds a copy this should be a steal.
 	FUniqueBuffer MutablePayloadBuffer = BulkDataBuffer.MoveToUnique();
diff --git a/Engine/Source/Runtime/Engine/Public/Audio.h b/Engine/Source/Runtime/Engine/Public/Audio.h
--- a/Engine/Source/Runtime/Engine/Public/Audio.h
+++ b/Engine/Source/Runtime/Engine/Public/Audio.h
@@ -982,6 +982,7 @@ public:
 	/** Wave Chunk Id utils */
 	ENGINE_API static const TArray<uint32>& GetRequiredWaveChunkIds();
 	ENGINE_API static const TArray<uint32>& GetOptionalWaveChunkIds();
+	ENGINE_API static const uint32 GetChunkId(const char ChunkID[4]);
 
 	// Read headers and load all info pointers in WaveModInfo. 
 	// Returns 0 if invalid data encountered.
