clean up FMixerSourceVoiceBuffer, AudioMixerSourceVoice::OnMixBus

#rnx

diff --git a/Engine/Source/Runtime/AudioMixer/Private/AudioMixerSourceBuffer.cpp b/Engine/Source/Runtime/AudioMixer/Private/AudioMixerSourceBuffer.cpp
--- a/Engine/Source/Runtime/AudioMixer/Private/AudioMixerSourceBuffer.cpp
+++ b/Engine/Source/Runtime/AudioMixer/Private/AudioMixerSourceBuffer.cpp
@@ -12,10 +12,10 @@ extern int32 bDecodersCanReturnPartialBuffersCvar;
 
 namespace Audio
 {
-	bool FRawPCMDataBuffer::GetNextBuffer(FMixerSourceVoiceBuffer* OutSourceBufferPtr, const uint32 NumSampleToGet)
+	bool FRawPCMDataBuffer::GetNextBuffer(TArrayView<float> OutSourceBufferPtr, const uint32 NumSampleToGet)
 	{
 		// TODO: support loop counts
-		float* OutBufferPtr = OutSourceBufferPtr->AudioData.GetData();
+		float* OutBufferPtr = OutSourceBufferPtr.GetData();
 		int16* DataPtr = (int16*)Data;
 
 		if (LoopCount == Audio::LOOP_FOREVER)
@@ -122,10 +122,10 @@ namespace Audio
 		const uint32 TotalSamples = MONO_PCM_BUFFER_SAMPLES * NumChannels;
 		for (int32 BufferIndex = 0; BufferIndex < Audio::MAX_BUFFERS_QUEUED; ++BufferIndex)
 		{
-			SourceVoiceBuffers.Add(MakeShared<FMixerSourceVoiceBuffer, ESPMode::ThreadSafe>());
+			SourceVoiceBuffers.Add(MakeShared<FAlignedFloatBuffer, ESPMode::ThreadSafe>());
 
 			// Prepare the memory to fit the max number of samples
-			SourceVoiceBuffers[BufferIndex]->AudioData.Reset(TotalSamples);
+			SourceVoiceBuffers[BufferIndex]->Reset(TotalSamples);
 		}
 	}
 
@@ -227,7 +227,7 @@ namespace Audio
 		return 0;
 	}
 
-	TSharedPtr<FMixerSourceVoiceBuffer, ESPMode::ThreadSafe> FMixerSourceBuffer::GetNextBuffer()
+	TSharedPtr<FAlignedFloatBuffer, ESPMode::ThreadSafe> FMixerSourceBuffer::GetNextBuffer()
 	{
 		FScopeTryLock Lock(&SoundWaveCritSec);
 		if (!Lock.IsLocked())
@@ -235,7 +235,7 @@ namespace Audio
 			return nullptr;
 		}
 
-		TSharedPtr<FMixerSourceVoiceBuffer, ESPMode::ThreadSafe> NewBufferPtr;
+		TSharedPtr<FAlignedFloatBuffer, ESPMode::ThreadSafe> NewBufferPtr;
 		BufferQueue.Dequeue(NewBufferPtr);
 		--NumBuffersQeueued;
 		return NewBufferPtr;
@@ -261,10 +261,10 @@ namespace Audio
 		int16* RawPCMBufferDataPtr = (int16*)RawPCMDataBuffer.Data;
 
 		// Prepare the buffer for the PCM submission
-		SourceVoiceBuffers[0]->AudioData.Reset(NumSamplesPerBuffer);
-		SourceVoiceBuffers[0]->AudioData.AddZeroed(NumSamplesPerBuffer);
+		SourceVoiceBuffers[0]->Reset(NumSamplesPerBuffer);
+		SourceVoiceBuffers[0]->AddZeroed(NumSamplesPerBuffer);
 
-		RawPCMDataBuffer.GetNextBuffer(SourceVoiceBuffers[0].Get(), NumSamplesPerBuffer);
+		RawPCMDataBuffer.GetNextBuffer(*SourceVoiceBuffers[0], NumSamplesPerBuffer);
 
 		SubmitBuffer(SourceVoiceBuffers[0]);
 
@@ -291,14 +291,14 @@ namespace Audio
 				// Prepare the precache buffer memory
 				for (int32 i = 0; i < 2; ++i)
 				{
-					SourceVoiceBuffers[i]->AudioData.Reset();
-					SourceVoiceBuffers[i]->AudioData.AddZeroed(NumSamples);
+					SourceVoiceBuffers[i]->Reset();
+					SourceVoiceBuffers[i]->AddZeroed(NumSamples);
 				}
 
 				int16* CachedBufferPtr0 = (int16*)CachedRealtimeFirstBuffer.GetData();
 				int16* CachedBufferPtr1 = (int16*)(CachedRealtimeFirstBuffer.GetData() + BufferSize);
-				float* AudioData0 = SourceVoiceBuffers[0]->AudioData.GetData();
-				float* AudioData1 = SourceVoiceBuffers[1]->AudioData.GetData();
+				float* AudioData0 = SourceVoiceBuffers[0]->GetData();
+				float* AudioData1 = SourceVoiceBuffers[1]->GetData();
 
 				Audio::ArrayPcm16ToFloat(MakeArrayView(CachedBufferPtr0, NumSamples), MakeArrayView(AudioData0, NumSamples));
 				Audio::ArrayPcm16ToFloat(MakeArrayView(CachedBufferPtr1, NumSamples), MakeArrayView(AudioData1, NumSamples));
@@ -335,8 +335,8 @@ namespace Audio
 	{
 		const int32 MaxSamples = MONO_PCM_BUFFER_SAMPLES * NumChannels;
 
-		SourceVoiceBuffers[BufferIndex]->AudioData.Reset();
-		SourceVoiceBuffers[BufferIndex]->AudioData.AddUninitialized(MaxSamples);
+		SourceVoiceBuffers[BufferIndex]->Reset();
+		SourceVoiceBuffers[BufferIndex]->AddUninitialized(MaxSamples);
 
 		if (bProcedural)
 		{
@@ -358,7 +358,7 @@ namespace Audio
 					NewTaskData.ProceduralSoundWave = SoundWave;
 				}
 
-				NewTaskData.AudioData = SourceVoiceBuffers[BufferIndex]->AudioData.GetData();
+				NewTaskData.AudioData = SourceVoiceBuffers[BufferIndex]->GetData();
 				NewTaskData.NumSamples = MaxSamples;
 				NewTaskData.NumChannels = NumChannels;
 				AsyncTaskStartTimeInCycles = FPlatformTime::Cycles64();
@@ -373,13 +373,13 @@ namespace Audio
 			check(RawPCMDataBuffer.Data != nullptr);
 
 			// Read the next raw PCM buffer into the source buffer index. This converts raw PCM to float.
-			return RawPCMDataBuffer.GetNextBuffer(SourceVoiceBuffers[BufferIndex].Get(), MaxSamples);
+			return RawPCMDataBuffer.GetNextBuffer(*SourceVoiceBuffers[BufferIndex], MaxSamples);
 		}
 
 		// Handle the case that the decoder has an error and can't continue.
 		if (InDecoder && InDecoder->HasError())
 		{
-			FMemory::Memzero(SourceVoiceBuffers[BufferIndex]->AudioData.GetData(), MaxSamples * sizeof(float));
+			FMemory::Memzero(SourceVoiceBuffers[BufferIndex]->GetData(), MaxSamples * sizeof(float));
 
 			FScopeTryLock Lock(&SoundWaveCritSec);
 			if (Lock.IsLocked() && SoundWave)
@@ -395,7 +395,7 @@ namespace Audio
 		check(InDecoder != nullptr);
 
 		FDecodeAudioTaskData NewTaskData;
-		NewTaskData.AudioData = SourceVoiceBuffers[BufferIndex]->AudioData.GetData();
+		NewTaskData.AudioData = SourceVoiceBuffers[BufferIndex]->GetData();
 		NewTaskData.DecompressionState = InDecoder;
 		NewTaskData.BufferType = BufferType;
 		NewTaskData.NumChannels = NumChannels;
@@ -437,7 +437,7 @@ namespace Audio
 			}
 		}
 
-		if (SourceVoiceBuffers[CurrentBuffer]->AudioData.Num() > 0)
+		if (SourceVoiceBuffers[CurrentBuffer]->Num() > 0)
 		{
 			SubmitBuffer(SourceVoiceBuffers[CurrentBuffer]);
 		}
@@ -462,7 +462,7 @@ namespace Audio
 					
 					if (bDecodersCanReturnPartialBuffersCvar)
 					{
-						SourceVoiceBuffers[CurrentBuffer]->AudioData.SetNum(TaskResult.NumSamplesWritten);
+						SourceVoiceBuffers[CurrentBuffer]->SetNum(TaskResult.NumSamplesWritten);
 #if ENABLE_AUDIO_DEBUG
 						double AudioDuration = static_cast<double>(TaskResult.NumSamplesWritten) / static_cast<double>(FMath::Max(1, NumChannels * SampleRate));
 						UpdateCPUCoreUtilization(TaskResult.CPUDuration, AudioDuration);
@@ -483,7 +483,7 @@ namespace Audio
 					FProceduralAudioTaskResults TaskResult;
 					AsyncRealtimeAudioTask->GetResult(TaskResult);
 
-					SourceVoiceBuffers[CurrentBuffer]->AudioData.SetNum(TaskResult.NumSamplesWritten);
+					SourceVoiceBuffers[CurrentBuffer]->SetNum(TaskResult.NumSamplesWritten);
 					bIsFinishedOrLooped = TaskResult.bIsFinished;
 #if ENABLE_AUDIO_DEBUG
 					double AudioDuration = static_cast<double>(TaskResult.NumSamplesWritten) / static_cast<double>(FMath::Max(1, NumChannels * SampleRate));
@@ -542,7 +542,7 @@ namespace Audio
 		}
 	}
 
-	void FMixerSourceBuffer::SubmitBuffer(TSharedPtr<FMixerSourceVoiceBuffer, ESPMode::ThreadSafe> InSourceVoiceBuffer)
+	void FMixerSourceBuffer::SubmitBuffer(TSharedPtr<FAlignedFloatBuffer, ESPMode::ThreadSafe> InSourceVoiceBuffer)
 	{
 		NumBuffersQeueued++;
 		BufferQueue.Enqueue(InSourceVoiceBuffer);
diff --git a/Engine/Source/Runtime/AudioMixer/Private/AudioMixerSourceBuffer.h b/Engine/Source/Runtime/AudioMixer/Private/AudioMixerSourceBuffer.h
--- a/Engine/Source/Runtime/AudioMixer/Private/AudioMixerSourceBuffer.h
+++ b/Engine/Source/Runtime/AudioMixer/Private/AudioMixerSourceBuffer.h
@@ -10,8 +10,6 @@
 
 namespace Audio
 {
-	struct FMixerSourceVoiceBuffer;
-
 	static const int32 MAX_BUFFERS_QUEUED = 3;
 	static const int32 LOOP_FOREVER = -1;
 
@@ -23,7 +21,7 @@ namespace Audio
 		uint32 CurrentSample;
 		uint32 NumSamples;
 
-		bool GetNextBuffer(FMixerSourceVoiceBuffer* OutSourceBufferPtr, const uint32 NumSampleToGet);
+		bool GetNextBuffer(TArrayView<float> OutSourceBufferPtr, const uint32 NumSampleToGet);
 
 		FRawPCMDataBuffer()
 			: Data(nullptr)
@@ -88,7 +86,7 @@ namespace Audio
 		int32 GetNumBuffersQueued() const;
 
 		// Returns the next enqueued buffer, returns nullptr if no buffers enqueued
-		TSharedPtr<FMixerSourceVoiceBuffer, ESPMode::ThreadSafe> GetNextBuffer();
+		TSharedPtr<FAlignedFloatBuffer, ESPMode::ThreadSafe> GetNextBuffer();
 
 		// Returns if buffer looped
 		bool DidBufferLoop() const { return bLoopCallback; }
@@ -139,15 +137,15 @@ namespace Audio
 		void SubmitInitialRealtimeBuffers();
 		void SubmitRealTimeSourceData(const bool bFinishedOrLooped);
 		void ProcessRealTimeSource();
-		void SubmitBuffer(TSharedPtr<FMixerSourceVoiceBuffer, ESPMode::ThreadSafe> InSourceVoiceBuffer);
+		void SubmitBuffer(TSharedPtr<FAlignedFloatBuffer, ESPMode::ThreadSafe> InSourceVoiceBuffer);
 		void DeleteDecoder();
 
 
 		int32 NumBuffersQeueued;
 		FRawPCMDataBuffer RawPCMDataBuffer;
 
-		TArray<TSharedPtr<FMixerSourceVoiceBuffer, ESPMode::ThreadSafe>> SourceVoiceBuffers;
-		TQueue<TSharedPtr<FMixerSourceVoiceBuffer, ESPMode::ThreadSafe>> BufferQueue;
+		TArray<TSharedPtr<FAlignedFloatBuffer, ESPMode::ThreadSafe>> SourceVoiceBuffers;
+		TQueue<TSharedPtr<FAlignedFloatBuffer, ESPMode::ThreadSafe>> BufferQueue;
 		int32 CurrentBuffer;
 		// SoundWaves are only set for procedural sound waves
 		USoundWave* SoundWave;
diff --git a/Engine/Source/Runtime/AudioMixer/Private/AudioMixerSourceManager.cpp b/Engine/Source/Runtime/AudioMixer/Private/AudioMixerSourceManager.cpp
--- a/Engine/Source/Runtime/AudioMixer/Private/AudioMixerSourceManager.cpp
+++ b/Engine/Source/Runtime/AudioMixer/Private/AudioMixerSourceManager.cpp
@@ -2105,7 +2105,7 @@ namespace Audio
 				bReadCurrentFrame = false;
 
 				AUDIO_MIXER_CHECK(SourceInfo.CurrentPCMBuffer.IsValid());
-				const float* AudioData = SourceInfo.CurrentPCMBuffer->AudioData.GetData();
+				const float* AudioData = SourceInfo.CurrentPCMBuffer->GetData();
 				const int32 CurrentSampleIndex = SourceInfo.CurrentFrameIndex * NumChannels;
 
 				for (int32 Channel = 0; Channel < NumChannels; ++Channel)
@@ -2163,10 +2163,10 @@ namespace Audio
 		if (SourceInfo.CurrentPCMBuffer.IsValid())
 		{
 			// Grab the float PCM audio data (which could be a new audio chunk from previous ReadSourceFrame call)
-			const float* AudioData = SourceInfo.CurrentPCMBuffer->AudioData.GetData();
+			const float* AudioData = SourceInfo.CurrentPCMBuffer->GetData();
 			const int32 CurrentSampleIndex = SourceInfo.CurrentFrameIndex * NumChannels;
 			const int32 NextSampleIndex = (SourceInfo.CurrentFrameIndex + 1)  * NumChannels;
-			const int32 AudioDataNum = SourceInfo.CurrentPCMBuffer->AudioData.Num();
+			const int32 AudioDataNum = SourceInfo.CurrentPCMBuffer->Num();
 
 			if(ensureAlwaysMsgf(AudioDataNum >= NextSampleIndex + NumChannels
 				, TEXT("Bailing due to bad CurrentPCMBuffer:  AudioData.Num() = %i, NextSampleIndex = %i, NumChannels = %i"), AudioDataNum, NextSampleIndex, NumChannels))
@@ -2179,7 +2179,7 @@ namespace Audio
 						SourceInfo.NextFrameValues[Channel] = AudioData[NextSampleIndex + Channel];
 					}
 				}
-				else if (NextSampleIndex != SourceInfo.CurrentPCMBuffer->AudioData.Num())
+				else if (NextSampleIndex != SourceInfo.CurrentPCMBuffer->Num())
 				{
 					for (int32 Channel = 0; Channel < NumChannels; ++Channel)
 					{
@@ -2492,7 +2492,7 @@ namespace Audio
 
 						const int32 DestinationOffset = NumSamples - FramesToWrite * SourceInfo.NumInputChannels;
 						float* CopyDestination = SourceInfo.PreDistanceAttenuationBuffer.GetData() + DestinationOffset;
-						const float* CopySource = SourceInfo.CurrentPCMBuffer->AudioData.GetData() + SourceInfo.CurrentFrameIndex * SourceInfo.NumInputChannels;
+						const float* CopySource = SourceInfo.CurrentPCMBuffer->GetData() + SourceInfo.CurrentFrameIndex * SourceInfo.NumInputChannels;
 
 						FMemory::Memcpy(CopyDestination, CopySource, SamplesToCopy * sizeof(float));
 
diff --git a/Engine/Source/Runtime/AudioMixer/Private/AudioMixerSourceManager.h b/Engine/Source/Runtime/AudioMixer/Private/AudioMixerSourceManager.h
--- a/Engine/Source/Runtime/AudioMixer/Private/AudioMixerSourceManager.h
+++ b/Engine/Source/Runtime/AudioMixer/Private/AudioMixerSourceManager.h
@@ -71,14 +71,6 @@ namespace Audio
 	class ISourceListener;
 	class FMixerSourceSubmixOutputBuffer;
 
-	/** Struct defining a source voice buffer. */
-	struct FMixerSourceVoiceBuffer
-	{
-		/** PCM float data. */
-		FAlignedFloatBuffer AudioData;
-	};
-
-
 	class ISourceListener
 	{
 	public:
@@ -447,7 +439,7 @@ namespace Audio
 			ISourceListener* SourceListener;
 
 			// Data used for rendering sources
-			TSharedPtr<FMixerSourceVoiceBuffer, ESPMode::ThreadSafe> CurrentPCMBuffer;
+			TSharedPtr<Audio::FAlignedFloatBuffer, ESPMode::ThreadSafe> CurrentPCMBuffer;
 
 			// The post-attenuation source buffer, used to send audio to submixes
 			Audio::FAlignedFloatBuffer SourceBuffer;
@@ -608,7 +600,7 @@ namespace Audio
 				if (CurrentPCMBuffer)
 				{
 					check(NumInputChannels != 0);
-					return CurrentPCMBuffer->AudioData.Num() / NumInputChannels;
+					return CurrentPCMBuffer->Num() / NumInputChannels;
 				}
 				else
 				{
diff --git a/Engine/Source/Runtime/AudioMixer/Private/AudioMixerSourceVoice.cpp b/Engine/Source/Runtime/AudioMixer/Private/AudioMixerSourceVoice.cpp
--- a/Engine/Source/Runtime/AudioMixer/Private/AudioMixerSourceVoice.cpp
+++ b/Engine/Source/Runtime/AudioMixer/Private/AudioMixerSourceVoice.cpp
@@ -462,16 +462,4 @@ namespace Audio
 	{
 		return bEnableBaseSubmix || bEnableSubmixSends;
 	}
-
-	void FMixerSourceVoice::OnMixBus(FMixerSourceVoiceBuffer* OutMixerSourceBuffer)
-	{
-		AUDIO_MIXER_CHECK_AUDIO_PLAT_THREAD(MixerDevice);
-
-		check(OutMixerSourceBuffer->AudioData.Num() > 0);
-
-		for (int32 i = 0; i < OutMixerSourceBuffer->AudioData.Num(); ++i)
-		{
-			OutMixerSourceBuffer->AudioData[i] = 0.0f;
-		}
-	}
 }
diff --git a/Engine/Source/Runtime/AudioMixer/Private/AudioMixerSourceVoice.h b/Engine/Source/Runtime/AudioMixer/Private/AudioMixerSourceVoice.h
--- a/Engine/Source/Runtime/AudioMixer/Private/AudioMixerSourceVoice.h
+++ b/Engine/Source/Runtime/AudioMixer/Private/AudioMixerSourceVoice.h
@@ -7,7 +7,6 @@
 
 namespace Audio
 {
-	struct FMixerSourceVoiceBuffer;
 	struct FMixerSourceVoiceFilterParams;
 	struct FMixerSourceVoiceInitParams;
 	class FMixerDevice;
@@ -153,9 +152,6 @@ namespace Audio
 		// Set the source bus send levels
 		void SetAudioBusSendInfo(EBusSendType InBusSendType, uint32 AudioBusId, float BusSendLevel);
 
-		// Called when the source is a bus and needs to mix other sources together to generate output
-		void OnMixBus(FMixerSourceVoiceBuffer* OutMixerSourceBuffer);
-
 	private:
 
 		friend class FMixerSourceManager;
diff --git a/Engine/Source/Runtime/AudioMixer/Private/SoundWaveDecoder.cpp b/Engine/Source/Runtime/AudioMixer/Private/SoundWaveDecoder.cpp
--- a/Engine/Source/Runtime/AudioMixer/Private/SoundWaveDecoder.cpp
+++ b/Engine/Source/Runtime/AudioMixer/Private/SoundWaveDecoder.cpp
@@ -201,7 +201,7 @@ namespace Audio
 			{
 				bReadCurrentFrame = false;
 
-				const float* AudioData = SourceInfo.CurrentPCMBuffer->AudioData.GetData();
+				const float* AudioData = SourceInfo.CurrentPCMBuffer->GetData();
 
 				if (!AudioData)
 				{
@@ -230,7 +230,7 @@ namespace Audio
 					return;
 				}
 
-				SourceInfo.CurrentAudioChunkNumFrames = SourceInfo.CurrentPCMBuffer->AudioData.Num() / SourceInfo.NumSourceChannels;
+				SourceInfo.CurrentAudioChunkNumFrames = SourceInfo.CurrentPCMBuffer->Num() / SourceInfo.NumSourceChannels;
 
 				if (bReadCurrentFrame)
 				{
@@ -251,7 +251,7 @@ namespace Audio
 			bCurrentFrameOutOfRange = SourceInfo.CurrentFrameIndex >= SourceInfo.CurrentAudioChunkNumFrames;
 		}
 
-		const float* AudioData = SourceInfo.CurrentPCMBuffer->AudioData.GetData();
+		const float* AudioData = SourceInfo.CurrentPCMBuffer->GetData();
 
 		if (!AudioData)
 		{
diff --git a/Engine/Source/Runtime/AudioMixer/Public/SoundWaveDecoder.h b/Engine/Source/Runtime/AudioMixer/Public/SoundWaveDecoder.h
--- a/Engine/Source/Runtime/AudioMixer/Public/SoundWaveDecoder.h
+++ b/Engine/Source/Runtime/AudioMixer/Public/SoundWaveDecoder.h
@@ -19,7 +19,6 @@ namespace Audio
 {
 	class FMixerBuffer;
 	class FMixerSourceBuffer;
-	struct FMixerSourceVoiceBuffer;
 
 	struct FDecodingSoundSourceHandle
 	{
@@ -171,7 +170,7 @@ namespace Audio
 			TArray<float> NextFrameValues;
 
 			// The current decoded PCM buffer we are reading from
-			TSharedPtr<FMixerSourceVoiceBuffer, ESPMode::ThreadSafe> CurrentPCMBuffer;
+			TSharedPtr<FAlignedFloatBuffer, ESPMode::ThreadSafe> CurrentPCMBuffer;
 
 			// If this sound is done (has decoded all data)
 			bool bIsDone;
