Source Rendering Restructuring - Add individual resampler to each source

diff --git a/Engine/Source/Runtime/AudioMixer/Private/AudioMixerSourceManager.cpp b/Engine/Source/Runtime/AudioMixer/Private/AudioMixerSourceManager.cpp
--- a/Engine/Source/Runtime/AudioMixer/Private/AudioMixerSourceManager.cpp
+++ b/Engine/Source/Runtime/AudioMixer/Private/AudioMixerSourceManager.cpp
@@ -13,6 +13,7 @@
 #include "DSP/FloatArrayMath.h"
 #include "IAudioExtensionPlugin.h"
 #include "AudioMixer.h"
+#include "samplerate.h"
 #include "Sound/SoundModulationDestination.h"
 #include "SoundFieldRendering.h"
 #include "ProfilingDebugging/CsvProfiler.h"
@@ -350,6 +351,17 @@ namespace Audio
 		}
 
 		FPlatformProcess::ReturnSynchEventToPool(CommandsProcessedEvent);
+
+		for (FSourceInfo& SourceInfo : SourceInfos)
+		{
+			ReleaseResampler(SourceInfo);
+		}
+
+		for (int32 i = 0; i < ResamplerPool.Num(); ++i)
+		{
+			src_delete(ResamplerPool[i]);
+		}
+		ResamplerPool.Reset();
 	}
 
 	void FMixerSourceManager::Init(const FSourceManagerInitParams& InitParams)
@@ -454,6 +466,8 @@ namespace Audio
 			SourceInfo.NumInputChannels = 0;
 			SourceInfo.NumPostEffectChannels = 0;
 			SourceInfo.NumInputFrames = 0;
+
+			SourceInfo.Resampler = nullptr;
 		}
 		
 		GameThreadInfo.bIsBusy.AddDefaulted(NumTotalSources);
@@ -1006,6 +1020,8 @@ namespace Audio
 			SourceInfo.NumInputChannels = InitParams.NumInputChannels;
 			SourceInfo.NumInputFrames = InitParams.NumInputFrames;
 
+			ReleaseResampler(SourceInfo);
+
 			// init and zero-out buffers
 			const int32 BufferSize = NumOutputFrames * InitParams.NumInputChannels;
 			SourceInfo.PreEffectBuffer.Reset();
@@ -2463,7 +2479,7 @@ namespace Audio
 								GameThreadInfo.CPUCoreUtilization[SourceId] = SourceInfo.MixerSourceBuffer->GetCPUCoreUtilization();
 #endif // if ENABLE_AUDIO_DEBUG
 								GameThreadInfo.RelativeRenderCost[SourceId] = SourceInfo.MixerSourceBuffer->GetRelativeRenderCost();
-								
+
 								SourceInfo.MixerSourceBuffer->OnBufferEnd();
 							}
 
@@ -2484,23 +2500,47 @@ namespace Audio
 							}
 						}
 
-						// Copy as much data as possible from CurrentPCMBuffer to PreDistanceAttenuationBuffer
-						check(SourceInfo.CurrentPCMBuffer.IsValid());
-
-						const int32 FramesToCopy = FMath::Min(FramesToWrite, CurrentChunkFrames - SourceInfo.CurrentFrameIndex);
-						const int32 SamplesToCopy = FramesToCopy * SourceInfo.NumInputChannels;
+						if (FinalPitch != 1 && !SourceInfo.Resampler)
+						{
+							AcquireResampler(SourceInfo);
+						}
 
+						// Pump as many samples as possibe from CurrentPCMBuffer to PreDistanceAttenuationBuffer
+						check(SourceInfo.CurrentPCMBuffer.IsValid());
 						const int32 DestinationOffset = NumSamples - FramesToWrite * SourceInfo.NumInputChannels;
 						float* CopyDestination = SourceInfo.PreDistanceAttenuationBuffer.GetData() + DestinationOffset;
 						const float* CopySource = SourceInfo.CurrentPCMBuffer->GetData() + SourceInfo.CurrentFrameIndex * SourceInfo.NumInputChannels;
 
-						FMemory::Memcpy(CopyDestination, CopySource, SamplesToCopy * sizeof(float));
-
-						// These will be scaled when resampling is being done.
-						SourceInfo.CurrentFrameIndex += FramesToCopy;
-						SourceInfo.NumFramesPlayed += FramesToCopy;
-
-						FramesToWrite -= FramesToCopy;
+						if (SourceInfo.Resampler)
+						{
+							// Resample if required
+							SRC_DATA SampleRateParams{};
+							SampleRateParams.data_in = CopySource;
+							SampleRateParams.data_out = CopyDestination;
+							SampleRateParams.input_frames = CurrentChunkFrames - SourceInfo.CurrentFrameIndex;
+							SampleRateParams.output_frames = FramesToWrite;
+							check(FinalPitch > 0);
+							SampleRateParams.src_ratio = 1 / FinalPitch;
+
+							src_process(SourceInfo.Resampler, &SampleRateParams);
+
+							// Update bookkeeping
+							SourceInfo.CurrentFrameIndex += SampleRateParams.input_frames_used;
+							SourceInfo.NumFramesPlayed += SampleRateParams.input_frames_used;
+							FramesToWrite -= SampleRateParams.output_frames_gen;
+						}
+						else
+						{
+							// No resampling necessary, just copy
+							const int32 FramesToCopy = FMath::Min(FramesToWrite, CurrentChunkFrames - SourceInfo.CurrentFrameIndex);
+							const int32 SamplesToCopy = FramesToCopy * SourceInfo.NumInputChannels;
+							FMemory::Memcpy(CopyDestination, CopySource, SamplesToCopy * sizeof(float));
+
+							// Update bookkeeping
+							SourceInfo.CurrentFrameIndex += FramesToCopy;
+							SourceInfo.NumFramesPlayed += FramesToCopy;
+							FramesToWrite -= FramesToCopy;
+						}
 					}
 				}
 			}
@@ -3881,4 +3921,49 @@ namespace Audio
 			UE_LOG(LogAudioMixer, Warning, TEXT("***** ThreadStackWalkAndDump for ThreadId(%lu) ******\n%s"), InThreadId, ANSI_TO_TCHAR(StackTrace));
 		}
 	}
+
+	void FMixerSourceManager::AcquireResampler(FSourceInfo& SourceInfo)
+	{
+		check(!SourceInfo.Resampler);
+
+		{
+			UE::TScopeLock<UE::FSpinLock> ScopeLock(ResamplerPoolLock);
+
+			for (int32 i = 0; i < ResamplerPool.Num(); ++i)
+			{
+				SRC_STATE* Resampler = ResamplerPool[i];
+				check(Resampler);
+				if (src_get_channels(Resampler) == SourceInfo.NumInputChannels)
+				{
+					// Reuse this resampler
+					ResamplerPool.RemoveAtSwap(i);
+					SourceInfo.Resampler = Resampler;
+					break;
+				}
+			}
+		}
+
+		if (SourceInfo.Resampler)
+		{
+			// Reset and return the reusable resampler
+			src_reset(SourceInfo.Resampler);
+			return;
+		}
+
+		// Allocate a new resampler
+		int Error = 0;
+		SourceInfo.Resampler = src_new(SRC_LINEAR, SourceInfo.NumInputChannels, &Error);
+		checkf(SourceInfo.Resampler, TEXT("src_new failed with error code %i"), Error);
+	}
+
+	void FMixerSourceManager::ReleaseResampler(FSourceInfo& SourceInfo)
+	{
+		SRC_STATE* Resampler = SourceInfo.Resampler;
+		if (Resampler)
+		{
+			SourceInfo.Resampler = nullptr;
+			UE::TScopeLock<UE::FSpinLock> ScopeLock(ResamplerPoolLock);
+			ResamplerPool.Add(Resampler);
+		}
+	}
 }
diff --git a/Engine/Source/Runtime/AudioMixer/Private/AudioMixerSourceManager.h b/Engine/Source/Runtime/AudioMixer/Private/AudioMixerSourceManager.h
--- a/Engine/Source/Runtime/AudioMixer/Private/AudioMixerSourceManager.h
+++ b/Engine/Source/Runtime/AudioMixer/Private/AudioMixerSourceManager.h
@@ -15,6 +15,7 @@
 #include "DSP/ParamInterpolator.h"
 #include "IAudioExtensionPlugin.h"
 #include "ISoundfieldFormat.h"
+#include "Misc/SpinLock.h"
 #include "Sound/SoundModulationDestination.h"
 #include "Sound/QuartzQuantizationUtilities.h"
 #include "Stats/Stats.h"
@@ -62,6 +63,9 @@ enum class ESourceManagerRenderThreadPhase: uint8
 	Finished,
 };
 
+// Forward declaration from libSampleRate
+typedef struct SRC_STATE_tag SRC_STATE;
+
 namespace Audio
 {
 	class FMixerSubmix;
@@ -564,6 +568,9 @@ namespace Audio
 			// ID for associated Audio Component if there is one, 0 otherwise
 			uint64 AudioComponentID;
 
+			// Resampler instance
+			SRC_STATE* Resampler;
+
 			FORCEINLINE void ResetModulators(const Audio::FDeviceId InDeviceId)
 			{
 				VolumeModulation.Init(InDeviceId, FName("Volume"), false /* bInIsBuffered */, true /* bInValueLinear */);
@@ -620,6 +627,9 @@ namespace Audio
 		void LogCallstacks();
 		void LogCallstack(uint32 InThreadId);
 
+		void AcquireResampler(FSourceInfo& SourceInfo);
+		void ReleaseResampler(FSourceInfo& SourceInfo);
+
 		// Array of listener transforms
 		TArray<FTransform> ListenerTransforms;
 
@@ -685,6 +695,9 @@ namespace Audio
 
 		TMpscQueue<FPendingAudioBusConnection> PendingAudioBusConnections;
 
+		UE::FSpinLock ResamplerPoolLock;
+		TArray<SRC_STATE*> ResamplerPool;
+
 		friend class FMixerSourceVoice;
 	};
 }
