Adding a mute cvar and a mute command line argument for cases where users want to mute the audio but allow the audio engine to continue rendering audio and running (for reactivity, accurate perf metrics, etc)

[RN] New -muteaudio command line tool and au.MuteAudio cvar apply a mute on the final output buffer without modify audio engine behavior (audio analysis, memory, or performance profile)

diff --git a/Engine/Source/Runtime/AudioMixer/Private/AudioMixerDevice.cpp b/Engine/Source/Runtime/AudioMixer/Private/AudioMixerDevice.cpp
--- a/Engine/Source/Runtime/AudioMixer/Private/AudioMixerDevice.cpp
+++ b/Engine/Source/Runtime/AudioMixer/Private/AudioMixerDevice.cpp
@@ -114,6 +114,14 @@ FAutoConsoleVariableRef CVarAudioMixerPatchBufferBlocks(
 	TEXT("Determines the number of blocks that fit in a patch buffer."),
 	ECVF_Default);
 
+static int32 bMuteAudioCVar = 0;
+FAutoConsoleVariableRef CVarMuteAudio(
+	TEXT("au.MuteAudio"),
+	bMuteAudioCVar,
+	TEXT("Mutes the audio output in the final output without effecting audio engine behavior.\n")
+	TEXT("1: Mutes the final submix. 0: Does not mute the final submix."),
+	ECVF_Cheat);
+
 // Link to "Audio" profiling category
 CSV_DECLARE_CATEGORY_MODULE_EXTERN(AUDIOMIXERCORE_API, Audio);
 
@@ -517,6 +525,14 @@ namespace Audio
 		{
 			UE_LOG(LogAudioMixer, Display, TEXT("Initializing audio mixer using platform API: '%s'"), *AudioMixerPlatform->GetPlatformApi());
 
+			// Issue a log if the audio engine will be muted.
+#if !UE_BUILD_SHIPPING
+			if (bMuteAudioCVar || FApp::IsAudioMuted())
+			{
+				UE_LOG(LogAudioMixer, Display, TEXT("Audio output will be muted. If this not intentional, check for au.MuteAudio cvar or -muteaudio command line."));
+			}
+#endif // !UE_BUILD_SHIPPING
+
 			const UAudioSettings* AudioSettings = GetDefault<UAudioSettings>();
 			MonoChannelUpmixMethod = AudioSettings->MonoChannelUpmixMethod;
 			PanningMethod = AudioSettings->PanningMethod;
@@ -949,6 +965,14 @@ namespace Audio
 			{
 				// Process the audio output from the master submix
 				MainSubmixPtr->ProcessAudio(Output);
+
+#if !UE_BUILD_SHIPPING
+				// This is done at the very end here to avoid changing the behavior in any other way (analysis, profiling, etc)
+				if ((bMuteAudioCVar || FApp::IsAudioMuted()))
+				{
+					FMemory::Memzero((void*)Output.GetData(), sizeof(float) * Output.Num());
+				}
+#endif // !UE_BUILD_SHIPPING
 			}
 		}
 
diff --git a/Engine/Source/Runtime/Core/Public/Misc/App.h b/Engine/Source/Runtime/Core/Public/Misc/App.h
--- a/Engine/Source/Runtime/Core/Public/Misc/App.h
+++ b/Engine/Source/Runtime/Core/Public/Misc/App.h
@@ -439,6 +439,23 @@ public:
 #endif // UE_SERVER
 	}
 
+	/**
+	 * Checks whether this application should mute the audio output (-muteaudio)
+	 * This is a safer alternative to -nosound for users who want to mute the audio of the game but keep the game rendering sound. 
+	 * Unless disabling the audio engine is the intention, this is a preferred method to mute the output as it avoid breaking audio analysis-dependent features (e.g. gameplay)
+	 * and it prevents content creators from creating content with an incorrect performance and memory profile.
+	 *
+	 * The mute is applied only in non-shipping builds and applied at the final step in the final output submix.
+	 * Alternative method to do this it to enable the cvar, "au.MuteAudio 1"
+	 * 
+	 * @return true if the application will mute audio, false otherwise.
+	 */
+	INLINE_CANEVERRENDER static bool IsAudioMuted()
+	{
+		static bool bMuteAudio = FParse::Param(FCommandLine::Get(), TEXT("muteaudio"));
+		return bMuteAudio;
+	}
+
 	/**
 	 * Checks whether this application has been installed.
 	 *
