Audio command queue spam reduction
- predicate some stuff with WITH_EDITOR as its only useful for live edit changes and a waste in runtime
- add float compare (instead of binary equality), with tolerances that change with command queue fullness
- command queue diagnostic cvars to log commands when we fill to N commands.

diff --git a/Engine/Source/Runtime/AudioMixer/Private/AudioMixerSource.cpp b/Engine/Source/Runtime/AudioMixer/Private/AudioMixerSource.cpp
--- a/Engine/Source/Runtime/AudioMixer/Private/AudioMixerSource.cpp
+++ b/Engine/Source/Runtime/AudioMixer/Private/AudioMixerSource.cpp
@@ -827,10 +827,8 @@ namespace Audio
 		WaveInstance = InWaveInstance;
 
 		LPFFrequency = MAX_FILTER_FREQUENCY;
-		LastLPFFrequency = FLT_MAX;
 
 		HPFFrequency = 0.0f;
-		LastHPFFrequency = FLT_MAX;
 
 		bIsDone = false;
 
@@ -1544,17 +1542,8 @@ namespace Audio
 		// Update the default LPF filter frequency
 		SetFilterFrequency();
 
-		if (LastLPFFrequency != LPFFrequency)
-		{
-			MixerSourceVoice->SetLPFFrequency(LPFFrequency);
-			LastLPFFrequency = LPFFrequency;
-		}
-
-		if (LastHPFFrequency != HPFFrequency)
-		{
-			MixerSourceVoice->SetHPFFrequency(HPFFrequency);
-			LastHPFFrequency = HPFFrequency;
-		}
+		MixerSourceVoice->SetLPFFrequency(LPFFrequency);
+		MixerSourceVoice->SetHPFFrequency(HPFFrequency);
 
 		check(WaveInstance);
 		FActiveSound* ActiveSound = WaveInstance->ActiveSound;
@@ -1652,10 +1641,12 @@ namespace Audio
 		{
 			MixerSourceVoice->ClearSubmixSendInfo(i);
 		}
-	
 		MixerSourceVoice->SetEnablement(WaveInstance->bEnableBusSends, WaveInstance->bEnableBaseSubmix, WaveInstance->bEnableSubmixSends);
-
+#if WITH_EDITOR
+		// The following can spam to the command queue. But is mostly here so that the editor live edits are immedately heard
+		// For anything less than editor this is perf waste, so predicate this only to be run in editor. 
 		MixerSourceVoice->SetSourceBufferListener(WaveInstance->SourceBufferListener, WaveInstance->bShouldSourceBufferListenerZeroBuffer);
+#endif	// WITH_EDITOR
 	}
 
 	void FMixerSource::UpdateModulation()
diff --git a/Engine/Source/Runtime/AudioMixer/Private/AudioMixerSourceManager.cpp b/Engine/Source/Runtime/AudioMixer/Private/AudioMixerSourceManager.cpp
--- a/Engine/Source/Runtime/AudioMixer/Private/AudioMixerSourceManager.cpp
+++ b/Engine/Source/Runtime/AudioMixer/Private/AudioMixerSourceManager.cpp
@@ -130,6 +130,34 @@ FAutoConsoleVariableRef LogEveryAudioThreadCommand(
 	TEXT("Extremely verbose logging of each Audio Thread command caller and it's execution time"),
 	ECVF_Default);
 
+static int32 LogCmdQueueWhenNumberReachedCVar = 0;
+FAutoConsoleVariableRef CVarLogCmdQueueWhenNumberReached(
+	TEXT("au.debug.LogCmdQueueWhenReached"),
+	LogCmdQueueWhenNumberReachedCVar,
+	TEXT("When number of commands in queue is reached, it will dump the command list to the log."),
+	ECVF_Cheat);
+
+static int32 NumCmdsConsideredFullCVar = 1000;
+FAutoConsoleVariableRef CVarNumCmdsConsideredFullCVar(
+	TEXT("au.NumCmdsConsideredFullCVar"),
+	NumCmdsConsideredFullCVar,
+	TEXT("Num of commands in the queue is considered full"),
+	ECVF_Default);
+
+static float FloatCompareMinScaleCVar = 0.001;
+FAutoConsoleVariableRef CVarFloatCompareMinScale(
+	TEXT("au.FloatCompareMinScale"),
+	FloatCompareMinScaleCVar,
+	TEXT("Minimum parameter float compare tolerance"),
+	ECVF_Default);
+
+static float FloatCompareMaxScaleCVar = 0.1;
+FAutoConsoleVariableRef CVarFloatCompareMaxScale(
+	TEXT("au.FloatCompareMaxScale"),
+	FloatCompareMaxScaleCVar,
+	TEXT("Maximum parameter float compare tolerance"),
+	ECVF_Default);
+
 // +/- 4 Octaves (default)
 static float MaxModulationPitchRangeFreqCVar = 16.0f;
 static float MinModulationPitchRangeFreqCVar = 0.0625f;
@@ -184,6 +212,9 @@ FAutoConsoleVariableRef CVarAudioMixerPerSourceResampling(
 // Disable subframe timing logic
 #define AUDIO_SUBFRAME_ENABLED 0
 
+// Global count for _all_ queued commands (across all devices).
+TRACE_DECLARE_ATOMIC_INT_COUNTER(AudioMixerSourceManager_TotalQdCmds, TEXT("AudioSourceManager Qd Cmds"));
+
 // Define profiling for source manager. 
 DEFINE_STAT(STAT_AudioMixerHRTF);
 DEFINE_STAT(STAT_AudioMixerSourceBuffers);
@@ -619,6 +650,7 @@ namespace Audio
 
 						AudioCommand();
 						NumCommands.Decrement();
+						TRACE_COUNTER_DECREMENT(AudioMixerSourceManager_TotalQdCmds);
 					}
 
 					NextCommandBuffer.SourceCommandQueue.Reset();
@@ -2628,6 +2660,21 @@ namespace Audio
 		}
 	}
 
+	float FMixerSourceManager::GetFloatCompareTolerance() const
+	{
+		const float Percent = GetCommandQueueFillPercentage();
+		return FMath::IsNearlyEqual(Percent, 1.0f) ? FloatCompareMaxScaleCVar : FloatCompareMinScaleCVar;
+	}
+
+	float FMixerSourceManager::GetCommandQueueFillPercentage() const
+	{
+		if (NumCmdsConsideredFullCVar == 0)
+		{
+			return 0.f;
+		}
+		return static_cast<float>(FMath::Min(NumCommands.GetValue(), NumCmdsConsideredFullCVar) / NumCmdsConsideredFullCVar);
+	}
+
 	// ctor
 	FMixerSourceManager::FAudioMixerThreadCommand::FAudioMixerThreadCommand(TFunction<void()>&& InFunction, const char* InDebugString, bool bInDeferExecution)
 	: Function(MoveTemp(InFunction))
@@ -3647,6 +3694,21 @@ namespace Audio
 			// always add commands to the buffer. If we're not going to assert, might as well chug along and hope we can recover!
 			Commands.SourceCommandQueue.Add(AudioCommand);
 			NumCommands.Increment();
+			TRACE_COUNTER_INCREMENT(AudioMixerSourceManager_TotalQdCmds);
+#if !UE_BUILD_SHIPPING			
+			if (LogCmdQueueWhenNumberReachedCVar > 0 && LogCmdQueueWhenNumberReachedCVar == Commands.SourceCommandQueue.Num())
+			{
+				TMap<const char*, int32> CmdsCounts;
+				for (const FAudioMixerThreadCommand& i : Commands.SourceCommandQueue)
+				{
+					CmdsCounts.FindOrAdd(i.DebugString)++;
+				}
+				for (auto i : CmdsCounts)
+				{
+					UE_LOG(LogAudioMixer, Display, TEXT("Commands %s:%d"), StringCast<TCHAR>(i.Key).Get(), i.Value);
+				}
+			}
+#endif //!UE_BUILD_SHIPPING						
 
 			NewNum = Commands.SourceCommandQueue.Num();
 			NewMax = Commands.SourceCommandQueue.Max();
@@ -3752,6 +3814,7 @@ namespace Audio
 			}
 			
 			NumCommands.Decrement();
+			TRACE_COUNTER_DECREMENT(AudioMixerSourceManager_TotalQdCmds);
 		}
 
 		LastPumpCompleteTimeInCycles = FPlatformTime::Cycles64();
diff --git a/Engine/Source/Runtime/AudioMixer/Private/AudioMixerSourceManager.h b/Engine/Source/Runtime/AudioMixer/Private/AudioMixerSourceManager.h
--- a/Engine/Source/Runtime/AudioMixer/Private/AudioMixerSourceManager.h
+++ b/Engine/Source/Runtime/AudioMixer/Private/AudioMixerSourceManager.h
@@ -336,6 +336,9 @@ namespace Audio
 		void ComputeBuses();
 		void UpdateBuses();
 
+		float GetFloatCompareTolerance() const;
+		float GetCommandQueueFillPercentage() const;
+
 		struct FAudioMixerThreadCommand
 		{
 			// ctor
diff --git a/Engine/Source/Runtime/AudioMixer/Private/AudioMixerSourceVoice.cpp b/Engine/Source/Runtime/AudioMixer/Private/AudioMixerSourceVoice.cpp
--- a/Engine/Source/Runtime/AudioMixer/Private/AudioMixerSourceVoice.cpp
+++ b/Engine/Source/Runtime/AudioMixer/Private/AudioMixerSourceVoice.cpp
@@ -100,8 +100,7 @@ namespace Audio
 	void FMixerSourceVoice::SetPitch(const float InPitch)
 	{
 		AUDIO_MIXER_CHECK_GAME_THREAD(MixerDevice);
-
-		if (Pitch != InPitch)
+		if (!FMath::IsNearlyEqual(Pitch, InPitch, SourceManager->GetFloatCompareTolerance()))
 		{
 			Pitch = InPitch;
 			SourceManager->SetPitch(SourceId, InPitch);
@@ -111,8 +110,7 @@ namespace Audio
 	void FMixerSourceVoice::SetVolume(const float InVolume)
 	{
 		AUDIO_MIXER_CHECK_GAME_THREAD(MixerDevice);
-
-		if (Volume != InVolume)
+		if (!FMath::IsNearlyEqual(InVolume, InVolume, SourceManager->GetFloatCompareTolerance()))
 		{
 			Volume = InVolume;
 			SourceManager->SetVolume(SourceId, InVolume);
@@ -123,7 +121,7 @@ namespace Audio
 	{
 		AUDIO_MIXER_CHECK_GAME_THREAD(MixerDevice);
 
-		if (DistanceAttenuation != InDistanceAttenuation)
+		if (!FMath::IsNearlyEqual(DistanceAttenuation, InDistanceAttenuation, SourceManager->GetFloatCompareTolerance()))
 		{
 			DistanceAttenuation = InDistanceAttenuation;
 			SourceManager->SetDistanceAttenuation(SourceId, InDistanceAttenuation);
@@ -134,7 +132,7 @@ namespace Audio
 	{
 		AUDIO_MIXER_CHECK_GAME_THREAD(MixerDevice);
 
-		if (LPFFrequency != InLPFFrequency)
+		if (!FMath::IsNearlyEqual(LPFFrequency, InLPFFrequency, SourceManager->GetFloatCompareTolerance()))
 		{
 			LPFFrequency = InLPFFrequency;
 			SourceManager->SetLPFFrequency(SourceId, LPFFrequency);
@@ -145,7 +143,7 @@ namespace Audio
 	{
 		AUDIO_MIXER_CHECK_GAME_THREAD(MixerDevice);
 
-		if (HPFFrequency != InHPFFrequency)
+		if (!FMath::IsNearlyEqual(InHPFFrequency, InHPFFrequency, SourceManager->GetFloatCompareTolerance()))
 		{
 			HPFFrequency = InHPFFrequency;
 			SourceManager->SetHPFFrequency(SourceId, HPFFrequency);
@@ -156,7 +154,7 @@ namespace Audio
 	{
 		AUDIO_MIXER_CHECK_GAME_THREAD(MixerDevice);
 
-		if (InVolumeModBase != VolumeModBase)
+		if (!FMath::IsNearlyEqual(VolumeModBase, InVolumeModBase, SourceManager->GetFloatCompareTolerance()))
 		{
 			VolumeModBase = InVolumeModBase;
 			SourceManager->SetModVolume(SourceId, VolumeModBase);
@@ -167,7 +165,7 @@ namespace Audio
 	{
 		AUDIO_MIXER_CHECK_GAME_THREAD(MixerDevice);
 
-		if (InPitchModBase != PitchModBase)
+		if (!FMath::IsNearlyEqual(PitchModBase, InPitchModBase, SourceManager->GetFloatCompareTolerance()))	
 		{
 			PitchModBase = InPitchModBase;
 			SourceManager->SetModPitch(SourceId, InPitchModBase);
@@ -178,7 +176,7 @@ namespace Audio
 	{
 		AUDIO_MIXER_CHECK_GAME_THREAD(MixerDevice);
 
-		if (InHPFFrequencyModBase != HPFFrequencyModBase)
+		if (!FMath::IsNearlyEqual(HPFFrequencyModBase, InHPFFrequencyModBase, SourceManager->GetFloatCompareTolerance()))	
 		{
 			HPFFrequencyModBase = InHPFFrequencyModBase;
 			SourceManager->SetModHPFFrequency(SourceId, InHPFFrequencyModBase);
@@ -189,7 +187,7 @@ namespace Audio
 	{
 		AUDIO_MIXER_CHECK_GAME_THREAD(MixerDevice);
 
-		if (InLPFFrequencyModBase != LPFFrequencyModBase)
+		if (!FMath::IsNearlyEqual(LPFFrequencyModBase, InLPFFrequencyModBase, SourceManager->GetFloatCompareTolerance()))	
 		{
 			LPFFrequencyModBase = InLPFFrequencyModBase;
 			SourceManager->SetModLPFFrequency(SourceId, InLPFFrequencyModBase);
@@ -402,7 +400,7 @@ namespace Audio
 				SubmixSends.Add(SubmixPtr->GetId(), NewSubmixSend);
 				SourceManager->SetSubmixSendInfo(SourceId, NewSubmixSend);
 			}
-			else if (!FMath::IsNearlyEqual(SubmixSend->SendLevel, SendLevel) || SubmixSend->SubmixSendStage != SendStage)
+			else if (!FMath::IsNearlyEqual(SubmixSend->SendLevel, SendLevel, SourceManager->GetFloatCompareTolerance()) || SubmixSend->SubmixSendStage != SendStage)
 			{
 				SubmixSend->SendLevel = SendLevel;
 				SubmixSend->SubmixSendStage = SendStage;
diff --git a/Engine/Source/Runtime/Engine/Public/Audio.h b/Engine/Source/Runtime/Engine/Public/Audio.h
--- a/Engine/Source/Runtime/Engine/Public/Audio.h
+++ b/Engine/Source/Runtime/Engine/Public/Audio.h
@@ -630,8 +630,6 @@ public:
 		, LFEBleed(0.5f)
 		, LPFFrequency(MAX_FILTER_FREQUENCY)
 		, HPFFrequency(MIN_FILTER_FREQUENCY)
-		, LastLPFFrequency(MAX_FILTER_FREQUENCY)
-		, LastHPFFrequency(MIN_FILTER_FREQUENCY)
 		, PlaybackTime(0.0f)
 		, Pitch(1.0f)
 		, LastUpdate(0)
@@ -808,11 +806,13 @@ protected:
 	/** What frequency to set the HPF filter to. Note this could be caused by HPF distance attenuation. */
 	float HPFFrequency;
 
-	/** The last LPF frequency set. Used to avoid making API calls when parameter doesn't changing. */
-	float LastLPFFrequency;
+	/** The last LPF frequency set. Used to avoid making API calls when parameter doesn't changing. */	
+	UE_DEPRECATED(5.6, "This Variable duplicates cache behavior used at source voice level")
+	float LastLPFFrequency=MAX_FILTER_FREQUENCY;
 
 	/** The last HPF frequency set. Used to avoid making API calls when parameter doesn't changing. */
-	float LastHPFFrequency;
+	UE_DEPRECATED(5.6, "This Variable duplicates cache behavior used at source voice level")
+	float LastHPFFrequency=MIN_FILTER_FREQUENCY;
 
 	/** The virtual current playback time. Used to trigger notifications when finished. */
 	float PlaybackTime;
