clean up AUDIO_SUBFRAME_ENABLED and related state

diff --git a/Engine/Source/Runtime/AudioMixer/Private/AudioMixerSourceManager.cpp b/Engine/Source/Runtime/AudioMixer/Private/AudioMixerSourceManager.cpp
--- a/Engine/Source/Runtime/AudioMixer/Private/AudioMixerSourceManager.cpp
+++ b/Engine/Source/Runtime/AudioMixer/Private/AudioMixerSourceManager.cpp
@@ -209,9 +209,6 @@ FAutoConsoleVariableRef CVarAudioMixerPerSourceResampling(
 
 #endif
 
-// Disable subframe timing logic
-#define AUDIO_SUBFRAME_ENABLED 0
-
 // Global count for _all_ queued commands (across all devices).
 TRACE_DECLARE_ATOMIC_INT_COUNTER(AudioMixerSourceManager_TotalQdCmds, TEXT("AudioSourceManager Qd Cmds"));
 
@@ -446,7 +443,6 @@ namespace Audio
 			SourceInfo.CurrentFrameAlpha = 0.0f;
 			SourceInfo.CurrentFrameIndex = 0;
 			SourceInfo.NumFramesPlayed = 0;
-			SourceInfo.StartTime = 0.0;
 			SourceInfo.SubmixSends.Reset();
 			SourceInfo.AudioBusId = INDEX_NONE;
 			SourceInfo.SourceBusDurationFrames = INDEX_NONE;
@@ -841,7 +837,6 @@ namespace Audio
 		SourceInfo.CurrentFrameAlpha = 0.0f;
 		SourceInfo.CurrentFrameIndex = 0;
 		SourceInfo.NumFramesPlayed = 0;
-		SourceInfo.StartTime = 0.0;
 		SourceInfo.bIs3D = false;
 		SourceInfo.bIsCenterChannelOnly = false;
 		SourceInfo.bIsActive = false;
@@ -1481,10 +1476,7 @@ namespace Audio
 		AUDIO_MIXER_CHECK(GameThreadInfo.bIsBusy[SourceId]);
 		AUDIO_MIXER_CHECK_GAME_THREAD(MixerDevice);
 
-		// Compute the frame within which to start the sound based on the current "thread faction" on the audio thread
-		double StartTime = MixerDevice->GetAudioThreadTime();
-
-		AudioMixerThreadCommand([this, SourceId, StartTime]()
+		AudioMixerThreadCommand([this, SourceId]()
 		{
 			AUDIO_MIXER_CHECK_AUDIO_PLAT_THREAD(MixerDevice);
 
@@ -1494,8 +1486,6 @@ namespace Audio
 			SourceInfo.bIsPaused = false;
 			SourceInfo.bIsActive = true;
 
-			SourceInfo.StartTime = StartTime;
-					
 			UE_CLOG(Audio::MatchesLogFilter(*SourceInfo.GetDebugName()), LogAudioTiming, Verbose,
 				TEXT("FMixerSourceManager::Play (render thread), Wave=%s, SourceId=%d"), *SourceInfo.GetDebugName(), SourceId);
 
@@ -2252,9 +2242,6 @@ namespace Audio
 		CSV_SCOPED_TIMING_STAT(Audio, SourceBuffers);
 		SCOPE_CYCLE_COUNTER(STAT_AudioMixerSourceBuffers);
 
-		const double AudioRenderThreadTime = MixerDevice->GetAudioRenderThreadTime();
-		const double AudioClockDelta = MixerDevice->GetAudioClockDelta();
-
 		FSourceInfo& SourceInfo = SourceInfos[SourceId];
 
 		if (!SourceInfo.bIsBusy || !SourceInfo.bIsPlaying || SourceInfo.bIsPaused || SourceInfo.bIsPausedForQuantization)
@@ -2338,34 +2325,6 @@ namespace Audio
 		}
 		else
 		{
-
-#if AUDIO_SUBFRAME_ENABLED
-			// If we're not going to start yet, just continue
-			double StartFraction = (SourceInfo.StartTime - AudioRenderThreadTime) / AudioClockDelta;
-			if (StartFraction >= 1.0)
-			{
-				// note this is already zero'd so no need to write zeroes
-				SourceInfo.PitchSourceParam.Reset();
-				return;
-			}
-
-			// Init the frame index iterator to 0 (i.e. render whole buffer)
-			int32 StartFrame = 0;
-
-			// If the start fraction is greater than 0.0 (and is less than 1.0), we are starting on a sub-frame
-			// Otherwise, just start playing it right away
-			if (StartFraction > 0.0)
-			{
-				StartFrame = NumOutputFrames * StartFraction;
-			}
-
-			// Update sample index to the frame we're starting on, accounting for source channels
-			int32 SampleIndex = StartFrame * SourceInfo.NumInputChannels;
-#else
-			int32 SampleIndex = 0;
-			int32 StartFrame = 0;
-#endif
-
 			// Modulate parameter target should modulation be active
 			// Due to managing two separate pitch values that are updated at different rates
 			// (game thread rate and copy set by SetPitch and buffer callback rate set by Modulation System),
@@ -2400,9 +2359,10 @@ namespace Audio
 
 			if (!PerSourceResampling)
 			{
+				int32 SampleIndex = 0;
 				float CurrentAlpha = SourceInfo.CurrentFrameAlpha;
 
-				for (int32 Frame = StartFrame; Frame < NumOutputFrames; ++Frame)
+				for (int32 Frame = 0; Frame < NumOutputFrames; ++Frame)
 				{
 					// If we've read our last buffer, we're done
 					if (SourceInfo.bIsLastBuffer)
@@ -2492,9 +2452,6 @@ namespace Audio
 						SourceInfo.bDelayLineSet = true;
 					}
 
-					// Also skip some frames if subframe rendering is enabled
-					FramesToWrite = FMath::Max(FramesToWrite - StartFrame, 0);
-
 					while (FramesToWrite > 0)
 					{
 						int32 CurrentChunkFrames = SourceInfo.GetCurrentAudioChunkNumFrames();
diff --git a/Engine/Source/Runtime/AudioMixer/Private/AudioMixerSourceManager.h b/Engine/Source/Runtime/AudioMixer/Private/AudioMixerSourceManager.h
--- a/Engine/Source/Runtime/AudioMixer/Private/AudioMixerSourceManager.h
+++ b/Engine/Source/Runtime/AudioMixer/Private/AudioMixerSourceManager.h
@@ -464,9 +464,6 @@ namespace Audio
 			int32 CurrentFrameIndex;
 			int64 NumFramesPlayed;
 
-			// The number of frames to wait before starting the source
-			double StartTime;
-
 			TArray<FMixerSourceSubmixSend> SubmixSends;
 
 			// What audio bus Id this source is sonfiying, if it is a source bus. This is INDEX_NONE for sources which are not source buses.
