Fix copy paste error in Command Queue Spam Reduction CL

diff --git a/Engine/Source/Runtime/AudioMixer/Private/AudioMixerSourceVoice.cpp b/Engine/Source/Runtime/AudioMixer/Private/AudioMixerSourceVoice.cpp
--- a/Engine/Source/Runtime/AudioMixer/Private/AudioMixerSourceVoice.cpp
+++ b/Engine/Source/Runtime/AudioMixer/Private/AudioMixerSourceVoice.cpp
@@ -110,7 +110,7 @@ namespace Audio
 	void FMixerSourceVoice::SetVolume(const float InVolume)
 	{
 		AUDIO_MIXER_CHECK_GAME_THREAD(MixerDevice);
-		if (!FMath::IsNearlyEqual(InVolume, InVolume, SourceManager->GetFloatCompareTolerance()))
+		if (!FMath::IsNearlyEqual(Volume, InVolume, SourceManager->GetFloatCompareTolerance()))
 		{
 			Volume = InVolume;
 			SourceManager->SetVolume(SourceId, InVolume);
@@ -143,7 +143,7 @@ namespace Audio
 	{
 		AUDIO_MIXER_CHECK_GAME_THREAD(MixerDevice);
 
-		if (!FMath::IsNearlyEqual(InHPFFrequency, InHPFFrequency, SourceManager->GetFloatCompareTolerance()))
+		if (!FMath::IsNearlyEqual(HPFFrequency, InHPFFrequency, SourceManager->GetFloatCompareTolerance()))
 		{
 			HPFFrequency = InHPFFrequency;
 			SourceManager->SetHPFFrequency(SourceId, HPFFrequency);
