Fix for truncation of one-shots.

Issue was that the "doneness" of a source was set before the submix code was mixed in and subnmies check if a sound is done before mixing "sent" sources into it's submix

[RN] Fixed one-shot truncation in the audio renderer.

diff --git a/Engine/Source/Runtime/AudioMixer/Private/AudioMixerDevice.cpp b/Engine/Source/Runtime/AudioMixer/Private/AudioMixerDevice.cpp
--- a/Engine/Source/Runtime/AudioMixer/Private/AudioMixerDevice.cpp
+++ b/Engine/Source/Runtime/AudioMixer/Private/AudioMixerDevice.cpp
@@ -976,6 +976,9 @@ namespace Audio
 			}
 		}
 
+		// Only after submixes have finished rendering audio, update the source state in the source manager to reflect sound-doneness, etc
+		SourceManager->UpdateSourceState();
+
 		{
 			CSV_SCOPED_TIMING_STAT(Audio, EndpointSubmixes);
 			SCOPE_CYCLE_COUNTER(STAT_AudioMixerEndpointSubmixes);
diff --git a/Engine/Source/Runtime/AudioMixer/Private/AudioMixerSourceManager.cpp b/Engine/Source/Runtime/AudioMixer/Private/AudioMixerSourceManager.cpp
--- a/Engine/Source/Runtime/AudioMixer/Private/AudioMixerSourceManager.cpp
+++ b/Engine/Source/Runtime/AudioMixer/Private/AudioMixerSourceManager.cpp
@@ -3566,9 +3566,11 @@ namespace Audio
 			AUDIO_MIXER_CHECK(SourceDataOverridePlugin.IsValid());
 			LLM_SCOPE(ELLMTag::AudioMixerPlugins);
 			SourceDataOverridePlugin->OnAllSourcesProcessed();
-		}
+		}		
+	}
 
-		// Update the game thread copy of source doneness
+	void FMixerSourceManager::UpdateSourceState()
+	{
 		RenderThreadPhase = ESourceManagerRenderThreadPhase::UpdateGameThreadCopies;
 		for (int32 SourceId = 0; SourceId < NumTotalSources; ++SourceId)
 		{		
@@ -3986,4 +3988,4 @@ namespace Audio
 			ResamplerPool.Add(Resampler);
 		}
 	}
-}
+}
\ No newline at end of file
diff --git a/Engine/Source/Runtime/AudioMixer/Private/AudioMixerSourceManager.h b/Engine/Source/Runtime/AudioMixer/Private/AudioMixerSourceManager.h
--- a/Engine/Source/Runtime/AudioMixer/Private/AudioMixerSourceManager.h
+++ b/Engine/Source/Runtime/AudioMixer/Private/AudioMixerSourceManager.h
@@ -262,6 +262,7 @@ namespace Audio
 		bool IsUsingHRTFSpatializer(const int32 SourceId) const;
 		bool NeedsSpeakerMap(const int32 SourceId) const;
 		void ComputeNextBlockOfSamples();
+		void UpdateSourceState();
 		void ClearStoppingSounds();
 		void MixOutputBuffers(const int32 SourceId, int32 InNumOutputChannels, const float InSendLevel, EMixerSourceSubmixSendStage InSubmixSendStage, FAlignedFloatBuffer& OutWetBuffer) const;
 
