Also compute plugin audio IF AudioLinkSource is present

diff --git a/Engine/Source/Runtime/AudioMixer/Private/AudioMixerSourceManager.cpp b/Engine/Source/Runtime/AudioMixer/Private/AudioMixerSourceManager.cpp
--- a/Engine/Source/Runtime/AudioMixer/Private/AudioMixerSourceManager.cpp
+++ b/Engine/Source/Runtime/AudioMixer/Private/AudioMixerSourceManager.cpp
@@ -3165,7 +3165,7 @@ namespace Audio
 				}
 			}
 
-			if (SourceInfo.IsRenderingToSubmixes() || SpatialInterfaceInfo.bSpatializationIsExternalSend)
+			if (SourceInfo.IsRenderingToSubmixes() || SpatialInterfaceInfo.bSpatializationIsExternalSend || SourceInfo.AudioLink.IsValid())
 			{
 				// Apply distance attenuation
 				ApplyDistanceAttenuation(SourceInfo, NumSamples);
