Fix for Surround2D omni blend mode if output is stereo and not surround sound.

Stereo output will blend to a 5.1 channel map and automatically downmix to stereo ini the audio renderer.

[RN] Fix for Surround2D non-spatialized radius mode not working when output is stereo.

diff --git a/Engine/Source/Runtime/AudioMixer/Private/AudioMixerSource.cpp b/Engine/Source/Runtime/AudioMixer/Private/AudioMixerSource.cpp
--- a/Engine/Source/Runtime/AudioMixer/Private/AudioMixerSource.cpp
+++ b/Engine/Source/Runtime/AudioMixer/Private/AudioMixerSource.cpp
@@ -1942,7 +1942,8 @@ namespace Audio
 				else
 				{
 					// If we are in 5.1, we need to use the side-channel speakers
-					if (NumOutputChannels == 6)
+					// If we are outputting stereo, omni-blend to a 5.1 output. This will get downmixed to stereo as a fallback.
+					if (NumOutputChannels == 2 || NumOutputChannels == 6)
 					{
 						// Create some omni maps for left and right channels
 						auto CreateLeftOmniMap = []() -> TMap<EAudioMixerChannel::Type, float>
