(proper) Fix for omniblend Direct2D for mono sources

[RN] Omniblend is fixed for mono sources using Direct2D for the case of surround sound outputs.

diff --git a/Engine/Source/Runtime/AudioMixer/Private/AudioMixerSource.cpp b/Engine/Source/Runtime/AudioMixer/Private/AudioMixerSource.cpp
--- a/Engine/Source/Runtime/AudioMixer/Private/AudioMixerSource.cpp
+++ b/Engine/Source/Runtime/AudioMixer/Private/AudioMixerSource.cpp
@@ -1849,10 +1849,106 @@ namespace Audio
 			// Don't need to compute the source channel map if the absolute azimuth hasn't changed much
 			PreviousAzimuth = WaveInstance->AbsoluteAzimuth;
 			OutChannelMap.Reset();
+			int32 NumOutputChannels = MixerDevice->GetNumDeviceChannels();
 
-			int32 NumDeviceChannels = MixerDevice->GetNumDeviceChannels();
-			float DefaultOmniAmount = 1.0f / NumDeviceChannels;
-			MixerDevice->Get3DChannelMap(NumDeviceChannels, WaveInstance, WaveInstance->AbsoluteAzimuth, SpatializationParams.NonSpatializedAmount, nullptr, DefaultOmniAmount, OutChannelMap);
+			if (WaveInstance->NonSpatializedRadiusMode == ENonSpatializedRadiusSpeakerMapMode::OmniDirectional)
+			{
+				float DefaultOmniAmount = 1.0f / NumOutputChannels;
+				MixerDevice->Get3DChannelMap(NumOutputChannels, WaveInstance, WaveInstance->AbsoluteAzimuth, SpatializationParams.NonSpatializedAmount, nullptr, DefaultOmniAmount, OutChannelMap);
+			}
+			else if (WaveInstance->NonSpatializedRadiusMode == ENonSpatializedRadiusSpeakerMapMode::Direct2D)
+			{
+				// Create some omni maps for left and right channels, note we're 
+				// taking into account mono upmix method
+				auto CreateOmniMap = [this]() -> TMap<EAudioMixerChannel::Type, float>
+				{
+					EMonoChannelUpmixMethod MonoUpmixMethod = MixerDevice->GetMonoChannelUpmixMethod();
+					TMap<EAudioMixerChannel::Type, float> OmniMap;
+
+					if (MonoUpmixMethod == EMonoChannelUpmixMethod::FullVolume)
+					{
+						OmniMap.Add(EAudioMixerChannel::FrontLeft, Audio::MonoUpmixFullVolume);
+						OmniMap.Add(EAudioMixerChannel::FrontRight, Audio::MonoUpmixFullVolume);
+					}
+					else if (MonoUpmixMethod == EMonoChannelUpmixMethod::EqualPower)
+					{
+						OmniMap.Add(EAudioMixerChannel::FrontLeft, Audio::MonoUpmixEqualPower);
+						OmniMap.Add(EAudioMixerChannel::FrontRight, Audio::MonoUpmixEqualPower);
+					}
+					else
+					{
+						check(MonoUpmixMethod == EMonoChannelUpmixMethod::Linear);
+						OmniMap.Add(EAudioMixerChannel::FrontLeft, Audio::MonoUpmixLinear);
+						OmniMap.Add(EAudioMixerChannel::FrontRight, Audio::MonoUpmixLinear);
+					}
+					return OmniMap;
+				};
+
+				static const TMap<EAudioMixerChannel::Type, float> OmniMap = CreateOmniMap();
+				MixerDevice->Get3DChannelMap(NumOutputChannels, WaveInstance, WaveInstance->AbsoluteAzimuth, SpatializationParams.NonSpatializedAmount, &OmniMap, 0.0f, OutChannelMap);
+			}
+			else if (WaveInstance->NonSpatializedRadiusMode == ENonSpatializedRadiusSpeakerMapMode::Surround2D)
+			{
+				// Create some omni maps for left and right channels, note we're 
+				// taking into account mono upmix method
+				auto CreateOmniMap = [this, &NumOutputChannels]() -> TMap<EAudioMixerChannel::Type, float>
+					{
+						EMonoChannelUpmixMethod MonoUpmixMethod = MixerDevice->GetMonoChannelUpmixMethod();
+
+						TMap<EAudioMixerChannel::Type, float> OmniMap;
+
+						if (MonoUpmixMethod == EMonoChannelUpmixMethod::FullVolume)
+						{
+							OmniMap.Add(EAudioMixerChannel::FrontLeft, Audio::MonoUpmixFullVolume);
+							OmniMap.Add(EAudioMixerChannel::FrontRight, Audio::MonoUpmixFullVolume);
+							if (NumOutputChannels == 8)
+							{
+								OmniMap.Add(EAudioMixerChannel::BackLeft, Audio::MonoUpmixFullVolume);
+								OmniMap.Add(EAudioMixerChannel::BackRight, Audio::MonoUpmixFullVolume);
+							}
+							else if (NumOutputChannels == 6)
+							{
+								OmniMap.Add(EAudioMixerChannel::SideLeft, Audio::MonoUpmixFullVolume);
+								OmniMap.Add(EAudioMixerChannel::SideRight, Audio::MonoUpmixFullVolume);
+							}
+						}
+						else if (MonoUpmixMethod == EMonoChannelUpmixMethod::EqualPower)
+						{
+							OmniMap.Add(EAudioMixerChannel::FrontLeft, Audio::MonoUpmixEqualPower);
+							OmniMap.Add(EAudioMixerChannel::FrontRight, Audio::MonoUpmixEqualPower);
+							if (NumOutputChannels == 8)
+							{
+								OmniMap.Add(EAudioMixerChannel::BackLeft, Audio::MonoUpmixEqualPower);
+								OmniMap.Add(EAudioMixerChannel::BackRight, Audio::MonoUpmixEqualPower);
+							}
+							else
+							{
+								OmniMap.Add(EAudioMixerChannel::SideLeft, Audio::MonoUpmixEqualPower);
+								OmniMap.Add(EAudioMixerChannel::SideRight, Audio::MonoUpmixEqualPower);
+							}
+						}
+						else
+						{
+							check(MonoUpmixMethod == EMonoChannelUpmixMethod::Linear);
+							OmniMap.Add(EAudioMixerChannel::FrontLeft, Audio::MonoUpmixLinear);
+							OmniMap.Add(EAudioMixerChannel::FrontRight, Audio::MonoUpmixLinear);
+							if (NumOutputChannels == 8)
+							{
+								OmniMap.Add(EAudioMixerChannel::BackLeft, Audio::MonoUpmixLinear);
+								OmniMap.Add(EAudioMixerChannel::BackRight, Audio::MonoUpmixLinear);
+							}
+							else if (NumOutputChannels == 6)
+							{
+								OmniMap.Add(EAudioMixerChannel::SideLeft, Audio::MonoUpmixLinear);
+								OmniMap.Add(EAudioMixerChannel::SideRight, Audio::MonoUpmixLinear);
+							}
+						}
+						return OmniMap;
+					};
+
+				static const TMap<EAudioMixerChannel::Type, float> OmniMap = CreateOmniMap();
+				MixerDevice->Get3DChannelMap(NumOutputChannels, WaveInstance, WaveInstance->AbsoluteAzimuth, SpatializationParams.NonSpatializedAmount, &OmniMap, 0.0f, OutChannelMap);
+			}
 			return true;
 		}
 		else if (!OutChannelMap.Num() || (IsSpatializationCVarEnabled() != bPrevAllowedSpatializationSetting))
diff --git a/Engine/Source/Runtime/AudioMixerCore/Public/AudioDefines.h b/Engine/Source/Runtime/AudioMixerCore/Public/AudioDefines.h
--- a/Engine/Source/Runtime/AudioMixerCore/Public/AudioDefines.h
+++ b/Engine/Source/Runtime/AudioMixerCore/Public/AudioDefines.h
@@ -64,4 +64,11 @@ namespace Audio
 	 * Typed identifier for Audio Device Id
 	 */
 	using FDeviceId = uint32;
+
+
+	// Magic numbers for mono upmix method
+	static constexpr float MonoUpmixFullVolume = 1.0f;
+	static constexpr float MonoUpmixEqualPower = 0.707f;
+	static constexpr float MonoUpmixLinear = 0.5f;
+
 }
\ No newline at end of file
diff --git a/Engine/Source/Runtime/SignalProcessing/Private/ChannelMap.cpp b/Engine/Source/Runtime/SignalProcessing/Private/ChannelMap.cpp
--- a/Engine/Source/Runtime/SignalProcessing/Private/ChannelMap.cpp
+++ b/Engine/Source/Runtime/SignalProcessing/Private/ChannelMap.cpp
@@ -1,6 +1,7 @@
 // Copyright Epic Games, Inc. All Rights Reserved.
 
 #include "DSP/ChannelMap.h"
+#include "AudioDefines.h"
 
 namespace Audio
 {
@@ -280,18 +281,18 @@ namespace Audio
 					switch (InMonoChannelUpmixMethod)
 					{
 						case EChannelMapMonoUpmixMethod::Linear:
-							OutChannelMap.Add(0.5f);
-							OutChannelMap.Add(0.5f);
+							OutChannelMap.Add(Audio::MonoUpmixLinear);
+							OutChannelMap.Add(Audio::MonoUpmixLinear);
 							break;
 
 						case EChannelMapMonoUpmixMethod::EqualPower:
-							OutChannelMap.Add(0.707f);
-							OutChannelMap.Add(0.707f);
+							OutChannelMap.Add(Audio::MonoUpmixEqualPower);
+							OutChannelMap.Add(Audio::MonoUpmixEqualPower);
 							break;
 
 						case EChannelMapMonoUpmixMethod::FullVolume:
-							OutChannelMap.Add(1.0f);
-							OutChannelMap.Add(1.0f);
+							OutChannelMap.Add(Audio::MonoUpmixFullVolume);
+							OutChannelMap.Add(Audio::MonoUpmixFullVolume);
 							break;
 					}
 
