ISoundGenerators don't zero out memory past the end of the samples they actually write

#rnx

diff --git a/Engine/Source/Runtime/AudioMixer/Private/AudioMixerSourceBuffer.cpp b/Engine/Source/Runtime/AudioMixer/Private/AudioMixerSourceBuffer.cpp
--- a/Engine/Source/Runtime/AudioMixer/Private/AudioMixerSourceBuffer.cpp
+++ b/Engine/Source/Runtime/AudioMixer/Private/AudioMixerSourceBuffer.cpp
@@ -7,9 +7,6 @@
 #include "Misc/ScopeTryLock.h"
 #include "DSP/FloatArrayMath.h"
 
-// extern Cvars.
-extern int32 bDecodersCanReturnPartialBuffersCvar;
-
 namespace Audio
 {
 	int32 DirectProceduralRenderingCVar = 0;
@@ -569,21 +566,11 @@ namespace Audio
 				AsyncRealtimeAudioTask->GetResult(TaskResult);
 				bIsFinishedOrLooped = TaskResult.bIsFinishedOrLooped;
 
-				if (bDecodersCanReturnPartialBuffersCvar)
-				{
-					SourceVoiceBuffers[CurrentBuffer]->SetNum(TaskResult.NumSamplesWritten, EAllowShrinking::No);
+				SourceVoiceBuffers[CurrentBuffer]->SetNum(TaskResult.NumSamplesWritten, EAllowShrinking::No);
 #if ENABLE_AUDIO_DEBUG
-					double AudioDuration = static_cast<double>(TaskResult.NumSamplesWritten) / static_cast<double>(FMath::Max(1, NumChannels * SampleRate));
-					UpdateCPUCoreUtilization(TaskResult.CPUDuration, AudioDuration);
+				double AudioDuration = static_cast<double>(TaskResult.NumSamplesWritten) / static_cast<double>(FMath::Max(1, NumChannels * SampleRate));
+				UpdateCPUCoreUtilization(TaskResult.CPUDuration, AudioDuration);
 #endif // ENABLE_AUDIO_DEBUG
-				}
-				else
-				{
-#if ENABLE_AUDIO_DEBUG
-					double AudioDuration = static_cast<double>(MONO_PCM_BUFFER_SAMPLES) / FMath::Max(1., static_cast<double>(SampleRate));
-					UpdateCPUCoreUtilization(TaskResult.CPUDuration, AudioDuration);
-#endif //ENABLE_AUDIO_DEBUG
-				}
 			}
 			break;
 
diff --git a/Engine/Source/Runtime/AudioMixer/Private/AudioMixerSourceDecode.cpp b/Engine/Source/Runtime/AudioMixer/Private/AudioMixerSourceDecode.cpp
--- a/Engine/Source/Runtime/AudioMixer/Private/AudioMixerSourceDecode.cpp
+++ b/Engine/Source/Runtime/AudioMixer/Private/AudioMixerSourceDecode.cpp
@@ -28,15 +28,7 @@ FAutoConsoleVariableRef CVarForceSynchronizedAudioTaskKick(
 	TEXT("Force all Audio Tasks created in one \"audio render frame\" to be queued until they can all be \"kicked\" at once at the end of the frame.\n")
 	TEXT("0: Don't Force, 1: Force"),
 	ECVF_Default);
-
-int32 bDecodersCanReturnPartialBuffersCvar = 1;
-FAutoConsoleVariableRef CVarDecodersCanReturnPartialBuffers(
-	TEXT("au.DecodersCanReturnPartialBuffers"),
-	bDecodersCanReturnPartialBuffersCvar,
-	TEXT(""),
-	ECVF_Default
-);
-
+	
 namespace Audio
 {
 
@@ -134,23 +126,14 @@ public:
 				DecodeResult.bIsFinishedOrLooped = DecodeTaskData.DecompressionState->ReadCompressedData(DecodeBuffer.GetData(), DecodeTaskData.bLoopingMode, kPCMBufferSize);
 			}
 
-			if (bDecodersCanReturnPartialBuffersCvar)
-			{
-				const int32 NumSamplesStreamed = NumBytesStreamed / sizeof(int16);
+			const int32 NumSamplesStreamed = NumBytesStreamed / sizeof(int16);
 
-				DecodeResult.NumSamplesWritten = NumSamplesStreamed;
+			DecodeResult.NumSamplesWritten = NumSamplesStreamed;
 
-				// Convert the decoded PCM data into a float buffer while still in the async task
-				Audio::ArrayPcm16ToFloat(
-					MakeArrayView((int16*)DecodeBuffer.GetData(), NumSamplesStreamed)
-					, MakeArrayView(DecodeTaskData.AudioData, NumSamplesStreamed));
-			}
-			else
-			{
-				// Convert the decoded PCM data into a float buffer while still in the async task
-				Audio::ArrayPcm16ToFloat(MakeArrayView((int16*)DecodeBuffer.GetData(), DecodeTaskData.NumFramesToDecode * NumChannels)
-					, MakeArrayView(DecodeTaskData.AudioData, DecodeTaskData.NumFramesToDecode * NumChannels));
-			}
+			// Convert the decoded PCM data into a float buffer while still in the async task
+			Audio::ArrayPcm16ToFloat(
+				MakeArrayView((int16*)DecodeBuffer.GetData(), NumSamplesStreamed)
+				, MakeArrayView(DecodeTaskData.AudioData, NumSamplesStreamed));
 		}
 		break;
 		}
diff --git a/Engine/Source/Runtime/Engine/Classes/Sound/SoundGenerator.h b/Engine/Source/Runtime/Engine/Classes/Sound/SoundGenerator.h
--- a/Engine/Source/Runtime/Engine/Classes/Sound/SoundGenerator.h
+++ b/Engine/Source/Runtime/Engine/Classes/Sound/SoundGenerator.h
@@ -55,7 +55,8 @@ public:
 	// of concurrently rendering voices. 
 	virtual float GetRelativeRenderCost() const { return 1.f; }
 
-	// Retrieves the next buffer of audio from the generator, called from the audio mixer
+	// Retrieves the next buffer of audio from the generator, called from the audio mixer.
+	// Returns the number of samples actually written. The remainder are untouched.
 	ENGINE_API int32 GetNextBuffer(float* OutAudio, int32 NumSamples, bool bRequireNumberSamples = false);
 
 	virtual Audio::AudioTaskQueueId GetSynchronizedRenderQueueId() const { return 0; }
diff --git a/Engine/Source/Runtime/Engine/Private/SoundGenerator.cpp b/Engine/Source/Runtime/Engine/Private/SoundGenerator.cpp
--- a/Engine/Source/Runtime/Engine/Private/SoundGenerator.cpp
+++ b/Engine/Source/Runtime/Engine/Private/SoundGenerator.cpp
@@ -12,7 +12,7 @@ ISoundGenerator::~ISoundGenerator()
 {
 }
 
-int32 ISoundGenerator::GetNextBuffer(float* OutAudio, int32 NumSamples, bool bRequireNumberSamples)
+int32 ISoundGenerator::GetNextBuffer(float* OutAudio, const int32 NumSamples, bool bRequireNumberSamples)
 {
 	LLM_SCOPE(ELLMTag::AudioSynthesis);
 
@@ -21,15 +21,22 @@ int32 ISoundGenerator::GetNextBuffer(float* OutAudio, int32 NumSamples, bool bRe
 	int32 NumSamplesToGenerate = NumSamples;
 	if (!bRequireNumberSamples)
 	{
-		const int32 NewNumSamplesToGenerate = FMath::Min(NumSamples, GetDesiredNumSamplesToRenderPerCallback());
-		if (const int32 Diff = NumSamples - NewNumSamplesToGenerate)
-		{
-			FMemory::Memzero(&OutAudio[NewNumSamplesToGenerate], Diff * sizeof(float));
-		}
-		NumSamplesToGenerate = NewNumSamplesToGenerate;
+		// Defer to the generator's desired block size
+		NumSamplesToGenerate = FMath::Min(NumSamples, GetDesiredNumSamplesToRenderPerCallback());
 	}
 
-	return OnGenerateAudio(OutAudio, NumSamplesToGenerate);
+	int32 NumSamplesWritten = OnGenerateAudio(OutAudio, NumSamplesToGenerate);
+
+	// If the generator didn't write all the required samples, zero out the rest
+	if (bRequireNumberSamples && NumSamplesWritten != NumSamplesToGenerate)
+	{
+		check(NumSamplesWritten >= 0 && NumSamplesWritten < NumSamplesToGenerate);
+		const int32 Remainder = NumSamplesToGenerate - NumSamplesWritten;
+		FMemory::Memzero(OutAudio + NumSamplesWritten, Remainder * sizeof(float));
+		NumSamplesWritten = NumSamplesToGenerate;
+	}
+
+	return NumSamplesWritten;
 }
 
 void ISoundGenerator::SynthCommand(TFunction<void()> Command)
