Fix edge case in new resampling path

#rnx

diff --git a/Engine/Source/Runtime/AudioMixer/Private/AudioMixerSourceManager.cpp b/Engine/Source/Runtime/AudioMixer/Private/AudioMixerSourceManager.cpp
--- a/Engine/Source/Runtime/AudioMixer/Private/AudioMixerSourceManager.cpp
+++ b/Engine/Source/Runtime/AudioMixer/Private/AudioMixerSourceManager.cpp
@@ -2481,6 +2481,13 @@ namespace Audio
 							{
 								break;
 							}
+
+							if (CurrentChunkFrames == 0)
+							{
+								// A source can return an empty buffer at the end of its input. In this case
+								//  we just ask again for a buffer and it will report it's done.
+								continue;
+							}
 						}
 
 						if (FinalPitch != 1 && !SourceInfo.bResampling)
