Parallel source processing synchronous with the render thread using the Tasks System

#rnx

diff --git a/Engine/Source/Runtime/AudioMixer/Private/AudioMixerDevice.cpp b/Engine/Source/Runtime/AudioMixer/Private/AudioMixerDevice.cpp
--- a/Engine/Source/Runtime/AudioMixer/Private/AudioMixerDevice.cpp
+++ b/Engine/Source/Runtime/AudioMixer/Private/AudioMixerDevice.cpp
@@ -614,10 +614,6 @@ namespace Audio
 				FSourceManagerInitParams SourceManagerInitParams;
 				SourceManagerInitParams.NumSources = GetMaxSources();
 
-				// TODO: Migrate this to project settings properly
-				SourceManagerInitParams.NumSourceWorkers = 4;
-
-
 				AudioClock = 0.0;
 				AudioClockDelta = (double)OpenStreamParams.NumFrames / OpenStreamParams.SampleRate;
 				AudioClockTimingData.UpdateTime = 0.0;
diff --git a/Engine/Source/Runtime/AudioMixer/Private/AudioMixerSourceManager.cpp b/Engine/Source/Runtime/AudioMixer/Private/AudioMixerSourceManager.cpp
--- a/Engine/Source/Runtime/AudioMixer/Private/AudioMixerSourceManager.cpp
+++ b/Engine/Source/Runtime/AudioMixer/Private/AudioMixerSourceManager.cpp
@@ -20,6 +20,7 @@
 #include "ProfilingDebugging/CountersTrace.h"
 #include "HAL/PlatformStackWalk.h"
 #include "Stats/Stats.h"
+#include "Tasks/Task.h"
 #include "Trace/Trace.h"
 
 #if WITH_AUDIO_MIXER_THREAD_COMMAND_DEBUG
@@ -34,7 +35,7 @@ static int32 DisableParallelSourceProcessingCvar = 1;
 FAutoConsoleVariableRef CVarDisableParallelSourceProcessing(
 	TEXT("au.DisableParallelSourceProcessing"),
 	DisableParallelSourceProcessingCvar,
-	TEXT("Disables using async tasks for processing sources.\n")
+	TEXT("Disables using  tasks for processing sources.\n")
 	TEXT("0: Not Disabled, 1: Disabled"),
 	ECVF_Default);
 
@@ -350,7 +351,6 @@ namespace Audio
 		, NumTotalSources(0)
 		, NumOutputFrames(0)
 		, NumOutputSamples(0)
-		, NumSourceWorkers(4)
 		, bInitialized(false)
 		, bUsingSpatializationPlugin(false)
 		, bUsingSourceDataOverridePlugin(false)
@@ -369,17 +369,6 @@ namespace Audio
 
 	FMixerSourceManager::~FMixerSourceManager()
 	{
-		if (SourceWorkers.Num() > 0)
-		{
-			for (int32 i = 0; i < SourceWorkers.Num(); ++i)
-			{
-				delete SourceWorkers[i];
-				SourceWorkers[i] = nullptr;
-			}
-
-			SourceWorkers.Reset();
-		}
-
 		FPlatformProcess::ReturnSynchEventToPool(CommandsProcessedEvent);
 	}
 
@@ -516,22 +505,6 @@ namespace Audio
 			SourceInfo.AudioPluginOutputData.AudioBuffer.Reset(NumOutputFrames * 2);
 		}
 
-		// Setup the source workers
-		SourceWorkers.Reset();
-		if (NumSourceWorkers > 0)
-		{
-			const int32 NumSourcesPerWorker = FMath::Max(NumTotalSources / NumSourceWorkers, 1);
-			int32 StartId = 0;
-			int32 EndId = 0;
-			while (EndId < NumTotalSources)
-			{
-				EndId = FMath::Min(StartId + NumSourcesPerWorker, NumTotalSources);
-				SourceWorkers.Add(new FAsyncTask<FAudioMixerSourceWorker>(this, StartId, EndId));
-				StartId = EndId;
-			}
-		}
-		NumSourceWorkers = SourceWorkers.Num();
-
 		// Cache the spatialization plugin
 		bUsingSpatializationPlugin = false;
 		SpatialInterfaceInfo = MixerDevice->GetCurrentSpatializationPluginInterfaceInfo();
@@ -2882,346 +2855,340 @@ namespace Audio
 		}
 	}
 
-	void FMixerSourceManager::ComputePostSourceEffectBufferForIdRange(bool bGenerateBuses, const int32 SourceIdStart, const int32 SourceIdEnd)
+	void FMixerSourceManager::ComputePostSourceEffectBuffer(const bool bGenerateBuses, const int32 SourceId)
 	{
 		CSV_SCOPED_TIMING_STAT(Audio, SourceEffectsBuffers);
-		CONDITIONAL_SCOPE_CYCLE_COUNTER(STAT_AudioMixerSourceEffectBuffers, (SourceIdStart < SourceIdEnd));
+		SCOPE_CYCLE_COUNTER(STAT_AudioMixerSourceEffectBuffers);
 
 		const bool bIsDebugModeEnabled = DebugSoloSources.Num() > 0;
 
-		for (int32 SourceId = SourceIdStart; SourceId < SourceIdEnd; ++SourceId)
-		{
-			FSourceInfo& SourceInfo = SourceInfos[SourceId];
+		FSourceInfo& SourceInfo = SourceInfos[SourceId];
 
-			if (!SourceInfo.bIsBusy || !SourceInfo.bIsPlaying || SourceInfo.bIsPaused || SourceInfo.bIsPausedForQuantization || (SourceInfo.bIsDone && SourceInfo.bEffectTailsDone))
-			{
-				continue;
-			}
+		if (!SourceInfo.bIsBusy || !SourceInfo.bIsPlaying || SourceInfo.bIsPaused || SourceInfo.bIsPausedForQuantization || (SourceInfo.bIsDone && SourceInfo.bEffectTailsDone))
+		{
+			return;
+		}
 
-			const bool bIsSourceBus = SourceInfo.AudioBusId != INDEX_NONE;
-			if ((bGenerateBuses && !bIsSourceBus) || (!bGenerateBuses && bIsSourceBus))
-			{
-				continue;
-			}
+		const bool bIsSourceBus = SourceInfo.AudioBusId != INDEX_NONE;
+		if ((bGenerateBuses && !bIsSourceBus) || (!bGenerateBuses && bIsSourceBus))
+		{
+			return;
+		}
 
-			// Copy and store the current state of the pre-distance attenuation buffer before we feed it through our source effects
-			// This is used by pre-effect sends
-			if (SourceInfo.AudioBusSends[(int32)EBusSendType::PreEffect].Num() > 0)
-			{
-				SourceInfo.PreEffectBuffer.Reset();
-				SourceInfo.PreEffectBuffer.Reserve(SourceInfo.PreDistanceAttenuationBuffer.Num());
+		// Copy and store the current state of the pre-distance attenuation buffer before we feed it through our source effects
+		// This is used by pre-effect sends
+		if (SourceInfo.AudioBusSends[(int32)EBusSendType::PreEffect].Num() > 0)
+		{
+			SourceInfo.PreEffectBuffer.Reset();
+			SourceInfo.PreEffectBuffer.Reserve(SourceInfo.PreDistanceAttenuationBuffer.Num());
 
-				FMemory::Memcpy(SourceInfo.PreEffectBuffer.GetData(), SourceInfo.PreDistanceAttenuationBuffer.GetData(), sizeof(float)*SourceInfo.PreDistanceAttenuationBuffer.Num());
-			}
+			FMemory::Memcpy(SourceInfo.PreEffectBuffer.GetData(), SourceInfo.PreDistanceAttenuationBuffer.GetData(), sizeof(float) * SourceInfo.PreDistanceAttenuationBuffer.Num());
+		}
 
-			float* PreDistanceAttenBufferPtr = SourceInfo.PreDistanceAttenuationBuffer.GetData();
-			const int32 NumSamples = SourceInfo.PreDistanceAttenuationBuffer.Num();
+		float* PreDistanceAttenBufferPtr = SourceInfo.PreDistanceAttenuationBuffer.GetData();
+		const int32 NumSamples = SourceInfo.PreDistanceAttenuationBuffer.Num();
 
-			TArrayView<float> PreDistanceAttenBufferView(PreDistanceAttenBufferPtr, NumSamples);
+		TArrayView<float> PreDistanceAttenBufferView(PreDistanceAttenBufferPtr, NumSamples);
 
-			// Update volume fade information if we're stopping
+		// Update volume fade information if we're stopping
+		{
+			float VolumeStart = 1.0f;
+			float VolumeDestination = 1.0f;
+			if (SourceInfo.bIsStopping)
 			{
-				float VolumeStart = 1.0f;
-				float VolumeDestination = 1.0f;
-				if (SourceInfo.bIsStopping)
+				int32 NumFadeFrames = FMath::Min(SourceInfo.VolumeFadeNumFrames - SourceInfo.VolumeFadeFramePosition, NumOutputFrames);
+
+				SourceInfo.VolumeFadeFramePosition += NumFadeFrames;
+				SourceInfo.VolumeSourceDestination = SourceInfo.VolumeFadeSlope * (float)SourceInfo.VolumeFadeFramePosition + SourceInfo.VolumeFadeStart;
+
+				if (FMath::IsNearlyZero(SourceInfo.VolumeSourceDestination, KINDA_SMALL_NUMBER))
 				{
-					int32 NumFadeFrames = FMath::Min(SourceInfo.VolumeFadeNumFrames - SourceInfo.VolumeFadeFramePosition, NumOutputFrames);
+					SourceInfo.VolumeSourceDestination = 0.0f;
+				}
 
-					SourceInfo.VolumeFadeFramePosition += NumFadeFrames;
-					SourceInfo.VolumeSourceDestination = SourceInfo.VolumeFadeSlope * (float)SourceInfo.VolumeFadeFramePosition + SourceInfo.VolumeFadeStart;
+				const int32 NumFadeSamples = NumFadeFrames * SourceInfo.NumInputChannels;
 
-					if (FMath::IsNearlyZero(SourceInfo.VolumeSourceDestination, KINDA_SMALL_NUMBER))
+				VolumeStart = SourceInfo.VolumeSourceStart;
+				VolumeDestination = SourceInfo.VolumeSourceDestination;
+				if (MixerDevice->IsModulationPluginEnabled() && MixerDevice->ModulationInterface.IsValid())
+				{
+					const bool bHasProcessed = SourceInfo.VolumeModulation.GetHasProcessed();
+					const float ModVolumeStart = SourceInfo.VolumeModulation.GetValue();
+					SourceInfo.VolumeModulation.ProcessControl(SourceInfo.VolumeModulationBase);
+					const float ModVolumeEnd = SourceInfo.VolumeModulation.GetValue();
+					if (bHasProcessed)
 					{
-						SourceInfo.VolumeSourceDestination = 0.0f;
+						VolumeStart *= ModVolumeStart;
 					}
-
-					const int32 NumFadeSamples = NumFadeFrames * SourceInfo.NumInputChannels;
-
-					VolumeStart = SourceInfo.VolumeSourceStart;
-					VolumeDestination = SourceInfo.VolumeSourceDestination;
-					if (MixerDevice->IsModulationPluginEnabled() && MixerDevice->ModulationInterface.IsValid())
+					else
 					{
-						const bool bHasProcessed = SourceInfo.VolumeModulation.GetHasProcessed();
-						const float ModVolumeStart = SourceInfo.VolumeModulation.GetValue();
-						SourceInfo.VolumeModulation.ProcessControl(SourceInfo.VolumeModulationBase);
-						const float ModVolumeEnd = SourceInfo.VolumeModulation.GetValue();
-						if (bHasProcessed)
-						{
-							VolumeStart *= ModVolumeStart;
-						}
-						else
-						{
-							VolumeStart *= ModVolumeEnd;
-						}
-						VolumeDestination *= ModVolumeEnd;
-						GameThreadInfo.ModulationVolume[SourceId] = ModVolumeEnd;
+						VolumeStart *= ModVolumeEnd;
 					}
+					VolumeDestination *= ModVolumeEnd;
+					GameThreadInfo.ModulationVolume[SourceId] = ModVolumeEnd;
+				}
 
-					TArrayView<float> PreDistanceAttenBufferFadeSamplesView(PreDistanceAttenBufferPtr, NumFadeSamples);
-					Audio::ArrayFade(PreDistanceAttenBufferFadeSamplesView, VolumeStart, VolumeDestination);
+				TArrayView<float> PreDistanceAttenBufferFadeSamplesView(PreDistanceAttenBufferPtr, NumFadeSamples);
+				Audio::ArrayFade(PreDistanceAttenBufferFadeSamplesView, VolumeStart, VolumeDestination);
 
-					// Zero the rest of the buffer
-					if (NumFadeFrames < NumOutputFrames)
-					{
-						int32 SamplesLeft = NumSamples - NumFadeSamples;
+				// Zero the rest of the buffer
+				if (NumFadeFrames < NumOutputFrames)
+				{
+					int32 SamplesLeft = NumSamples - NumFadeSamples;
 
-						// Protect memzero call with some sanity checking on the inputs.
-						if (SamplesLeft > 0 && NumFadeSamples >= 0 && NumFadeSamples < NumSamples)
-						{
-							FMemory::Memzero(&PreDistanceAttenBufferPtr[NumFadeSamples], sizeof(float) * SamplesLeft);
-						}
+					// Protect memzero call with some sanity checking on the inputs.
+					if (SamplesLeft > 0 && NumFadeSamples >= 0 && NumFadeSamples < NumSamples)
+					{
+						FMemory::Memzero(&PreDistanceAttenBufferPtr[NumFadeSamples], sizeof(float) * SamplesLeft);
 					}
 				}
-				else
+			}
+			else
+			{
+				VolumeStart = SourceInfo.VolumeSourceStart;
+				VolumeDestination = SourceInfo.VolumeSourceDestination;
+				if (MixerDevice->IsModulationPluginEnabled() && MixerDevice->ModulationInterface.IsValid())
 				{
-					VolumeStart = SourceInfo.VolumeSourceStart;
-					VolumeDestination = SourceInfo.VolumeSourceDestination;
-					if (MixerDevice->IsModulationPluginEnabled() && MixerDevice->ModulationInterface.IsValid())
+					const bool bHasProcessed = SourceInfo.VolumeModulation.GetHasProcessed();
+					const float ModVolumeStart = SourceInfo.VolumeModulation.GetValue();
+					SourceInfo.VolumeModulation.ProcessControl(SourceInfo.VolumeModulationBase);
+					const float ModVolumeEnd = SourceInfo.VolumeModulation.GetValue();
+					if (bHasProcessed)
 					{
-						const bool bHasProcessed = SourceInfo.VolumeModulation.GetHasProcessed();
-						const float ModVolumeStart = SourceInfo.VolumeModulation.GetValue();
-						SourceInfo.VolumeModulation.ProcessControl(SourceInfo.VolumeModulationBase);
-						const float ModVolumeEnd = SourceInfo.VolumeModulation.GetValue();
-						if (bHasProcessed)
-						{
-							VolumeStart *= ModVolumeStart;
-						}
-						else
-						{
-							VolumeStart *= ModVolumeEnd;
-						}
-						VolumeDestination *= ModVolumeEnd;
-						GameThreadInfo.ModulationVolume[SourceId] = ModVolumeEnd;
+						VolumeStart *= ModVolumeStart;
 					}
-
-					Audio::ArrayFade(PreDistanceAttenBufferView, VolumeStart, VolumeDestination);
+					else
+					{
+						VolumeStart *= ModVolumeEnd;
+					}
+					VolumeDestination *= ModVolumeEnd;
+					GameThreadInfo.ModulationVolume[SourceId] = ModVolumeEnd;
 				}
 
+				Audio::ArrayFade(PreDistanceAttenBufferView, VolumeStart, VolumeDestination);
+			}
+
 #if UE_AUDIO_PROFILERTRACE_ENABLED
-				const bool bChannelEnabled = UE_TRACE_CHANNELEXPR_IS_ENABLED(AudioMixerChannel);
-				if (bChannelEnabled)
-				{
-					UE_TRACE_LOG(Audio, MixerSourceVolume, AudioMixerChannel)
-						<< MixerSourceVolume.DeviceId(MixerDevice->DeviceID)
-						<< MixerSourceVolume.Timestamp(FPlatformTime::Cycles64())
-						<< MixerSourceVolume.PlayOrder(SourceInfo.PlayOrder)
-						<< MixerSourceVolume.ActiveSoundPlayOrder(SourceInfo.ActiveSoundPlayOrder)
-						<< MixerSourceVolume.Volume(VolumeDestination);
-
-					UE_TRACE_LOG(Audio, MixerSourceDistanceAttenuation, AudioMixerChannel)
-						<< MixerSourceDistanceAttenuation.DeviceId(MixerDevice->DeviceID)
-						<< MixerSourceDistanceAttenuation.Timestamp(FPlatformTime::Cycles64())
-						<< MixerSourceDistanceAttenuation.PlayOrder(SourceInfo.PlayOrder)
-						<< MixerSourceDistanceAttenuation.DistanceAttenuation(SourceInfo.DistanceAttenuationSourceDestination);
-				}
-#endif // UE_AUDIO_PROFILERTRACE_ENABLED
+			const bool bChannelEnabled = UE_TRACE_CHANNELEXPR_IS_ENABLED(AudioMixerChannel);
+			if (bChannelEnabled)
+			{
+				UE_TRACE_LOG(Audio, MixerSourceVolume, AudioMixerChannel)
+					<< MixerSourceVolume.DeviceId(MixerDevice->DeviceID)
+					<< MixerSourceVolume.Timestamp(FPlatformTime::Cycles64())
+					<< MixerSourceVolume.PlayOrder(SourceInfo.PlayOrder)
+					<< MixerSourceVolume.ActiveSoundPlayOrder(SourceInfo.ActiveSoundPlayOrder)
+					<< MixerSourceVolume.Volume(VolumeDestination);
+
+				UE_TRACE_LOG(Audio, MixerSourceDistanceAttenuation, AudioMixerChannel)
+					<< MixerSourceDistanceAttenuation.DeviceId(MixerDevice->DeviceID)
+					<< MixerSourceDistanceAttenuation.Timestamp(FPlatformTime::Cycles64())
+					<< MixerSourceDistanceAttenuation.PlayOrder(SourceInfo.PlayOrder)
+					<< MixerSourceDistanceAttenuation.DistanceAttenuation(SourceInfo.DistanceAttenuationSourceDestination);
 			}
+#endif // UE_AUDIO_PROFILERTRACE_ENABLED
+		}
 
-			SourceInfo.VolumeSourceStart = SourceInfo.VolumeSourceDestination;
+		SourceInfo.VolumeSourceStart = SourceInfo.VolumeSourceDestination;
 
-			// Now process the effect chain if it exists
-			if (!DisableSourceEffectsCvar && SourceInfo.SourceEffects.Num() > 0)
-			{
-				// Prepare this source's effect chain input data
-				SourceInfo.SourceEffectInputData.CurrentVolume = SourceInfo.VolumeSourceDestination;
+		// Now process the effect chain if it exists
+		if (!DisableSourceEffectsCvar && SourceInfo.SourceEffects.Num() > 0)
+		{
+			// Prepare this source's effect chain input data
+			SourceInfo.SourceEffectInputData.CurrentVolume = SourceInfo.VolumeSourceDestination;
 
-				const float Pitch = Audio::GetFrequencyMultiplier(SourceInfo.PitchModulation.GetValue());
-				SourceInfo.SourceEffectInputData.CurrentPitch = SourceInfo.PitchSourceParam.GetValue() * Pitch;
-				SourceInfo.SourceEffectInputData.AudioClock = MixerDevice->GetAudioClock();
-				if (SourceInfo.NumInputFrames > 0)
-				{
-					SourceInfo.SourceEffectInputData.CurrentPlayFraction = (float)SourceInfo.NumFramesPlayed / SourceInfo.NumInputFrames;
-				}
-				SourceInfo.SourceEffectInputData.SpatParams = SourceInfo.SpatParams;
+			const float Pitch = Audio::GetFrequencyMultiplier(SourceInfo.PitchModulation.GetValue());
+			SourceInfo.SourceEffectInputData.CurrentPitch = SourceInfo.PitchSourceParam.GetValue() * Pitch;
+			SourceInfo.SourceEffectInputData.AudioClock = MixerDevice->GetAudioClock();
+			if (SourceInfo.NumInputFrames > 0)
+			{
+				SourceInfo.SourceEffectInputData.CurrentPlayFraction = (float)SourceInfo.NumFramesPlayed / SourceInfo.NumInputFrames;
+			}
+			SourceInfo.SourceEffectInputData.SpatParams = SourceInfo.SpatParams;
 
-				// Get a ptr to pre-distance attenuation buffer ptr
-				float* OutputSourceEffectBufferPtr = SourceInfo.SourceEffectScratchBuffer.GetData();
+			// Get a ptr to pre-distance attenuation buffer ptr
+			float* OutputSourceEffectBufferPtr = SourceInfo.SourceEffectScratchBuffer.GetData();
 
-				SourceInfo.SourceEffectInputData.InputSourceEffectBufferPtr = SourceInfo.PreDistanceAttenuationBuffer.GetData();
-				SourceInfo.SourceEffectInputData.NumSamples = NumSamples;
+			SourceInfo.SourceEffectInputData.InputSourceEffectBufferPtr = SourceInfo.PreDistanceAttenuationBuffer.GetData();
+			SourceInfo.SourceEffectInputData.NumSamples = NumSamples;
 
-				// Loop through the effect chain passing in buffers
-				FScopeLock ScopeLock(&EffectChainMutationCriticalSection);
+			// Loop through the effect chain passing in buffers
+			FScopeLock ScopeLock(&EffectChainMutationCriticalSection);
+			{
+				for (TSoundEffectSourcePtr& SoundEffectSource : SourceInfo.SourceEffects)
 				{
-					for (TSoundEffectSourcePtr& SoundEffectSource : SourceInfo.SourceEffects)
+					bool bPresetUpdated = false;
+					if (SoundEffectSource->IsActive())
 					{
-						bool bPresetUpdated = false;
-						if (SoundEffectSource->IsActive())
-						{
-							bPresetUpdated = SoundEffectSource->Update();
-						}
+						bPresetUpdated = SoundEffectSource->Update();
+					}
 
-						if (SoundEffectSource->IsActive())
-						{
-							SoundEffectSource->ProcessAudio(SourceInfo.SourceEffectInputData, OutputSourceEffectBufferPtr);
+					if (SoundEffectSource->IsActive())
+					{
+						SoundEffectSource->ProcessAudio(SourceInfo.SourceEffectInputData, OutputSourceEffectBufferPtr);
 
-							// Copy output to input
-							FMemory::Memcpy(SourceInfo.SourceEffectInputData.InputSourceEffectBufferPtr, OutputSourceEffectBufferPtr, sizeof(float) * NumSamples);
-						}
+						// Copy output to input
+						FMemory::Memcpy(SourceInfo.SourceEffectInputData.InputSourceEffectBufferPtr, OutputSourceEffectBufferPtr, sizeof(float) * NumSamples);
 					}
 				}
 			}
+		}
 
-			const bool bWasEffectTailsDone = SourceInfo.bEffectTailsDone;
+		const bool bWasEffectTailsDone = SourceInfo.bEffectTailsDone;
 
-			if (!DisableEnvelopeFollowingCvar)
+		if (!DisableEnvelopeFollowingCvar)
+		{
+			// Compute the source envelope using pre-distance attenuation buffer
+			float AverageSampleValue = Audio::ArrayGetAverageAbsValue(PreDistanceAttenBufferView);
+			SourceInfo.SourceEnvelopeValue = SourceInfo.SourceEnvelopeFollower.ProcessSample(AverageSampleValue);
+			SourceInfo.SourceEnvelopeValue = FMath::Clamp(SourceInfo.SourceEnvelopeValue, 0.f, 1.f);
+
+			SourceInfo.bEffectTailsDone = SourceInfo.bEffectTailsDone || SourceInfo.SourceEnvelopeValue < ENVELOPE_TAIL_THRESHOLD;
+		}
+		else
+		{
+			SourceInfo.bEffectTailsDone = true;
+		}
+
+		if (!bWasEffectTailsDone && SourceInfo.bEffectTailsDone)
+		{
+			SourceInfo.SourceListener->OnEffectTailsDone();
+		}
+
+		const bool bModActive = MixerDevice->IsModulationPluginEnabled() && MixerDevice->ModulationInterface.IsValid();
+		bool bUpdateModFilters = bModActive && (SourceInfo.bModFiltersUpdated || SourceInfo.LowpassModulation.IsActive() || SourceInfo.HighpassModulation.IsActive());
+
+		if (SourceInfo.IsRenderingToSubmixes() || bUpdateModFilters)
+		{
+			// Only scale with distance attenuation and send to source audio to plugins if we're not in output-to-bus only mode
+			const int32 NumOutputSamplesThisSource = NumOutputFrames * SourceInfo.NumInputChannels;
+
+			if (!SourceInfo.IsRenderingToSubmixes())
 			{
-				// Compute the source envelope using pre-distance attenuation buffer
-				float AverageSampleValue = Audio::ArrayGetAverageAbsValue(PreDistanceAttenBufferView);
-				SourceInfo.SourceEnvelopeValue = SourceInfo.SourceEnvelopeFollower.ProcessSample(AverageSampleValue);
-				SourceInfo.SourceEnvelopeValue = FMath::Clamp(SourceInfo.SourceEnvelopeValue, 0.f, 1.f);
+				SourceInfo.LowpassModulation.ProcessControl(SourceInfo.LowpassModulationBase);
+				SourceInfo.LowPassFilter.StartFrequencyInterpolation(SourceInfo.LowpassModulation.GetValue(), NumOutputFrames);
 
-				SourceInfo.bEffectTailsDone = SourceInfo.bEffectTailsDone || SourceInfo.SourceEnvelopeValue < ENVELOPE_TAIL_THRESHOLD;
+				SourceInfo.HighpassModulation.ProcessControl(SourceInfo.HighpassModulationBase);
+				SourceInfo.HighPassFilter.StartFrequencyInterpolation(SourceInfo.HighpassModulation.GetValue(), NumOutputFrames);
 			}
-			else
+			else if (bUpdateModFilters)
 			{
-				SourceInfo.bEffectTailsDone = true;
-			}
+				const float LowpassFreq = FMath::Min(SourceInfo.LowpassModulationBase, SourceInfo.LowPassFreq);
+				SourceInfo.LowpassModulation.ProcessControl(LowpassFreq);
+				SourceInfo.LowPassFilter.StartFrequencyInterpolation(SourceInfo.LowpassModulation.GetValue(), NumOutputFrames);
 
-			if (!bWasEffectTailsDone && SourceInfo.bEffectTailsDone)
-			{
-				SourceInfo.SourceListener->OnEffectTailsDone();
+				const float HighpassFreq = FMath::Max(SourceInfo.HighpassModulationBase, SourceInfo.HighPassFreq);
+				SourceInfo.HighpassModulation.ProcessControl(HighpassFreq);
+				SourceInfo.HighPassFilter.StartFrequencyInterpolation(SourceInfo.HighpassModulation.GetValue(), NumOutputFrames);
 			}
 
-			const bool bModActive = MixerDevice->IsModulationPluginEnabled() && MixerDevice->ModulationInterface.IsValid();
-			bool bUpdateModFilters = bModActive && (SourceInfo.bModFiltersUpdated || SourceInfo.LowpassModulation.IsActive() || SourceInfo.HighpassModulation.IsActive());
+			const bool bBypassLPF = DisableFilteringCvar || (SourceInfo.LowPassFilter.GetCutoffFrequency() >= (MAX_FILTER_FREQUENCY - KINDA_SMALL_NUMBER));
+			const bool bBypassHPF = DisableFilteringCvar || DisableHPFilteringCvar || (SourceInfo.HighPassFilter.GetCutoffFrequency() <= (MIN_FILTER_FREQUENCY + KINDA_SMALL_NUMBER));
 
-			if (SourceInfo.IsRenderingToSubmixes() || bUpdateModFilters)
-			{
-				// Only scale with distance attenuation and send to source audio to plugins if we're not in output-to-bus only mode
-				const int32 NumOutputSamplesThisSource = NumOutputFrames * SourceInfo.NumInputChannels;
-
-				if (!SourceInfo.IsRenderingToSubmixes())
-				{
-					SourceInfo.LowpassModulation.ProcessControl(SourceInfo.LowpassModulationBase);
-					SourceInfo.LowPassFilter.StartFrequencyInterpolation(SourceInfo.LowpassModulation.GetValue(), NumOutputFrames);
-
-					SourceInfo.HighpassModulation.ProcessControl(SourceInfo.HighpassModulationBase);
-					SourceInfo.HighPassFilter.StartFrequencyInterpolation(SourceInfo.HighpassModulation.GetValue(), NumOutputFrames);
-				}
-				else if (bUpdateModFilters)
-				{
-					const float LowpassFreq = FMath::Min(SourceInfo.LowpassModulationBase, SourceInfo.LowPassFreq);
-					SourceInfo.LowpassModulation.ProcessControl(LowpassFreq);
-					SourceInfo.LowPassFilter.StartFrequencyInterpolation(SourceInfo.LowpassModulation.GetValue(), NumOutputFrames);
+			float* SourceBuffer = SourceInfo.SourceBuffer.GetData();
+			float* HpfInputBuffer = PreDistanceAttenBufferPtr; // assume bypassing LPF (HPF uses input buffer as input)
 
-					const float HighpassFreq = FMath::Max(SourceInfo.HighpassModulationBase, SourceInfo.HighPassFreq);
-					SourceInfo.HighpassModulation.ProcessControl(HighpassFreq);
-					SourceInfo.HighPassFilter.StartFrequencyInterpolation(SourceInfo.HighpassModulation.GetValue(), NumOutputFrames);
-				}
+			if (!bBypassLPF)
+			{
+				// Not bypassing LPF, so tell HPF to use LPF output buffer as input
+				HpfInputBuffer = SourceBuffer;
 
-				const bool bBypassLPF = DisableFilteringCvar || (SourceInfo.LowPassFilter.GetCutoffFrequency() >= (MAX_FILTER_FREQUENCY - KINDA_SMALL_NUMBER));
-				const bool bBypassHPF = DisableFilteringCvar || DisableHPFilteringCvar || (SourceInfo.HighPassFilter.GetCutoffFrequency() <= (MIN_FILTER_FREQUENCY + KINDA_SMALL_NUMBER));
+				// process LPF audio block
+				SourceInfo.LowPassFilter.ProcessAudioBuffer(PreDistanceAttenBufferPtr, SourceBuffer, NumOutputSamplesThisSource);
+			}
 
-				float* SourceBuffer = SourceInfo.SourceBuffer.GetData();
-				float* HpfInputBuffer = PreDistanceAttenBufferPtr; // assume bypassing LPF (HPF uses input buffer as input)
+			if (!bBypassHPF)
+			{
+				// process HPF audio block
+				SourceInfo.HighPassFilter.ProcessAudioBuffer(HpfInputBuffer, SourceBuffer, NumOutputSamplesThisSource);
+			}
 
+#if UE_AUDIO_PROFILERTRACE_ENABLED
+			const bool bChannelEnabled = UE_TRACE_CHANNELEXPR_IS_ENABLED(AudioMixerChannel);
+			if (bChannelEnabled)
+			{
+				float LPFFrequency = MAX_FILTER_FREQUENCY;
 				if (!bBypassLPF)
 				{
-					// Not bypassing LPF, so tell HPF to use LPF output buffer as input
-					HpfInputBuffer = SourceBuffer;
-
-					// process LPF audio block
-					SourceInfo.LowPassFilter.ProcessAudioBuffer(PreDistanceAttenBufferPtr, SourceBuffer, NumOutputSamplesThisSource);
+					LPFFrequency = SourceInfo.LowpassModulation.GetValue();
 				}
 
+				float HPFFrequency = MIN_FILTER_FREQUENCY;
 				if (!bBypassHPF)
 				{
-					// process HPF audio block
-					SourceInfo.HighPassFilter.ProcessAudioBuffer(HpfInputBuffer, SourceBuffer, NumOutputSamplesThisSource);
+					HPFFrequency = SourceInfo.HighpassModulation.GetValue();
 				}
 
-#if UE_AUDIO_PROFILERTRACE_ENABLED
-				const bool bChannelEnabled = UE_TRACE_CHANNELEXPR_IS_ENABLED(AudioMixerChannel);
-				if (bChannelEnabled)
-				{
-					float LPFFrequency = MAX_FILTER_FREQUENCY;
-					if (!bBypassLPF)
-					{
-						LPFFrequency = SourceInfo.LowpassModulation.GetValue();
-					}
-
-					float HPFFrequency = MIN_FILTER_FREQUENCY;
-					if (!bBypassHPF)
-					{
-						HPFFrequency = SourceInfo.HighpassModulation.GetValue();
-					}
-
-					UE_TRACE_LOG(Audio, MixerSourceFilters, AudioMixerChannel)
-						<< MixerSourceFilters.DeviceId(MixerDevice->DeviceID)
-						<< MixerSourceFilters.Timestamp(FPlatformTime::Cycles64())
-						<< MixerSourceFilters.PlayOrder(SourceInfo.PlayOrder)
-						<< MixerSourceFilters.HPFFrequency(HPFFrequency)
-						<< MixerSourceFilters.LPFFrequency(LPFFrequency);
-					UE_TRACE_LOG(Audio, MixerSourceEnvelope, AudioMixerChannel)
-						<< MixerSourceEnvelope.DeviceId(MixerDevice->DeviceID)
-						<< MixerSourceEnvelope.Timestamp(FPlatformTime::Cycles64())
-						<< MixerSourceEnvelope.PlayOrder(SourceInfo.PlayOrder)
-						<< MixerSourceEnvelope.ActiveSoundPlayOrder(SourceInfo.ActiveSoundPlayOrder)
-						<< MixerSourceEnvelope.Envelope(SourceInfo.SourceEnvelopeValue);
-				}
+				UE_TRACE_LOG(Audio, MixerSourceFilters, AudioMixerChannel)
+					<< MixerSourceFilters.DeviceId(MixerDevice->DeviceID)
+					<< MixerSourceFilters.Timestamp(FPlatformTime::Cycles64())
+					<< MixerSourceFilters.PlayOrder(SourceInfo.PlayOrder)
+					<< MixerSourceFilters.HPFFrequency(HPFFrequency)
+					<< MixerSourceFilters.LPFFrequency(LPFFrequency);
+				UE_TRACE_LOG(Audio, MixerSourceEnvelope, AudioMixerChannel)
+					<< MixerSourceEnvelope.DeviceId(MixerDevice->DeviceID)
+					<< MixerSourceEnvelope.Timestamp(FPlatformTime::Cycles64())
+					<< MixerSourceEnvelope.PlayOrder(SourceInfo.PlayOrder)
+					<< MixerSourceEnvelope.ActiveSoundPlayOrder(SourceInfo.ActiveSoundPlayOrder)
+					<< MixerSourceEnvelope.Envelope(SourceInfo.SourceEnvelopeValue);
+			}
 #endif // UE_AUDIO_PROFILERTRACE_ENABLED
 
-				// We manually reset interpolation to avoid branches in filter code
-				SourceInfo.LowPassFilter.StopFrequencyInterpolation();
-				SourceInfo.HighPassFilter.StopFrequencyInterpolation();
+			// We manually reset interpolation to avoid branches in filter code
+			SourceInfo.LowPassFilter.StopFrequencyInterpolation();
+			SourceInfo.HighPassFilter.StopFrequencyInterpolation();
 
-				if (bBypassLPF && bBypassHPF)
-				{
-					FMemory::Memcpy(SourceBuffer, PreDistanceAttenBufferPtr, NumSamples * sizeof(float));
-				}
+			if (bBypassLPF && bBypassHPF)
+			{
+				FMemory::Memcpy(SourceBuffer, PreDistanceAttenBufferPtr, NumSamples * sizeof(float));
 			}
+		}
 
-			if (SourceInfo.IsRenderingToSubmixes() || SpatialInterfaceInfo.bSpatializationIsExternalSend || SourceInfo.AudioLink.IsValid())
-			{
-				// Apply distance attenuation
-				ApplyDistanceAttenuation(SourceInfo, NumSamples);
+		if (SourceInfo.IsRenderingToSubmixes() || SpatialInterfaceInfo.bSpatializationIsExternalSend || SourceInfo.AudioLink.IsValid())
+		{
+			// Apply distance attenuation
+			ApplyDistanceAttenuation(SourceInfo, NumSamples);
 
-				FMixerSourceSubmixOutputBuffer& SourceSubmixOutputBuffer = SourceSubmixOutputBuffers[SourceId];
+			FMixerSourceSubmixOutputBuffer& SourceSubmixOutputBuffer = SourceSubmixOutputBuffers[SourceId];
 
-				// Send source audio to plugins
-				ComputePluginAudio(SourceInfo, SourceSubmixOutputBuffer, SourceId, NumSamples);
-			}
+			// Send source audio to plugins
+			ComputePluginAudio(SourceInfo, SourceSubmixOutputBuffer, SourceId, NumSamples);
+		}
 
-			// Check the source effect tails condition
-			if (SourceInfo.bIsLastBuffer && SourceInfo.bEffectTailsDone)
-			{
-				// If we're done and our tails our done, clear everything out
-				SourceInfo.CurrentFrameValues.Reset();
-				SourceInfo.NextFrameValues.Reset();
-				SourceInfo.CurrentPCMBuffer = nullptr;
-			}
+		// Check the source effect tails condition
+		if (SourceInfo.bIsLastBuffer && SourceInfo.bEffectTailsDone)
+		{
+			// If we're done and our tails our done, clear everything out
+			SourceInfo.CurrentFrameValues.Reset();
+			SourceInfo.NextFrameValues.Reset();
+			SourceInfo.CurrentPCMBuffer = nullptr;
 		}
 	}
 
-	void FMixerSourceManager::ComputeOutputBuffersForIdRange(const bool bGenerateBuses, const int32 SourceIdStart, const int32 SourceIdEnd)
+	void FMixerSourceManager::ComputeOutputBuffers(const bool bGenerateBuses, const int32 SourceId)
 	{
 		CSV_SCOPED_TIMING_STAT(Audio, SourceOutputBuffers);
-		CONDITIONAL_SCOPE_CYCLE_COUNTER(STAT_AudioMixerSourceOutputBuffers, (SourceIdStart < SourceIdEnd));
+		SCOPE_CYCLE_COUNTER(STAT_AudioMixerSourceOutputBuffers);
 
-		for (int32 SourceId = SourceIdStart; SourceId < SourceIdEnd; ++SourceId)
-		{
-			FSourceInfo& SourceInfo = SourceInfos[SourceId];
-
-			// Don't need to compute anything if the source is not playing or paused (it will remain at 0.0 volume)
-			// Note that effect chains will still be able to continue to compute audio output. The source output 
-			// will simply stop being read from.
-			if (!SourceInfo.bIsBusy || !SourceInfo.bIsPlaying || (SourceInfo.bIsDone && SourceInfo.bEffectTailsDone))
-			{
-				continue;
-			}
+		FSourceInfo& SourceInfo = SourceInfos[SourceId];
 
-			// If we're in generate buses mode and not a bus, or vice versa, or if we're set to only output audio to buses.
-			// If set to output buses, no need to do any panning for the source. The buses will do the panning.
-			const bool bIsSourceBus = SourceInfo.AudioBusId != INDEX_NONE;
-			if ((bGenerateBuses && !bIsSourceBus) || (!bGenerateBuses && bIsSourceBus) || !SourceInfo.IsRenderingToSubmixes())
-			{
-				continue;
-			}
+		// Don't need to compute anything if the source is not playing or paused (it will remain at 0.0 volume)
+		// Note that effect chains will still be able to continue to compute audio output. The source output 
+		// will simply stop being read from.
+		if (!SourceInfo.bIsBusy || !SourceInfo.bIsPlaying || (SourceInfo.bIsDone && SourceInfo.bEffectTailsDone))
+		{
+			return;
+		}
 
-			FMixerSourceSubmixOutputBuffer& SourceSubmixOutputBuffer = SourceSubmixOutputBuffers[SourceId];
-			SourceSubmixOutputBuffer.ComputeOutput(SourceInfo.SpatParams);
+		// If we're in generate buses mode and not a bus, or vice versa, or if we're set to only output audio to buses.
+		// If set to output buses, no need to do any panning for the source. The buses will do the panning.
+		const bool bIsSourceBus = SourceInfo.AudioBusId != INDEX_NONE;
+		if ((bGenerateBuses && !bIsSourceBus) || (!bGenerateBuses && bIsSourceBus) || !SourceInfo.IsRenderingToSubmixes())
+		{
+			return;
 		}
+
+		FMixerSourceSubmixOutputBuffer& SourceSubmixOutputBuffer = SourceSubmixOutputBuffers[SourceId];
+		SourceSubmixOutputBuffer.ComputeOutput(SourceInfo.SpatParams);
 	}
 
 	void FMixerSourceManager::GenerateSourceAudio(const bool bGenerateBuses, const int32 SourceIdStart, const int32 SourceIdEnd)
@@ -3234,10 +3201,16 @@ namespace Audio
 		}
 
 		// Compute the audio source buffers after their individual effect chain processing
-		ComputePostSourceEffectBufferForIdRange(bGenerateBuses, SourceIdStart, SourceIdEnd);
+		for (int32 SourceId = SourceIdStart; SourceId < SourceIdEnd; ++SourceId)
+		{
+			ComputePostSourceEffectBuffer(bGenerateBuses, SourceId);
+		}
 
 		// Get the audio for the output buffers
-		ComputeOutputBuffersForIdRange(bGenerateBuses, SourceIdStart, SourceIdEnd);
+		for (int32 SourceId = SourceIdStart; SourceId < SourceIdEnd; ++SourceId)
+		{
+			ComputeOutputBuffers(bGenerateBuses, SourceId);
+		}
 	}
 
 	void FMixerSourceManager::GenerateSourceAudio(const bool bGenerateBuses)
@@ -3252,21 +3225,31 @@ namespace Audio
 			return;
 		}
 
-		if (NumSourceWorkers > 0 && !DisableParallelSourceProcessingCvar)
+		if (!DisableParallelSourceProcessingCvar)
 		{
-			AUDIO_MIXER_CHECK(SourceWorkers.Num() == NumSourceWorkers);
-			for (int32 i = 0; i < SourceWorkers.Num(); ++i)
-			{
-				FAudioMixerSourceWorker& Worker = SourceWorkers[i]->GetTask();
-				Worker.SetGenerateBuses(bGenerateBuses);
-
-				SourceWorkers[i]->StartBackgroundTask();
-			}
+			auto RenderAll = [this, bGenerateBuses]()
+				{
+					for (int SourceId = 0; SourceId < NumTotalSources; ++SourceId)
+					{
+						// Skip launching a task if there's nothing to do for this source
+						const FSourceInfo& SourceInfo = SourceInfos[SourceId];
+						const bool bIsSourceBus = SourceInfo.AudioBusId != INDEX_NONE;
+						if (!SourceInfo.bIsPlaying || !SourceInfo.bIsBusy || bGenerateBuses != bIsSourceBus)
+							continue;
+
+						// Add a nested task to render this source
+						auto RenderOne = [this, bGenerateBuses, SourceId]()
+							{
+								ComputeSourceBuffer(bGenerateBuses, SourceId);
+								ComputePostSourceEffectBuffer(bGenerateBuses, SourceId);
+								ComputeOutputBuffers(bGenerateBuses, SourceId);
+							};
+						UE::Tasks::AddNested(UE::Tasks::Launch(TEXT("Render Audio Source"), MoveTemp(RenderOne), LowLevelTasks::ETaskPriority::High));
+					}
+				};
 
-			for (int32 i = 0; i < SourceWorkers.Num(); ++i)
-			{
-				SourceWorkers[i]->EnsureCompletion();
-			}
+			// Launch a single task to do all the rendering, and block on its completion
+			UE::Tasks::Launch(TEXT("Render All Audio Sources"), MoveTemp(RenderAll), LowLevelTasks::ETaskPriority::High).Wait();
 		}
 		else
 		{
diff --git a/Engine/Source/Runtime/AudioMixer/Private/AudioMixerSourceManager.h b/Engine/Source/Runtime/AudioMixer/Private/AudioMixerSourceManager.h
--- a/Engine/Source/Runtime/AudioMixer/Private/AudioMixerSourceManager.h
+++ b/Engine/Source/Runtime/AudioMixer/Private/AudioMixerSourceManager.h
@@ -173,9 +173,6 @@ namespace Audio
 	{
 		// Total number of sources to use in the source manager
 		int32 NumSources = 0;
-
-		// Number of worker threads to use for the source manager.
-		int32 NumSourceWorkers = 0;
 	};
 
 	class FMixerSourceManager
@@ -329,8 +326,8 @@ namespace Audio
 		void GenerateSourceAudio(const bool bGenerateBuses, const int32 SourceIdStart, const int32 SourceIdEnd);
 
 		void ComputeSourceBuffer(const bool bGenerateBuses, const int32 SourceId);
-		void ComputePostSourceEffectBufferForIdRange(const bool bGenerateBuses, const int32 SourceIdStart, const int32 SourceIdEnd);
-		void ComputeOutputBuffersForIdRange(const bool bGenerateBuses, const int32 SourceIdStart, const int32 SourceIdEnd);
+		void ComputePostSourceEffectBuffer(const bool bGenerateBuses, const int32 SourceId);
+		void ComputeOutputBuffers(const bool bGenerateBuses, const int32 SourceId);
 
 		void ConnectBusPatches();
 		void ComputeBuses();
@@ -367,41 +364,6 @@ namespace Audio
 
 		void AudioMixerThreadCommand(TFunction<void()>&& InFunction, const char* InDebugString = nullptr, bool bInDeferExecution = false);
 
-		static const int32 NUM_BYTES_PER_SAMPLE = 2;
-
-		// Private class which perform source buffer processing in a worker task
-		class FAudioMixerSourceWorker : public FNonAbandonableTask
-		{
-			FMixerSourceManager* SourceManager;
-			int32 StartSourceId;
-			int32 EndSourceId;
-			bool bGenerateBuses;
-
-		public:
-			FAudioMixerSourceWorker(FMixerSourceManager* InSourceManager, const int32 InStartSourceId, const int32 InEndSourceId)
-				: SourceManager(InSourceManager)
-				, StartSourceId(InStartSourceId)
-				, EndSourceId(InEndSourceId)
-				, bGenerateBuses(false)
-			{
-			}
-
-			void SetGenerateBuses(bool bInGenerateBuses)
-			{
-				bGenerateBuses = bInGenerateBuses;
-			}
-
-			void DoWork()
-			{
-				SourceManager->GenerateSourceAudio(bGenerateBuses, StartSourceId, EndSourceId);
-			}
-
-			FORCEINLINE TStatId GetStatId() const
-			{
-				RETURN_QUICK_DECLARE_CYCLE_STAT(FAudioMixerSourceWorker, STATGROUP_ThreadPoolAsyncTasks);
-			}
-		};
-
 		// Critical section to ensure mutating effect chains is thread-safe
 		FCriticalSection EffectChainMutationCriticalSection;
 
@@ -642,9 +604,6 @@ namespace Audio
 		TMap<FAudioBusKey, TSharedPtr<FMixerAudioBus>> AudioBuses; 
 		TArray<FAudioBusKey> AudioBusKeys_AudioThread;
 
-		// Async task workers for processing sources in parallel
-		TArray<FAsyncTask<FAudioMixerSourceWorker>*> SourceWorkers;
-
 		// Array of task data waiting to finished. Processed on audio render thread.
 		TArray<TSharedPtr<FMixerSourceBuffer, ESPMode::ThreadSafe>> PendingSourceBuffers;
 
@@ -667,7 +626,6 @@ namespace Audio
 		int32 NumTotalSources;
 		int32 NumOutputFrames;
 		int32 NumOutputSamples;
-		int32 NumSourceWorkers;
 
 		// Commands queued up to execute
 		FThreadSafeCounter NumCommands;
