AudioLink - Allow > Strero sends when using AudioLink. bIs3d  refactor to make clear what's going on.

diff --git a/Engine/Source/Runtime/AudioMixer/Private/AudioMixerSource.cpp b/Engine/Source/Runtime/AudioMixer/Private/AudioMixerSource.cpp
--- a/Engine/Source/Runtime/AudioMixer/Private/AudioMixerSource.cpp
+++ b/Engine/Source/Runtime/AudioMixer/Private/AudioMixerSource.cpp
@@ -200,6 +200,37 @@ namespace Audio
 		FreeResources();
 	}
 
+	bool FMixerSource::Is3D(const FSoundBuffer* SoundBuffer) const
+	{
+		// If we're not spatialized, we're not 3d.
+		if (!WaveInstance->GetUseSpatialization())
+		{
+			return false;
+		}
+
+		// If we're using object spatialization, we can't be panned in 3d.
+		if (UseObjectBasedSpatialization())
+		{
+			return false;
+		}
+
+		// We've got this far, and sending this source to audiolink, allow before
+		// stereo check below.
+		if (AudioLink.IsValid())
+		{
+			return true;
+		}
+		
+		// 3d panner can't handle more than stereo currently.
+		if (SoundBuffer->NumChannels > 2)
+		{
+			return false;
+		}
+
+		// We are 3d.
+		return true;
+	}
+
 	bool FMixerSource::Init(FWaveInstance* InWaveInstance)
 	{
 		AUDIO_MIXER_TRACE_CPUPROFILER_EVENT_SCOPE(AudioMixerSource::Init);
@@ -516,7 +547,7 @@ namespace Audio
 					(int32)MixerBuffer->GetType(),WaveData->CachedRealtimeFirstBuffer);
 
 			// Whether or not we're 3D
-			bIs3D = !UseObjectBasedSpatialization() && WaveInstance->GetUseSpatialization() && SoundBuffer->NumChannels < 3;
+			bIs3D = Is3D(SoundBuffer);
 
 			// Pass on the fact that we're 3D to the init params
 			InitParams.bIs3D = bIs3D;
diff --git a/Engine/Source/Runtime/AudioMixer/Private/AudioMixerSource.h b/Engine/Source/Runtime/AudioMixer/Private/AudioMixerSource.h
--- a/Engine/Source/Runtime/AudioMixer/Private/AudioMixerSource.h
+++ b/Engine/Source/Runtime/AudioMixer/Private/AudioMixerSource.h
@@ -135,6 +135,8 @@ namespace Audio
 	private:
 		void UpdateSubmixSendLevels(const FSoundSubmixSendInfoBase& InSendInfo, EMixerSourceSubmixSendStage InSendStage, TSet<FMixerSubmixWeakPtr>& OutTouchedSubmixes);
 
+		bool Is3D(const FSoundBuffer* SoundBuffer) const;
+		
 		FMixerDevice* MixerDevice;
 		FMixerBuffer* MixerBuffer;
 		TSharedPtr<FMixerSourceBuffer, ESPMode::ThreadSafe> MixerSourceBuffer;
