Fix for sounds not playing on first try.
- SetNumFrames not working with OGG, so sanitze data from GetStreamingInfo
- Ogg failing to block for chunk loads causing race in editor
- Make Ogg not windows platform default.

diff --git a/Engine/Config/Windows/BaseWindowsEngine.ini b/Engine/Config/Windows/BaseWindowsEngine.ini
--- a/Engine/Config/Windows/BaseWindowsEngine.ini
+++ b/Engine/Config/Windows/BaseWindowsEngine.ini
@@ -2,8 +2,8 @@
 ; Defining below allows switching to audio mixer using -audiomixer commandline
 AudioCaptureModuleName=AudioCaptureWasapi
 AudioMixerModuleName=AudioMixerXAudio2
-PlatformFormat=OGG
-PlatformStreamingFormat=OGG
+PlatformFormat=BINKA
+PlatformStreamingFormat=BINKA
 
 ; Defines a platform-specific volume headroom (in dB) for audio to provide better platform consistency with respect to volume levels.
 PlatformHeadroomDB=-3
diff --git a/Engine/Source/Runtime/AudioMixer/Private/AudioMixerBuffer.cpp b/Engine/Source/Runtime/AudioMixer/Private/AudioMixerBuffer.cpp
--- a/Engine/Source/Runtime/AudioMixer/Private/AudioMixerBuffer.cpp
+++ b/Engine/Source/Runtime/AudioMixer/Private/AudioMixerBuffer.cpp
@@ -292,18 +292,29 @@ namespace Audio
 		if (Buffer->DecompressionState && Buffer->DecompressionState->StreamCompressedInfo(InWave, &QualityInfo))
 		{
 			// Refresh the wave data
-			const int64 NumFrames = QualityInfo.NumChannels && QualityInfo.SampleDataSize > 0 ?
-				QualityInfo.SampleDataSize / sizeof(int16) / QualityInfo.NumChannels : 0;
-			
-			InWave->SetNumFrames(IntCastChecked<uint32>(NumFrames));
-			InWave->SetSampleRate(QualityInfo.SampleRate, /* bFromDecoders */ true);
-			InWave->NumChannels = QualityInfo.NumChannels;
+			if (QualityInfo.SampleRate != 0)
+			{
+				InWave->SetSampleRate(QualityInfo.SampleRate, /* bFromDecoders */ true);
+			}
+			if (QualityInfo.NumChannels != 0)
+			{
+				ensure(QualityInfo.NumChannels <= AUDIO_MIXER_MAX_OUTPUT_CHANNELS);
+				InWave->NumChannels = QualityInfo.NumChannels;
+			}
 			if (QualityInfo.SampleDataSize != 0)
 			{
+				if (QualityInfo.NumChannels > 0)
+				{
+					// Update the NumFrames *if* the decoders have returned how big the sample data is.
+					// Some decoder implementations don't do this because of how they organise the bit-stream for streaming.
+					const int64 NumFrames = QualityInfo.SampleDataSize / sizeof(int16) / QualityInfo.NumChannels;
+					InWave->SetNumFrames(IntCastChecked<uint32>(NumFrames));
+				}
 				InWave->RawPCMDataSize = QualityInfo.SampleDataSize;
 			}
 			if (QualityInfo.Duration != 0.0f)
 			{
+				ensure(QualityInfo.Duration > 0.0f);
 				InWave->Duration = QualityInfo.Duration;
 			}
 		}
diff --git a/Engine/Source/Runtime/VorbisAudioDecoder/Module/Private/VorbisAudioInfo.cpp b/Engine/Source/Runtime/VorbisAudioDecoder/Module/Private/VorbisAudioInfo.cpp
--- a/Engine/Source/Runtime/VorbisAudioDecoder/Module/Private/VorbisAudioInfo.cpp
+++ b/Engine/Source/Runtime/VorbisAudioDecoder/Module/Private/VorbisAudioInfo.cpp
@@ -674,7 +674,7 @@ const uint8* FVorbisAudioInfo::GetLoadedChunk(FSoundWaveProxyPtr InSoundWave, ui
 		}
 		else
 		{
-			CurCompressedChunkHandle = IStreamingManager::Get().GetAudioStreamingManager().GetLoadedChunk(InSoundWave, ChunkIndex, false, true);
+			CurCompressedChunkHandle = IStreamingManager::Get().GetAudioStreamingManager().GetLoadedChunk(InSoundWave, ChunkIndex, /* bBlockForLoad */ true, true);
 			OutChunkSize = CurCompressedChunkHandle.Num();
 			return CurCompressedChunkHandle.GetData();
 		}
