Make sure MetaSounds get connected to buses in new rendering path

#rnx

diff --git a/Engine/Source/Runtime/AudioMixer/Private/AudioMixerSourceBuffer.cpp b/Engine/Source/Runtime/AudioMixer/Private/AudioMixerSourceBuffer.cpp
--- a/Engine/Source/Runtime/AudioMixer/Private/AudioMixerSourceBuffer.cpp
+++ b/Engine/Source/Runtime/AudioMixer/Private/AudioMixerSourceBuffer.cpp
@@ -349,6 +349,11 @@ namespace Audio
 
 		// Set the render cost encountered during the last render
 		SetRelativeRenderCost(TaskResult.RelativeRenderCost);
+
+		if (bDirectRendering)
+		{
+			ConnectToBuses();
+		}
 	}
 
 	void FMixerSourceBuffer::SubmitInitialPCMBuffers()
@@ -591,15 +596,7 @@ namespace Audio
 			SubmitRealTimeSourceData(bIsFinishedOrLooped);
 		}
 
-		if (FAudioDeviceManager* ADM = FAudioDeviceManager::Get())
-		{
-			if (FAudioDevice* AudioDevice = ADM->GetAudioDeviceRaw(AuioDeviceID))
-			{
-				UAudioBusSubsystem* AudioBusSubsystem = AudioDevice->GetSubsystem<UAudioBusSubsystem>();
-				check(AudioBusSubsystem);
-				AudioBusSubsystem->ConnectPatches(InstanceID);
-			}
-		}
+		ConnectToBuses();
 
 		// Start a new decoding task
 		if (!AsyncRealtimeAudioTask && !bDirectRendering)
@@ -637,6 +634,19 @@ namespace Audio
 		}
 	}
 
+	void FMixerSourceBuffer::ConnectToBuses()
+	{
+		if (FAudioDeviceManager* ADM = FAudioDeviceManager::Get())
+		{
+			if (FAudioDevice* AudioDevice = ADM->GetAudioDeviceRaw(AuioDeviceID))
+			{
+				UAudioBusSubsystem* AudioBusSubsystem = AudioDevice->GetSubsystem<UAudioBusSubsystem>();
+				check(AudioBusSubsystem);
+				AudioBusSubsystem->ConnectPatches(InstanceID);
+			}
+		}
+	}
+
 	void FMixerSourceBuffer::SubmitBuffer(TSharedPtr<FAlignedFloatBuffer, ESPMode::ThreadSafe> InSourceVoiceBuffer)
 	{
 		NumBuffersQeueued++;
diff --git a/Engine/Source/Runtime/AudioMixer/Private/AudioMixerSourceBuffer.h b/Engine/Source/Runtime/AudioMixer/Private/AudioMixerSourceBuffer.h
--- a/Engine/Source/Runtime/AudioMixer/Private/AudioMixerSourceBuffer.h
+++ b/Engine/Source/Runtime/AudioMixer/Private/AudioMixerSourceBuffer.h
@@ -148,6 +148,7 @@ namespace Audio
 		void ProcessRealTimeSource();
 		void SubmitBuffer(TSharedPtr<FAlignedFloatBuffer, ESPMode::ThreadSafe> InSourceVoiceBuffer);
 		void DeleteDecoder();
+		void ConnectToBuses();
 
 		int32 NumBuffersQeueued;
 		FRawPCMDataBuffer RawPCMDataBuffer;
