[Source Rendering Restructuring] enable new audio source rendering path by default

#rnx

diff --git a/Engine/Source/Runtime/AudioMixer/Private/AudioMixerSourceBuffer.cpp b/Engine/Source/Runtime/AudioMixer/Private/AudioMixerSourceBuffer.cpp
--- a/Engine/Source/Runtime/AudioMixer/Private/AudioMixerSourceBuffer.cpp
+++ b/Engine/Source/Runtime/AudioMixer/Private/AudioMixerSourceBuffer.cpp
@@ -9,7 +9,7 @@
 
 namespace Audio
 {
-	int32 DirectProceduralRenderingCVar = 0;
+	int32 DirectProceduralRenderingCVar = 1;
 	FAutoConsoleVariableRef CVarDirectProceduralRendering(
 		TEXT("au.DirectProceduralRendering"),
 		DirectProceduralRenderingCVar,
diff --git a/Engine/Source/Runtime/AudioMixer/Private/AudioMixerSourceManager.cpp b/Engine/Source/Runtime/AudioMixer/Private/AudioMixerSourceManager.cpp
--- a/Engine/Source/Runtime/AudioMixer/Private/AudioMixerSourceManager.cpp
+++ b/Engine/Source/Runtime/AudioMixer/Private/AudioMixerSourceManager.cpp
@@ -31,7 +31,7 @@
 
 // Link to "Audio" profiling category
 CSV_DECLARE_CATEGORY_MODULE_EXTERN(AUDIOMIXERCORE_API, Audio);
-static int32 DisableParallelSourceProcessingCvar = 1;
+static int32 DisableParallelSourceProcessingCvar = 0;
 FAutoConsoleVariableRef CVarDisableParallelSourceProcessing(
 	TEXT("au.DisableParallelSourceProcessing"),
 	DisableParallelSourceProcessingCvar,
@@ -180,7 +180,7 @@ static FAutoConsoleCommand GModulationSetMaxPitchRange(
 	)
 );
 
-static int32 PerSourceResampling = 0;
+static int32 PerSourceResampling = 1;
 FAutoConsoleVariableRef CVarAudioMixerPerSourceResampling(
 	TEXT("au.PerSourceResampling"),
 	PerSourceResampling,
