Submix Auto-Disable fixes

diff --git a/Engine/Source/Runtime/AudioMixer/Private/AudioMixerSubmix.cpp b/Engine/Source/Runtime/AudioMixer/Private/AudioMixerSubmix.cpp
--- a/Engine/Source/Runtime/AudioMixer/Private/AudioMixerSubmix.cpp
+++ b/Engine/Source/Runtime/AudioMixer/Private/AudioMixerSubmix.cpp
@@ -1124,18 +1124,6 @@ namespace Audio
 			}
 		}
 
-		// Query Modulation; if any of the submix's Modulation Destinations are being modulated they need to be processed,
-		// because binaural sources still use these Destinations when the submix is set
-		{
-			if (MixerDevice->IsModulationPluginEnabled() && MixerDevice->ModulationInterface.IsValid())
-			{
-				if (DryLevelMod.IsActive() || VolumeMod.IsActive())
-				{
-					return true;
-				}
-			}
-		}
-
 		// If this submix is not rendering any sources directly and silence has been detected, we need to check it's children submixes
 		if (MixerSourceVoices.Num() == 0 && SilenceTimeStartSeconds >= 0.0)
 		{
@@ -1257,7 +1245,7 @@ namespace Audio
 			{
 				bIsCurrentlyDisabled = true;
 
-				UE_CLOG(LogSubmixEnablementCVar == 1, LogAudioMixer, Display,
+				UE_CLOG(LogSubmixEnablementCVar == 1, LogAudioMixer, Verbose,
 					TEXT("Submix Disabled. Num Sources: %d, Time Silent: %.2f, Disablement Threshold: %.2f, Submix Name: %s"), 
 					MixerSourceVoices.Num(),
 					(float)(MixerDevice->GetAudioClock() - SilenceTimeStartSeconds),
@@ -1274,7 +1262,7 @@ namespace Audio
 		if (bIsCurrentlyDisabled)
 		{
 			bIsCurrentlyDisabled = false;
-			UE_CLOG(LogSubmixEnablementCVar == 1, LogAudioMixer, Display, TEXT("Submix Re-Enabled: %s"), *SubmixName);
+			UE_CLOG(LogSubmixEnablementCVar == 1, LogAudioMixer, Verbose, TEXT("Submix Re-Enabled: %s"), *SubmixName);
 		}
 
 		// Mix all submix audio into this submix's input scratch buffer
@@ -1292,13 +1280,10 @@ namespace Audio
 				// forcibly deleted in editor, so submix validity (in addition to pointer validity) is checked before processing
 				if (ChildSubmix.IsValid() && ChildSubmix->IsValid())
 				{
-					if (ChildSubmix->IsRenderingAudio())
-					{
-						ChildSubmix->ProcessAudio(InputBuffer);
+					ChildSubmix->ProcessAudio(InputBuffer);
 
-						// Check the buffer after processing to catch any bad values.
-						AUDIO_CHECK_BUFFER_NAMED_MSG(InputBuffer, TEXT("Submix Chidren"), TEXT("Submix: %s"), *ChildSubmix->GetName());
-					}
+					// Check the buffer after processing to catch any bad values.
+					AUDIO_CHECK_BUFFER_NAMED_MSG(InputBuffer, TEXT("Submix Chidren"), TEXT("Submix: %s"), *ChildSubmix->GetName());
 				}
 				else
 				{
