Add & Remove Modulators when submix is enabled/disabled

diff --git a/Engine/Source/Runtime/AudioMixer/Private/AudioMixerSubmix.cpp b/Engine/Source/Runtime/AudioMixer/Private/AudioMixerSubmix.cpp
--- a/Engine/Source/Runtime/AudioMixer/Private/AudioMixerSubmix.cpp
+++ b/Engine/Source/Runtime/AudioMixer/Private/AudioMixerSubmix.cpp
@@ -1245,6 +1245,10 @@ namespace Audio
 			{
 				bIsCurrentlyDisabled = true;
 
+				VolumeMod.ClearModulators();
+				WetLevelMod.ClearModulators();
+				DryLevelMod.ClearModulators();
+				
 				UE_CLOG(LogSubmixEnablementCVar == 1, LogAudioMixer, Verbose,
 					TEXT("Submix Disabled. Num Sources: %d, Time Silent: %.2f, Disablement Threshold: %.2f, Submix Name: %s"), 
 					MixerSourceVoices.Num(),
@@ -1261,6 +1265,14 @@ namespace Audio
 
 		if (bIsCurrentlyDisabled)
 		{
+			if (const USoundSubmix* SoundSubmix = Cast<const USoundSubmix>(OwningSubmixObject))
+			{
+				const TSet<TObjectPtr<USoundModulatorBase>>& VolumeModulator = SoundSubmix->OutputVolumeModulation.Modulators;
+				const TSet<TObjectPtr<USoundModulatorBase>>& WetLevelModulator = SoundSubmix->WetLevelModulation.Modulators;
+				const TSet<TObjectPtr<USoundModulatorBase>>& DryLevelModulator = SoundSubmix->DryLevelModulation.Modulators;
+				UpdateModulationSettings(VolumeModulator, WetLevelModulator, DryLevelModulator);
+			}
+			
 			bIsCurrentlyDisabled = false;
 			UE_CLOG(LogSubmixEnablementCVar == 1, LogAudioMixer, Verbose, TEXT("Submix Re-Enabled: %s"), *SubmixName);
 		}
diff --git a/Engine/Source/Runtime/Engine/Classes/Sound/SoundModulationDestination.h b/Engine/Source/Runtime/Engine/Classes/Sound/SoundModulationDestination.h
--- a/Engine/Source/Runtime/Engine/Classes/Sound/SoundModulationDestination.h
+++ b/Engine/Source/Runtime/Engine/Classes/Sound/SoundModulationDestination.h
@@ -198,6 +198,8 @@ namespace Audio
 			ENGINE_API void UpdateModulators(const TSet<TObjectPtr<USoundModulatorBase>>& InModulators);
 			ENGINE_API void UpdateModulators(const TSet<USoundModulatorBase*>& InModulators);
 			ENGINE_API void UpdateModulators(const TSet<const USoundModulatorBase*>& InModulators);
+		
+			ENGINE_API void ClearModulators();
 
 	private:
 			ENGINE_API void UpdateModulatorsInternal(TArray<TUniquePtr<Audio::IModulatorSettings>>&& ProxySettings);
diff --git a/Engine/Source/Runtime/Engine/Private/SoundModulationDestination.cpp b/Engine/Source/Runtime/Engine/Private/SoundModulationDestination.cpp
--- a/Engine/Source/Runtime/Engine/Private/SoundModulationDestination.cpp
+++ b/Engine/Source/Runtime/Engine/Private/SoundModulationDestination.cpp
@@ -296,6 +296,11 @@ namespace Audio
 		UpdateModulatorsInternal(MoveTemp(ProxySettings));
 	}
 
+	void FModulationDestination::ClearModulators()
+	{
+		DestinationData->ResetHandles();
+	}
+
 	void FModulationDestination::FModulationDestinationData::SetHandles(TSet<FModulatorHandle>&& NewHandles)
 	{
 		FScopeLock Lock(&HandleCritSection);
