Handling the case of being unable to create a sound wasve proxy reader in scrubber

[RN] Fixed crash with soundwave scrubber in sequencer when a sound wave proxy reader fails to be instantiated due to bad data.

diff --git a/Engine/Source/Runtime/AudioMixer/Classes/Generators/SoundWaveScrubber.h b/Engine/Source/Runtime/AudioMixer/Classes/Generators/SoundWaveScrubber.h
--- a/Engine/Source/Runtime/AudioMixer/Classes/Generators/SoundWaveScrubber.h
+++ b/Engine/Source/Runtime/AudioMixer/Classes/Generators/SoundWaveScrubber.h
@@ -164,6 +164,9 @@ private:
 
 	// Whether or not we scrub while the playhead doesn't move
 	bool bIsScrubbingWhileStationary = true;
+
+	// Whether or not we have an error with the decoder
+	bool bHasErrorWithDecoder = false;
 };
 
 
diff --git a/Engine/Source/Runtime/AudioMixer/Private/Generators/SoundWaveScrubber.cpp b/Engine/Source/Runtime/AudioMixer/Private/Generators/SoundWaveScrubber.cpp
--- a/Engine/Source/Runtime/AudioMixer/Private/Generators/SoundWaveScrubber.cpp
+++ b/Engine/Source/Runtime/AudioMixer/Private/Generators/SoundWaveScrubber.cpp
@@ -44,6 +44,9 @@ void FSoundWaveScrubber::Init(FSoundWaveProxyPtr InSoundWaveProxyPtr, float InSa
 
 	float DecoderSeekTimeSeconds = FMath::Max(CurrentPlayheadTimeSeconds.GetValue() - 0.5f * DecodedAudioSizeInSeconds, 0.0f);
 
+	bHasErrorWithDecoder = false;
+
+	// This could result in a decoder error if there is an issue with the sound wave proxy
 	DecodeToDataChunk(DecodedChunks[0], DecoderSeekTimeSeconds);
 }
 
@@ -67,6 +70,8 @@ int32 FSoundWaveScrubber::GetDecodedDataChunkIndexForCurrentReadIndex(int32 InRe
 
 int32 FSoundWaveScrubber::DecodeDataChunkIndexForCurrentReadIndex(int32 InReadFrameIndex)
 {
+	check(!bHasErrorWithDecoder);
+	
 	// No decoded chunk was found for desired read frame index. 
 	// This indicates that we need to decode more audio.
 	for (int32 i = 0; i < DecodedChunks.Num(); ++i)
@@ -86,6 +91,10 @@ int32 FSoundWaveScrubber::DecodeDataChunkIndexForCurrentReadIndex(int32 InReadFr
 	DecodeToDataChunk(NewChunk, DecoderSeekTimeSeconds);
 	check(NewChunk.PCMAudio.Num() > 0);
 	DecodedChunks.Add(MoveTemp(NewChunk));
+	
+	// This should not be able to error at this point
+	check(!bHasErrorWithDecoder);
+
 	return DecodedChunks.Num() - 1;
 }
 
@@ -101,11 +110,21 @@ void FSoundWaveScrubber::DecodeToDataChunk(FDecodedDataChunk& InOutDataChunk, fl
 
 	if (!SoundWaveProxyReaderPtr.IsValid())
 	{
+		// Make sure we haven't already tried to create a proxy reader (decoder)
+		check(!bHasErrorWithDecoder);
+
 		// Create the proxy reader, which is our decoder. Initialize it at the initial playhead time
 		FSoundWaveProxyReader::FSettings ProxyReaderSettings;
 		ProxyReaderSettings.MaxDecodeSizeInFrames = DecodedAudioSizeInSeconds * SourceFileSampleRate;
 		ProxyReaderSettings.StartTimeInSeconds = InDecoderSeekTimeSeconds;
 		SoundWaveProxyReaderPtr = FSoundWaveProxyReader::Create(SoundWaveProxyPtr.ToSharedRef(), ProxyReaderSettings);
+
+		if (!SoundWaveProxyReaderPtr.IsValid())
+		{
+			bHasErrorWithDecoder = true;
+			UE_LOG(LogAudioMixer, Warning, TEXT("Unable to make a sound wave proxy reader for sound wave '%s' in the sound wave scrubber."), *SoundWaveProxyPtr->GetFName().ToString());
+			return;
+		}
 	}
 	else
 	{
@@ -140,6 +159,9 @@ void FSoundWaveScrubber::DecodeToDataChunk(FDecodedDataChunk& InOutDataChunk, fl
 
 FSoundWaveScrubber::FGrain FSoundWaveScrubber::SpawnGrain()
 {
+	// Make sure we don't have a decoder error when we're trying to spawn grains
+	check(!bHasErrorWithDecoder);
+	
 	// Try to retrieve a decoded data chunk for the current read frame based on the curret, interpolated playhead time seconds value
 	int32 CurrentReadFrame = CurrentPlayheadTimeSeconds.GetValue() * AudioMixerSampleRate;
 	int32 DecodedDataChunkIndex = GetDecodedDataChunkIndexForCurrentReadIndex(CurrentReadFrame);
@@ -196,6 +218,13 @@ int32 FSoundWaveScrubber::RenderAudio(TArrayView<float>& OutAudio)
 {
 	SCOPED_NAMED_EVENT_TEXT("FSoundWaveScrubber::RenderAudio", FColor::Emerald);
 
+	// If we have an error with our decoder, we don't need to render audio
+	// To avoid spamming, we'll allow it to render silence
+	if (bHasErrorWithDecoder)
+	{
+		return OutAudio.Num();
+	}
+
 	// Number of frames of this generate audio
 	int32 NumFrames = OutAudio.Num() / NumChannels;
 	float DeltaTimeSecond = (float)NumFrames / AudioMixerSampleRate;
