[AudioMixer] Added Submix trace messages for use in Audio Insights standalone Submix dashboard

diff --git a/Engine/Source/Runtime/AudioMixer/Private/AudioMixerDevice.cpp b/Engine/Source/Runtime/AudioMixer/Private/AudioMixerDevice.cpp
--- a/Engine/Source/Runtime/AudioMixer/Private/AudioMixerDevice.cpp
+++ b/Engine/Source/Runtime/AudioMixer/Private/AudioMixerDevice.cpp
@@ -125,6 +125,18 @@ FAutoConsoleVariableRef CVarMuteAudio(
 	TEXT("1: Mutes the final submix. 0: Does not mute the final submix."),
 	ECVF_Cheat);
 
+#if UE_AUDIO_PROFILERTRACE_ENABLED
+
+UE_TRACE_EVENT_BEGIN(Audio, SubmixLoadedOnTraceStart)
+	UE_TRACE_EVENT_FIELD(uint32, DeviceId)
+	UE_TRACE_EVENT_FIELD(uint32, SubmixId)
+	UE_TRACE_EVENT_FIELD(double, Timestamp)
+	UE_TRACE_EVENT_FIELD(UE::Trace::WideString, Name)
+	UE_TRACE_EVENT_FIELD(int32, NumChannels)
+UE_TRACE_EVENT_END()
+
+#endif // UE_AUDIO_PROFILERTRACE_ENABLED
+
 // Link to "Audio" profiling category
 CSV_DECLARE_CATEGORY_MODULE_EXTERN(AUDIOMIXERCORE_API, Audio);
 
@@ -341,12 +353,20 @@ namespace Audio
 			check(Factories.Num()==1);
 			AudioLinkFactory=IAudioLinkFactory::FindFactory(Factories[0]);
 		}
+
+#if UE_AUDIO_PROFILERTRACE_ENABLED
+		FTraceAuxiliary::OnTraceStarted.AddRaw(this, &FMixerDevice::OnTraceStarted);
+#endif // UE_AUDIO_PROFILERTRACE_ENABLED
 	}
 
 	FMixerDevice::~FMixerDevice()
 	{
 		AUDIO_MIXER_CHECK_GAME_THREAD(this);
 
+#if UE_AUDIO_PROFILERTRACE_ENABLED
+		FTraceAuxiliary::OnTraceStarted.RemoveAll(this);
+#endif // UE_AUDIO_PROFILERTRACE_ENABLED
+
 		// Shutdown all pending clock events, as they may have references to 
 		// the FMixerSourceManager that is about to be destroyed
 		QuantizedEventClockManager.Shutdown();
@@ -357,6 +377,42 @@ namespace Audio
 		}
 	}
 
+#if UE_AUDIO_PROFILERTRACE_ENABLED
+	void FMixerDevice::OnTraceStarted(FTraceAuxiliary::EConnectionType TraceType, const FString& TraceDestination)
+	{
+		auto TraceSubmixLoadedMessage = [this](const FMixerSubmixPtr& InSubmix)
+		{
+			UE_TRACE_LOG(Audio, SubmixLoadedOnTraceStart, AudioChannel)
+				<< SubmixLoadedOnTraceStart.DeviceId(DeviceID)
+				<< SubmixLoadedOnTraceStart.SubmixId(GetTypeHash(InSubmix->GetPathName()))
+				<< SubmixLoadedOnTraceStart.Timestamp(FPlatformTime::Cycles64())
+				<< SubmixLoadedOnTraceStart.Name(*InSubmix->GetPathName())
+				<< SubmixLoadedOnTraceStart.NumChannels(InSubmix->NumChannels);
+		};
+
+		const bool bChannelEnabled = UE_TRACE_CHANNELEXPR_IS_ENABLED(AudioChannel);
+		if (bChannelEnabled)
+		{
+			Submixes.Iterate([this, &TraceSubmixLoadedMessage](const FSubmixMap::FPair& SubmixPair)
+			{
+				const FMixerSubmixPtr& Submix = SubmixPair.Value;
+				if (Submix.IsValid())
+				{
+					TraceSubmixLoadedMessage(Submix);
+				}
+			});
+
+			for (const FMixerSubmixPtr& RequiredSubmix : RequiredSubmixInstances)
+			{
+				if (RequiredSubmix.IsValid())
+				{
+					TraceSubmixLoadedMessage(RequiredSubmix);
+				}
+			}
+		}
+	}
+#endif // UE_AUDIO_PROFILERTRACE_ENABLED
+
 	void FMixerDevice::AddReferencedObjects(FReferenceCollector& Collector)
 	{
 	}
diff --git a/Engine/Source/Runtime/AudioMixer/Private/AudioMixerSubmix.cpp b/Engine/Source/Runtime/AudioMixer/Private/AudioMixerSubmix.cpp
--- a/Engine/Source/Runtime/AudioMixer/Private/AudioMixerSubmix.cpp
+++ b/Engine/Source/Runtime/AudioMixer/Private/AudioMixerSubmix.cpp
@@ -59,6 +59,7 @@ DEFINE_STAT(STAT_AudioMixerSubmixSoundfieldProcessors);
 
 #if UE_AUDIO_PROFILERTRACE_ENABLED
 
+// @TODO UE-297630: remove this trace message once the full dashboard transition is finished (only editor -> editor and standalone)
 UE_TRACE_EVENT_BEGIN(Audio, SoundSubmixHasActivity)
 	UE_TRACE_EVENT_FIELD(uint32, DeviceId)
 	UE_TRACE_EVENT_FIELD(uint32, SubmixId)
@@ -66,6 +67,34 @@ UE_TRACE_EVENT_BEGIN(Audio, SoundSubmixHasActivity)
 	UE_TRACE_EVENT_FIELD(bool, HasActivity)
 UE_TRACE_EVENT_END()
 
+UE_TRACE_EVENT_BEGIN(Audio, SubmixLoaded)
+	UE_TRACE_EVENT_FIELD(uint32, DeviceId)
+	UE_TRACE_EVENT_FIELD(uint32, SubmixId)
+	UE_TRACE_EVENT_FIELD(double, Timestamp)
+	UE_TRACE_EVENT_FIELD(UE::Trace::WideString, Name)
+	UE_TRACE_EVENT_FIELD(int32, NumChannels)
+UE_TRACE_EVENT_END()
+
+UE_TRACE_EVENT_BEGIN(Audio, SubmixHasActivity)
+	UE_TRACE_EVENT_FIELD(uint32, DeviceId)
+	UE_TRACE_EVENT_FIELD(uint32, SubmixId)
+	UE_TRACE_EVENT_FIELD(double, Timestamp)
+	UE_TRACE_EVENT_FIELD(bool, HasActivity)
+UE_TRACE_EVENT_END()
+
+UE_TRACE_EVENT_BEGIN(Audio, SubmixEnvelopeValues)
+	UE_TRACE_EVENT_FIELD(uint32, DeviceId)
+	UE_TRACE_EVENT_FIELD(uint32, SubmixId)
+	UE_TRACE_EVENT_FIELD(double, Timestamp)
+	UE_TRACE_EVENT_FIELD(float[], EnvelopeValues)
+UE_TRACE_EVENT_END()
+
+UE_TRACE_EVENT_BEGIN(Audio, SubmixUnloaded)
+	UE_TRACE_EVENT_FIELD(uint32, DeviceId)
+	UE_TRACE_EVENT_FIELD(uint32, SubmixId)
+	UE_TRACE_EVENT_FIELD(double, Timestamp)
+UE_TRACE_EVENT_END()
+
 #endif // UE_AUDIO_PROFILERTRACE_ENABLED
 
 namespace Audio
@@ -206,6 +235,17 @@ namespace Audio
 				MutableSubmix->StopRecordingOutput(MixerDevice, EAudioRecordingExportType::WavFile, InterruptedFileName, FString());
 			}
 		}
+
+#if UE_AUDIO_PROFILERTRACE_ENABLED
+		const bool bChannelEnabled = UE_TRACE_CHANNELEXPR_IS_ENABLED(AudioChannel);
+		if (bChannelEnabled)
+		{
+			UE_TRACE_LOG(Audio, SubmixUnloaded, AudioChannel)
+				<< SubmixUnloaded.DeviceId(MixerDevice->DeviceID)
+				<< SubmixUnloaded.SubmixId(GetTypeHash(OwningSubmixObject->GetPathName()))
+				<< SubmixUnloaded.Timestamp(FPlatformTime::Cycles64());
+		}
+#endif // UE_AUDIO_PROFILERTRACE_ENABLED
 	}
 
 	void FMixerSubmix::Init(const USoundSubmixBase* InSoundSubmix, bool bAllowReInit)
@@ -215,6 +255,10 @@ namespace Audio
 		{
 			SubmixName = InSoundSubmix->GetName();
 
+#if UE_AUDIO_PROFILERTRACE_ENABLED
+			SubmixPathName = InSoundSubmix->GetPathName();
+#endif // UE_AUDIO_PROFILERTRACE_ENABLED
+
 			// This is a first init and needs to be synchronous
 			if (!OwningSubmixObject.IsValid())
 			{
@@ -239,6 +283,19 @@ namespace Audio
 				SetupSoundfieldStreams(EncodingSettings, Effects, SoundfieldFactory);
 			}
 		}
+
+#if UE_AUDIO_PROFILERTRACE_ENABLED
+		const bool bChannelEnabled = UE_TRACE_CHANNELEXPR_IS_ENABLED(AudioChannel);
+		if (bChannelEnabled)
+		{
+			UE_TRACE_LOG(Audio, SubmixLoaded, AudioChannel)
+				<< SubmixLoaded.DeviceId(MixerDevice->DeviceID)
+				<< SubmixLoaded.SubmixId(GetTypeHash(OwningSubmixObject->GetPathName()))
+				<< SubmixLoaded.Timestamp(FPlatformTime::Cycles64())
+				<< SubmixLoaded.Name(*OwningSubmixObject->GetPathName())
+				<< SubmixLoaded.NumChannels(NumChannels);
+		}
+#endif // UE_AUDIO_PROFILERTRACE_ENABLED
 	}
 
 	void FMixerSubmix::InitInternal()
@@ -1520,6 +1577,18 @@ namespace Audio
 
 				FMemory::Memcpy(EnvelopeValues, EnvValues.GetData(), sizeof(float) * NumChannels);
 				Audio::ArrayClampInPlace(MakeArrayView(EnvelopeValues, NumChannels), 0.f, 1.f);
+
+#if UE_AUDIO_PROFILERTRACE_ENABLED
+				const bool bChannelEnabled = UE_TRACE_CHANNELEXPR_IS_ENABLED(AudioChannel);
+				if (bChannelEnabled)
+				{
+					UE_TRACE_LOG(Audio, SubmixEnvelopeValues, AudioChannel)
+						<< SubmixEnvelopeValues.DeviceId(MixerDevice->DeviceID)
+						<< SubmixEnvelopeValues.SubmixId(GetTypeHash(OwningSubmixObject->GetPathName()))
+						<< SubmixEnvelopeValues.Timestamp(FPlatformTime::Cycles64())
+						<< SubmixEnvelopeValues.EnvelopeValues(EnvelopeValues, NumChannels);
+				}
+#endif // UE_AUDIO_PROFILERTRACE_ENABLED
 			}
 
 			EnvelopeNumChannels = NumChannels;
@@ -1609,11 +1678,18 @@ namespace Audio
 #if UE_AUDIO_PROFILERTRACE_ENABLED
 			if (OwningSubmixObject.IsValid())
 			{
+				// @TODO UE-297630: remove this trace message once the full dashboard transition is finished (only editor -> editor and standalone)
 				UE_TRACE_LOG(Audio, SoundSubmixHasActivity, AudioChannel)
 					<< SoundSubmixHasActivity.DeviceId(MixerDevice->DeviceID)
 					<< SoundSubmixHasActivity.SubmixId(GetTypeHash(OwningSubmixObject->GetPathName()))
 					<< SoundSubmixHasActivity.Timestamp(FPlatformTime::Cycles64())
 					<< SoundSubmixHasActivity.HasActivity(!bIsNowSilent);
+
+				UE_TRACE_LOG(Audio, SubmixHasActivity, AudioChannel)
+					<< SubmixHasActivity.DeviceId(MixerDevice->DeviceID)
+					<< SubmixHasActivity.SubmixId(GetTypeHash(OwningSubmixObject->GetPathName()))
+					<< SubmixHasActivity.Timestamp(FPlatformTime::Cycles64())
+					<< SubmixHasActivity.HasActivity(!bIsNowSilent);
 			}
 #endif
 		}
diff --git a/Engine/Source/Runtime/AudioMixer/Public/AudioMixerDevice.h b/Engine/Source/Runtime/AudioMixer/Public/AudioMixerDevice.h
--- a/Engine/Source/Runtime/AudioMixer/Public/AudioMixerDevice.h
+++ b/Engine/Source/Runtime/AudioMixer/Public/AudioMixerDevice.h
@@ -400,7 +400,11 @@ namespace Audio
 
 		/** Updates the audio clock and the associated timing data */
 		void UpdateAudioClock();
-		
+
+#if UE_AUDIO_PROFILERTRACE_ENABLED
+		void OnTraceStarted(FTraceAuxiliary::EConnectionType TraceType, const FString& TraceDestination);
+#endif // UE_AUDIO_PROFILERTRACE_ENABLED
+
 		TArray<USoundSubmix*> RequiredSubmixes;
 		TArray<FMixerSubmixPtr> RequiredSubmixInstances;
 
diff --git a/Engine/Source/Runtime/AudioMixer/Public/AudioMixerSubmix.h b/Engine/Source/Runtime/AudioMixer/Public/AudioMixerSubmix.h
--- a/Engine/Source/Runtime/AudioMixer/Public/AudioMixerSubmix.h
+++ b/Engine/Source/Runtime/AudioMixer/Public/AudioMixerSubmix.h
@@ -4,6 +4,7 @@
 
 #include "AudioBusSubsystem.h"
 #include "AudioMixer.h"
+#include "AudioMixerTrace.h"
 #include "AudioDefines.h"
 #include "CoreMinimal.h"
 #include "SampleBuffer.h"
@@ -120,6 +121,11 @@ namespace Audio
 		// Return the owners name 
 		const FString& GetName() const { return SubmixName; }
 
+#if UE_AUDIO_PROFILERTRACE_ENABLED
+		// Return the owners path name 
+		const FString& GetPathName() const { return SubmixPathName; }
+#endif // UE_AUDIO_PROFILERTRACE_ENABLED
+
 		// Sets the parent submix to the given submix
 		AUDIOMIXER_API void SetParentSubmix(TWeakPtr<FMixerSubmix, ESPMode::ThreadSafe> Submix);
 
@@ -387,6 +393,11 @@ namespace Audio
 
 		// The name of this submix (the owning USoundSubmix) (at top so we can see in debugger it's name)
 		FString SubmixName;
+
+#if UE_AUDIO_PROFILERTRACE_ENABLED
+		// The path name of this submix (the owning USoundSubmix)
+		FString SubmixPathName;
+#endif // UE_AUDIO_PROFILERTRACE_ENABLED
 		
 		// This mixer submix's Id
 		uint32 Id;
