[AudioMixer] FMixerDevice: Implemented Start/Stop envelope following CVars so we can use them as mechanism in Audio Insights standalone to enable/disable audio metering visualization for each registered submix

diff --git a/Engine/Source/Runtime/AudioMixer/Private/AudioMixerDevice.cpp b/Engine/Source/Runtime/AudioMixer/Private/AudioMixerDevice.cpp
--- a/Engine/Source/Runtime/AudioMixer/Private/AudioMixerDevice.cpp
+++ b/Engine/Source/Runtime/AudioMixer/Private/AudioMixerDevice.cpp
@@ -411,6 +411,101 @@ namespace Audio
 			}
 		}
 	}
+
+	void FMixerDevice::RegisterSubmixEnvelopeFollowerCVars()
+	{
+		// StartSubmixEnvelopeFollowerCVar
+		StartSubmixEnvelopeFollowerCVar = IConsoleManager::Get().RegisterConsoleVariable(
+			*FString::Printf(TEXT("au.StartSubmixEnvelopeFollower_AD%d"), DeviceID),
+			FString(),
+			TEXT("Starts a submix envelope follower (args= submix name)\n"),
+			ECVF_Default
+		);
+
+		if (StartSubmixEnvelopeFollowerCVar)
+		{
+			StartSubmixEnvelopeFollowerCVar->SetOnChangedCallback(FConsoleVariableDelegate::CreateRaw(this, &FMixerDevice::OnStartSubmixEnvelopeFollowerCVarChanged));
+		}
+
+		// StopSubmixEnvelopeFollowerCVar
+		StopSubmixEnvelopeFollowerCVar = IConsoleManager::Get().RegisterConsoleVariable(
+			*FString::Printf(TEXT("au.StopSubmixEnvelopeFollower_AD%d"), DeviceID),
+			FString(),
+			TEXT("Stops a submix envelope follower (args= submix name)\n"),
+			ECVF_Default
+		);
+
+		if (StopSubmixEnvelopeFollowerCVar)
+		{
+			StopSubmixEnvelopeFollowerCVar->SetOnChangedCallback(FConsoleVariableDelegate::CreateRaw(this, &FMixerDevice::OnStopSubmixEnvelopeFollowerCVarChanged));
+		}
+	}
+
+	void FMixerDevice::UnregisterSubmixEnvelopeFollowerCVars()
+	{
+		IConsoleManager::Get().UnregisterConsoleObject(StartSubmixEnvelopeFollowerCVar);
+		IConsoleManager::Get().UnregisterConsoleObject(StopSubmixEnvelopeFollowerCVar);
+	}
+
+	void FMixerDevice::OnStartSubmixEnvelopeFollowerCVarChanged(IConsoleVariable* InStartSubmixEnvelopeFollowerCVar)
+	{
+		if (InStartSubmixEnvelopeFollowerCVar == nullptr)
+		{
+			return;
+		}
+
+		const FString& SubmixName = InStartSubmixEnvelopeFollowerCVar->GetString();
+
+		// Attack/Release default times used in USoundSubmix constructor
+		constexpr int32 AttackTime  = 10;
+		constexpr int32 ReleaseTime = 500;
+
+		Submixes.Iterate([this, &SubmixName, AttackTime, ReleaseTime](const FSubmixMap::FPair& SubmixPair)
+		{
+			const FMixerSubmixPtr& Submix = SubmixPair.Value;
+
+			if (Submix.IsValid() && Submix->GetName() == SubmixName)
+			{
+				Submix->StartEnvelopeFollowing(AttackTime, ReleaseTime);
+			}
+		});
+
+		for (const FMixerSubmixPtr& RequiredSubmix : RequiredSubmixInstances)
+		{
+			if (RequiredSubmix.IsValid() && RequiredSubmix->GetName() == SubmixName)
+			{
+				RequiredSubmix->StartEnvelopeFollowing(AttackTime, ReleaseTime);
+			}
+		}
+	}
+
+	void FMixerDevice::OnStopSubmixEnvelopeFollowerCVarChanged(IConsoleVariable* InStopSubmixEnvelopeFollowerCVar)
+	{
+		if (InStopSubmixEnvelopeFollowerCVar == nullptr)
+		{
+			return;
+		}
+
+		const FString& SubmixName = InStopSubmixEnvelopeFollowerCVar->GetString();
+
+		Submixes.Iterate([this, &SubmixName](const FSubmixMap::FPair& SubmixPair)
+		{
+			const FMixerSubmixPtr& Submix = SubmixPair.Value;
+
+			if (Submix.IsValid() && Submix->GetName() == SubmixName)
+			{
+				Submix->StopEnvelopeFollowing();
+			}
+		});
+
+		for (const FMixerSubmixPtr& RequiredSubmix : RequiredSubmixInstances)
+		{
+			if (RequiredSubmix.IsValid() && RequiredSubmix->GetName() == SubmixName)
+			{
+				RequiredSubmix->StopEnvelopeFollowing();
+			}
+		}
+	}
 #endif // UE_AUDIO_PROFILERTRACE_ENABLED
 
 	void FMixerDevice::AddReferencedObjects(FReferenceCollector& Collector)
@@ -710,6 +805,10 @@ namespace Audio
 				// Need to set these up before we start the audio stream.
 				InitSoundSubmixes();
 
+#if UE_AUDIO_PROFILERTRACE_ENABLED
+				RegisterSubmixEnvelopeFollowerCVars();
+#endif // UE_AUDIO_PROFILERTRACE_ENABLED
+
 				AudioMixerPlatform->PostInitializeHardware();
 
 				// Initialize the data used for audio thread sub-frame timing.
@@ -761,6 +860,10 @@ namespace Audio
 
 		if (IsInitialized())
 		{
+#if UE_AUDIO_PROFILERTRACE_ENABLED
+			UnregisterSubmixEnvelopeFollowerCVars();
+#endif // UE_AUDIO_PROFILERTRACE_ENABLED
+
 			for (TObjectIterator<USoundSubmix> It; It; ++It)
 			{
 				UnregisterSoundSubmix(*It, true);
diff --git a/Engine/Source/Runtime/AudioMixer/Public/AudioMixerDevice.h b/Engine/Source/Runtime/AudioMixer/Public/AudioMixerDevice.h
--- a/Engine/Source/Runtime/AudioMixer/Public/AudioMixerDevice.h
+++ b/Engine/Source/Runtime/AudioMixer/Public/AudioMixerDevice.h
@@ -403,6 +403,12 @@ namespace Audio
 
 #if UE_AUDIO_PROFILERTRACE_ENABLED
 		void OnTraceStarted(FTraceAuxiliary::EConnectionType TraceType, const FString& TraceDestination);
+
+		void RegisterSubmixEnvelopeFollowerCVars();
+		void UnregisterSubmixEnvelopeFollowerCVars();
+
+		void OnStartSubmixEnvelopeFollowerCVarChanged(IConsoleVariable* InStartSubmixEnvelopeFollowerCVar);
+		void OnStopSubmixEnvelopeFollowerCVarChanged(IConsoleVariable* InStopSubmixEnvelopeFollowerCVar);
 #endif // UE_AUDIO_PROFILERTRACE_ENABLED
 
 		TArray<USoundSubmix*> RequiredSubmixes;
@@ -482,7 +488,12 @@ namespace Audio
 		TMpscQueue<TFunction<void()>> GameThreadCommandQueue;
 
 		IAudioLinkFactory* AudioLinkFactory = nullptr;
-		
+
+#if UE_AUDIO_PROFILERTRACE_ENABLED
+		IConsoleVariable* StartSubmixEnvelopeFollowerCVar = nullptr;
+		IConsoleVariable* StopSubmixEnvelopeFollowerCVar  = nullptr;
+#endif // UE_AUDIO_PROFILERTRACE_ENABLED
+
 		/** Whether or not we generate output audio to test multi-platform mixer. */
 		bool bDebugOutputEnabled;
 
