[AudioMixer] FMixerDevice: Added thread safety to Start/Stop SubmixEnvelopeFollower CVars logic

diff --git a/Engine/Source/Runtime/AudioMixer/Private/AudioMixerDevice.cpp b/Engine/Source/Runtime/AudioMixer/Private/AudioMixerDevice.cpp
--- a/Engine/Source/Runtime/AudioMixer/Private/AudioMixerDevice.cpp
+++ b/Engine/Source/Runtime/AudioMixer/Private/AudioMixerDevice.cpp
@@ -449,22 +449,43 @@ namespace Audio
 
 	void FMixerDevice::OnStartSubmixEnvelopeFollowerCVarChanged(IConsoleVariable* InStartSubmixEnvelopeFollowerCVar)
 	{
-		if (InStartSubmixEnvelopeFollowerCVar == nullptr)
+		ensure(IsInGameThread());
+
+		if (InStartSubmixEnvelopeFollowerCVar)
 		{
-			return;
+			FAudioThread::RunCommandOnAudioThread([this, SubmixName = InStartSubmixEnvelopeFollowerCVar->GetString()]()
+			{
+				StartSubmixEnvelopeFollowerFromCVar(SubmixName);
+			});
+		}
+	}
+
+	void FMixerDevice::OnStopSubmixEnvelopeFollowerCVarChanged(IConsoleVariable* InStopSubmixEnvelopeFollowerCVar)
+	{
+		ensure(IsInGameThread());
+
+		if (InStopSubmixEnvelopeFollowerCVar)
+		{
+			FAudioThread::RunCommandOnAudioThread([this, SubmixName = InStopSubmixEnvelopeFollowerCVar->GetString()]()
+			{
+				StopSubmixEnvelopeFollowerFromCVar(SubmixName);
+			});
 		}
+	}
 
-		const FString& SubmixName = InStartSubmixEnvelopeFollowerCVar->GetString();
+	void FMixerDevice::StartSubmixEnvelopeFollowerFromCVar(const FString& InSubmixName)
+	{
+		ensure(IsInAudioThread());
 
 		// Attack/Release default times used in USoundSubmix constructor
 		constexpr int32 AttackTime  = 10;
 		constexpr int32 ReleaseTime = 500;
 
-		Submixes.Iterate([this, &SubmixName, AttackTime, ReleaseTime](const FSubmixMap::FPair& SubmixPair)
+		Submixes.Iterate([this, &InSubmixName, AttackTime, ReleaseTime](const FSubmixMap::FPair& SubmixPair)
 		{
 			const FMixerSubmixPtr& Submix = SubmixPair.Value;
 
-			if (Submix.IsValid() && Submix->GetName() == SubmixName)
+			if (Submix.IsValid() && Submix->GetName() == InSubmixName)
 			{
 				Submix->StartEnvelopeFollowing(AttackTime, ReleaseTime);
 			}
@@ -472,27 +493,22 @@ namespace Audio
 
 		for (const FMixerSubmixPtr& RequiredSubmix : RequiredSubmixInstances)
 		{
-			if (RequiredSubmix.IsValid() && RequiredSubmix->GetName() == SubmixName)
+			if (RequiredSubmix.IsValid() && RequiredSubmix->GetName() == InSubmixName)
 			{
 				RequiredSubmix->StartEnvelopeFollowing(AttackTime, ReleaseTime);
 			}
 		}
 	}
 
-	void FMixerDevice::OnStopSubmixEnvelopeFollowerCVarChanged(IConsoleVariable* InStopSubmixEnvelopeFollowerCVar)
+	void FMixerDevice::StopSubmixEnvelopeFollowerFromCVar(const FString& InSubmixName)
 	{
-		if (InStopSubmixEnvelopeFollowerCVar == nullptr)
-		{
-			return;
-		}
-
-		const FString& SubmixName = InStopSubmixEnvelopeFollowerCVar->GetString();
+		ensure(IsInAudioThread());
 
-		Submixes.Iterate([this, &SubmixName](const FSubmixMap::FPair& SubmixPair)
+		Submixes.Iterate([this, &InSubmixName](const FSubmixMap::FPair& SubmixPair)
 		{
 			const FMixerSubmixPtr& Submix = SubmixPair.Value;
 
-			if (Submix.IsValid() && Submix->GetName() == SubmixName)
+			if (Submix.IsValid() && Submix->GetName() == InSubmixName)
 			{
 				Submix->StopEnvelopeFollowing();
 			}
@@ -500,7 +516,7 @@ namespace Audio
 
 		for (const FMixerSubmixPtr& RequiredSubmix : RequiredSubmixInstances)
 		{
-			if (RequiredSubmix.IsValid() && RequiredSubmix->GetName() == SubmixName)
+			if (RequiredSubmix.IsValid() && RequiredSubmix->GetName() == InSubmixName)
 			{
 				RequiredSubmix->StopEnvelopeFollowing();
 			}
diff --git a/Engine/Source/Runtime/AudioMixer/Public/AudioMixerDevice.h b/Engine/Source/Runtime/AudioMixer/Public/AudioMixerDevice.h
--- a/Engine/Source/Runtime/AudioMixer/Public/AudioMixerDevice.h
+++ b/Engine/Source/Runtime/AudioMixer/Public/AudioMixerDevice.h
@@ -409,6 +409,9 @@ namespace Audio
 
 		void OnStartSubmixEnvelopeFollowerCVarChanged(IConsoleVariable* InStartSubmixEnvelopeFollowerCVar);
 		void OnStopSubmixEnvelopeFollowerCVarChanged(IConsoleVariable* InStopSubmixEnvelopeFollowerCVar);
+
+		void StartSubmixEnvelopeFollowerFromCVar(const FString& InSubmixName);
+		void StopSubmixEnvelopeFollowerFromCVar(const FString& InSubmixName);
 #endif // UE_AUDIO_PROFILERTRACE_ENABLED
 
 		TArray<USoundSubmix*> RequiredSubmixes;
