Add audio rendering scheduler, use it to render normal (non-bus) sources

#rnx
#tests Added automated tests in new AudioCQTests plugin

diff --git a/Engine/Plugins/Tests/AudioCQTests/AudioCQTests.uplugin b/Engine/Plugins/Tests/AudioCQTests/AudioCQTests.uplugin
new file mode 100644
--- /dev/null
+++ b/Engine/Plugins/Tests/AudioCQTests/AudioCQTests.uplugin
@@ -0,0 +1,23 @@
+{
+  "FileVersion": 3,
+  "Version": 1,
+  "VersionName": "1.0",
+  "FriendlyName": "Audio Code Quality Tests",
+  "Description": "Audio Code Quality Tests",
+  "Category": "Testing",
+  "CreatedBy": "Epic Games, Inc.",
+  "CreatedByURL": "https://epicgames.com",
+  "DocsURL": "",
+  "MarketplaceURL": "",
+  "SupportURL": "",
+  "CanContainContent": false,
+  "IsBetaVersion": false,
+  "Installed": false,
+  "Modules": [
+    {
+      "Name": "AudioCQTests",
+      "Type": "Runtime",
+      "LoadingPhase": "Default"
+    }
+  ]
+}
diff --git a/Engine/Plugins/Tests/AudioCQTests/Source/AudioCQTests/AudioCQTests.Build.cs b/Engine/Plugins/Tests/AudioCQTests/Source/AudioCQTests/AudioCQTests.Build.cs
new file mode 100644
--- /dev/null
+++ b/Engine/Plugins/Tests/AudioCQTests/Source/AudioCQTests/AudioCQTests.Build.cs
@@ -0,0 +1,21 @@
+// Copyright Epic Games, Inc. All Rights Reserved.
+
+using UnrealBuildTool;
+
+public class AudioCQTests : ModuleRules
+{
+	public AudioCQTests(ReadOnlyTargetRules Target)
+		: base(Target)
+	{
+		PCHUsage = PCHUsageMode.UseExplicitOrSharedPCHs;
+
+		PrivateDependencyModuleNames.AddRange(
+			new string[] {
+					"AudioMixer",
+					"Core",
+					"CQTest",
+					"Engine",
+				 }
+			);
+	}
+}
diff --git a/Engine/Plugins/Tests/AudioCQTests/Source/AudioCQTests/Private/AudioCQTestsModule.cpp b/Engine/Plugins/Tests/AudioCQTests/Source/AudioCQTests/Private/AudioCQTestsModule.cpp
new file mode 100644
--- /dev/null
+++ b/Engine/Plugins/Tests/AudioCQTests/Source/AudioCQTests/Private/AudioCQTestsModule.cpp
@@ -0,0 +1,5 @@
+// Copyright Epic Games, Inc. All Rights Reserved.
+
+#include "Modules/ModuleManager.h"
+
+IMPLEMENT_MODULE(FDefaultModuleImpl, AudioCQTests)
diff --git a/Engine/Plugins/Tests/AudioCQTests/Source/AudioCQTests/Private/AudioRenderSchedulerTests.cpp b/Engine/Plugins/Tests/AudioCQTests/Source/AudioCQTests/Private/AudioRenderSchedulerTests.cpp
new file mode 100644
--- /dev/null
+++ b/Engine/Plugins/Tests/AudioCQTests/Source/AudioCQTests/Private/AudioRenderSchedulerTests.cpp
@@ -0,0 +1,329 @@
+// Copyright Epic Games, Inc. All Rights Reserved.
+
+#if WITH_DEV_AUTOMATION_TESTS
+
+#include "Algo/RandomShuffle.h"
+#include "AudioRenderSchedulerTester.h"
+#include "CoreMinimal.h"
+#include "CQTest.h"
+
+TEST_CLASS(AudioRenderSchedulerTests, "Audio.RenderScheduler")
+{
+	TUniquePtr<IAudioRenderSchedulerTester> Tester;
+	TArray<int> Results;
+
+	void AddStepRange(int Num)
+	{
+		// Add the steps in random order for better test coverage
+		TArray<int> Range;
+		for (int i = 0; i < Num; ++i)
+		{
+			Range.Add(i);
+		}
+
+		Algo::RandomShuffle(Range);
+		for (int i : Range)
+		{
+			Tester->AddStep(i);
+		}
+	}
+
+	BEFORE_EACH()
+	{
+		Tester = IAudioRenderSchedulerTester::Create();
+	}
+
+	AFTER_EACH()
+	{
+		Tester.Reset();
+		Results.Empty();
+	}
+
+	TEST_METHOD(EmptyTest)
+	{
+		TestCommandBuilder.Do([this]()
+			{
+				Results = Tester->Run();
+				ASSERT_THAT(AreEqual(Results.Num(), 0));
+			});
+	}
+
+	TEST_METHOD(OneStep)
+	{
+		TestCommandBuilder.Do([this]()
+			{
+				Tester->AddStep(7);
+
+				Results = Tester->Run();
+
+				ASSERT_THAT(AreEqual(Results.Num(), 1));
+				ASSERT_THAT(AreEqual(Results[0], 7));
+			});
+	}
+
+	TEST_METHOD(TwoSteps)
+	{
+		TestCommandBuilder.Do([this]()
+			{
+				AddStepRange(2);
+				Tester->AddDependency(0, 1);
+
+				Results = Tester->Run();
+
+				ASSERT_THAT(AreEqual(Results.Num(), 2));
+				ASSERT_THAT(AreEqual(Results[0], 0));
+				ASSERT_THAT(AreEqual(Results[1], 1));
+			});
+	}
+
+	TEST_METHOD(UnconnectedSteps)
+	{
+		TestCommandBuilder.Do([this]()
+			{
+				AddStepRange(10);
+
+				Results = Tester->Run();
+
+				ASSERT_THAT(AreEqual(Results.Num(), 10));
+				for (int i = 0; i < 10; ++i)
+				{
+					ASSERT_THAT(IsTrue(Results.Contains(i)));
+				}
+			});
+	}
+
+	TEST_METHOD(SinglePrereq)
+	{
+		TestCommandBuilder.Do([this]()
+			{
+				AddStepRange(10);
+				for (int i = 1; i < 10; ++i)
+				{
+					Tester->AddDependency(0, i);
+				}
+
+				Results = Tester->Run();
+
+				ASSERT_THAT(AreEqual(Results.Num(), 10));
+				for (int i = 0; i < 10; ++i)
+				{
+					ASSERT_THAT(IsTrue(Results.Contains(i)));
+				}
+				ASSERT_THAT(AreEqual(Results[0], 0));
+			});
+	}
+
+	TEST_METHOD(SinglePostreq)
+	{
+		TestCommandBuilder.Do([this]()
+			{
+				AddStepRange(10);
+				for (int i = 1; i < 10; ++i)
+				{
+					Tester->AddDependency(i, 0);
+				}
+
+				Results = Tester->Run();
+
+				ASSERT_THAT(AreEqual(Results.Num(), 10));
+				for (int i = 0; i < 10; ++i)
+				{
+					ASSERT_THAT(IsTrue(Results.Contains(i)));
+				}
+				ASSERT_THAT(AreEqual(Results[9], 0));
+			});
+	}
+
+	TEST_METHOD(BasicCylce)
+	{
+		TestCommandBuilder.Do([this]()
+			{
+				AddStepRange(3);
+				Tester->AddDependency(0, 1);
+				Tester->AddDependency(1, 2);
+				Tester->AddDependency(2, 0);
+
+				Results = Tester->Run();
+
+				// One link of the cycle should be broken, but we don't know exactly where.
+				// Results should be some cyclic permutation of [0, 1, 2].
+				ASSERT_THAT(AreEqual(Results.Num(), 3));
+				ASSERT_THAT(AreEqual(Results[1], (Results[0] + 1) % 3));
+				ASSERT_THAT(AreEqual(Results[2], (Results[1] + 1) % 3));
+			});
+	}
+
+	TEST_METHOD(MiddleCycle)
+	{
+		TestCommandBuilder.Do([this]()
+			{
+				AddStepRange(5);
+				Tester->AddDependency(0, 1);
+				Tester->AddDependency(1, 2);
+				Tester->AddDependency(2, 3);
+				Tester->AddDependency(3, 1);
+				Tester->AddDependency(1, 4);
+				Tester->AddDependency(2, 4);
+				Tester->AddDependency(3, 4);
+
+				Results = Tester->Run();
+
+				// Should be 0, then a cyclic permutation of [1,2,3], then 4.
+				ASSERT_THAT(AreEqual(Results.Num(), 5));
+				ASSERT_THAT(AreEqual(Results[0], 0));
+				ASSERT_THAT(AreEqual(Results[4], 4));
+				ASSERT_THAT(IsTrue(Results.Contains(1) && Results.Contains(2) && Results.Contains(3)));
+				ASSERT_THAT(AreEqual((Results[2] - Results[1] - 1) % 3, 0));
+				ASSERT_THAT(AreEqual((Results[3] - Results[2] - 1) % 3, 0));
+			});
+	}
+
+	TEST_METHOD(TwoCycles)
+	{
+		TestCommandBuilder.Do([this]()
+			{
+				AddStepRange(7);
+				Tester->AddDependency(0, 1);
+				Tester->AddDependency(1, 2);
+				Tester->AddDependency(2, 0);
+
+				Tester->AddDependency(0, 3);
+				Tester->AddDependency(1, 3);
+				Tester->AddDependency(2, 3);
+
+				Tester->AddDependency(3, 4);
+				Tester->AddDependency(3, 5);
+				Tester->AddDependency(3, 6);
+
+				Tester->AddDependency(4, 5);
+				Tester->AddDependency(5, 6);
+				Tester->AddDependency(6, 4);
+
+				Results = Tester->Run();
+
+				// Should be a cyclic permutation of [0,1,2], then 3, then a cyclic permutation of [4,5,6].
+				ASSERT_THAT(AreEqual(Results.Num(), 7));
+				ASSERT_THAT(AreEqual((Results[1] - Results[0] - 1) % 3, 0));
+				ASSERT_THAT(AreEqual((Results[2] - Results[1] - 1) % 3, 0));
+				ASSERT_THAT(AreEqual(Results[3], 3));
+				ASSERT_THAT(AreEqual((Results[5] - Results[4] - 1) % 3, 0));
+				ASSERT_THAT(AreEqual((Results[6] - Results[5] - 1) % 3, 0));
+			});
+	}
+
+	TEST_METHOD(OnlyPhantomSteps)
+	{
+		TestCommandBuilder.Do([this]()
+			{
+				// In this test we register a dependency, but not the actual steps involved.
+				Tester->AddDependency(0, 1);
+
+				Results = Tester->Run();
+
+				ASSERT_THAT(AreEqual(Results.Num(), 0));
+			});
+	}
+
+	TEST_METHOD(MiddlePhantomSteps)
+	{
+		TestCommandBuilder.Do([this]()
+			{
+				AddStepRange(2);
+				Tester->AddDependency(1, 2);
+				Tester->AddDependency(2, 3);
+				Tester->AddDependency(3, 0);
+
+				Results = Tester->Run();
+
+				ASSERT_THAT(AreEqual(Results.Num(), 2));
+				ASSERT_THAT(AreEqual(Results[0], 1));
+				ASSERT_THAT(AreEqual(Results[1], 0));
+			});
+	}
+
+	TEST_METHOD(DenseGraph)
+	{
+		TestCommandBuilder.Do([this]()
+			{
+				AddStepRange(10);
+				for (int i = 0; i < 10; ++i)
+				{
+					for (int j = i; j < 10; ++j)
+					{
+						Tester->AddDependency(i, j);
+					}
+				}
+
+				Results = Tester->Run();
+
+				ASSERT_THAT(AreEqual(Results.Num(), 10));
+				for (int i = 0; i < 10; ++i)
+				{
+					ASSERT_THAT(AreEqual(Results[i], i));
+				}
+			});
+	}
+
+	TEST_METHOD(TwoCyclesRepeated)
+	{
+		TestCommandBuilder.Do([this]()
+			{
+				// 0-1-2 form a cycle, then 3, then 4-5-6 in a cycle
+				AddStepRange(7);
+				Tester->AddDependency(0, 1);
+				Tester->AddDependency(1, 2);
+				Tester->AddDependency(2, 0);
+
+				Tester->AddDependency(0, 3);
+				Tester->AddDependency(1, 3);
+				Tester->AddDependency(2, 3);
+
+				Tester->AddDependency(3, 4);
+				Tester->AddDependency(3, 5);
+				Tester->AddDependency(3, 6);
+
+				Tester->AddDependency(4, 5);
+				Tester->AddDependency(5, 6);
+				Tester->AddDependency(6, 4);
+
+				Results = Tester->Run();
+
+				ASSERT_THAT(AreEqual(Results.Num(), 7));
+				ASSERT_THAT(AreEqual((Results[1] - Results[0] - 1) % 3, 0));
+				ASSERT_THAT(AreEqual((Results[2] - Results[1] - 1) % 3, 0));
+				ASSERT_THAT(AreEqual(Results[3], 3));
+				ASSERT_THAT(AreEqual((Results[5] - Results[4] - 1) % 3, 0));
+				ASSERT_THAT(AreEqual((Results[6] - Results[5] - 1) % 3, 0));
+
+				// Check that running without changing anything gets the same results
+				Results = Tester->Run();
+
+				ASSERT_THAT(AreEqual(Results.Num(), 7));
+				ASSERT_THAT(AreEqual((Results[1] - Results[0] - 1) % 3, 0));
+				ASSERT_THAT(AreEqual((Results[2] - Results[1] - 1) % 3, 0));
+				ASSERT_THAT(AreEqual(Results[3], 3));
+				ASSERT_THAT(AreEqual((Results[5] - Results[4] - 1) % 3, 0));
+				ASSERT_THAT(AreEqual((Results[6] - Results[5] - 1) % 3, 0));
+
+				// Now rearrange a bit and check that the results update correctly.
+				Tester->RemoveDependency(3, 4);
+				Tester->RemoveDependency(4, 5);
+				Tester->RemoveDependency(6, 4);
+				Tester->AddDependency(4, 0);
+				Tester->AddDependency(4, 1);
+				Tester->AddDependency(4, 2);
+
+				Results = Tester->Run();
+
+				ASSERT_THAT(AreEqual(Results.Num(), 7));
+				ASSERT_THAT(AreEqual(Results[0], 4));
+				ASSERT_THAT(AreEqual((Results[2] - Results[1] - 1) % 3, 0));
+				ASSERT_THAT(AreEqual((Results[3] - Results[2] - 1) % 3, 0));
+				ASSERT_THAT(AreEqual(Results[4], 3));
+				ASSERT_THAT(AreEqual(Results[5], 5));
+				ASSERT_THAT(AreEqual(Results[6], 6));
+			});
+	}
+};
+
+#endif // WITH_DEV_AUTOMATION_TESTS
\ No newline at end of file
diff --git a/Engine/Source/Runtime/AudioMixer/Private/AudioMixerDevice.cpp b/Engine/Source/Runtime/AudioMixer/Private/AudioMixerDevice.cpp
--- a/Engine/Source/Runtime/AudioMixer/Private/AudioMixerDevice.cpp
+++ b/Engine/Source/Runtime/AudioMixer/Private/AudioMixerDevice.cpp
@@ -14,6 +14,7 @@
 #include "AudioMixerSourceVoice.h"
 #include "AudioPluginUtilities.h"
 #include "AudioMixerEffectsManager.h"
+#include "AudioRenderScheduler.h"
 #include "DSP/Noise.h"
 #include "DSP/SinOsc.h"
 #include "Sound/AudioSettings.h"
@@ -344,6 +345,7 @@ namespace Audio
 		bAudioMixerModuleLoaded = true;
 
 		SourceManager = MakeUnique<FMixerSourceManager>(this);
+		RenderScheduler = MakeUnique<FAudioRenderScheduler>(this);
 
 		// Register AudioLink Factory. 	
 		TArray<FName> Factories = IAudioLinkFactory::GetAllRegisteredFactoryNames();
diff --git a/Engine/Source/Runtime/AudioMixer/Private/AudioMixerSourceManager.cpp b/Engine/Source/Runtime/AudioMixer/Private/AudioMixerSourceManager.cpp
--- a/Engine/Source/Runtime/AudioMixer/Private/AudioMixerSourceManager.cpp
+++ b/Engine/Source/Runtime/AudioMixer/Private/AudioMixerSourceManager.cpp
@@ -9,6 +9,7 @@
 #include "AudioMixerSourceVoice.h"
 #include "AudioMixerSubmix.h"
 #include "AudioMixerTrace.h"
+#include "AudioRenderScheduler.h"
 #include "AudioThread.h"
 #include "DSP/FloatArrayMath.h"
 #include "IAudioExtensionPlugin.h"
@@ -33,6 +34,7 @@
 
 // Link to "Audio" profiling category
 CSV_DECLARE_CATEGORY_MODULE_EXTERN(AUDIOMIXERCORE_API, Audio);
+
 static int32 DisableParallelSourceProcessingCvar = 0;
 FAutoConsoleVariableRef CVarDisableParallelSourceProcessing(
 	TEXT("au.DisableParallelSourceProcessing"),
@@ -190,6 +192,14 @@ FAutoConsoleVariableRef CVarAudioMixerPerSourceResampling(
 	TEXT("0: Disabled, 1: Enabled"),
 	ECVF_Default);
 
+static int32 UseRenderSchedulerCVar = 0;
+FAutoConsoleVariableRef CVarUseAudioRenderScheduler(
+	TEXT("au.UseRenderScheduler"),
+	UseRenderSchedulerCVar,
+	TEXT("Use audio rendering scheduler to handle render dependencies.\n")
+	TEXT("0: Disabled, 1: Enabled"),
+	ECVF_Default);
+
 #define ENVELOPE_TAIL_THRESHOLD (1.58489e-5f) // -96 dB
 
 #define VALIDATE_SOURCE_MIXER_STATE 1
@@ -407,6 +417,7 @@ namespace Audio
 		{
 			FSourceInfo& SourceInfo = SourceInfos[i];
 
+			SourceInfo.SourceManager = this;
 			SourceInfo.MixerSourceBuffer = nullptr;
 
 			SourceInfo.VolumeSourceStart = -1.0f;
@@ -1240,6 +1251,12 @@ namespace Audio
 
 			SourceSubmixOutputBuffer.Reset(SourceSubmixOutputResetSettings);
 
+			const bool bIsBus = InitParams.AudioBusId != INDEX_NONE;
+			if (UseRenderSchedulerCVar && !bIsBus)
+			{
+				MixerDevice->GetRenderScheduler().AddStep(FAudioRenderStepId::FromWaveInstancePlayOrder(InitParams.PlayOrder), &SourceInfo);
+			}
+
 #if AUDIO_MIXER_ENABLE_DEBUG_MODE
 			AUDIO_MIXER_CHECK(!SourceInfo.bIsDebugMode);
 			SourceInfo.bIsDebugMode = InitParams.bIsDebugMode;
@@ -1277,6 +1294,12 @@ namespace Audio
 		{
 			AUDIO_MIXER_CHECK_AUDIO_PLAT_THREAD(MixerDevice);
 
+			const bool bIsBus = SourceInfos[SourceId].AudioBusId != INDEX_NONE;
+			if (UseRenderSchedulerCVar && !bIsBus)
+			{
+				MixerDevice->GetRenderScheduler().RemoveStep(FAudioRenderStepId::FromWaveInstancePlayOrder(SourceInfos[SourceId].PlayOrder));
+			}
+
 			ReleaseSource(SourceId);
 		});
 	}
@@ -3244,11 +3267,9 @@ namespace Audio
 							continue;
 
 						// Add a nested task to render this source
-						auto RenderOne = [this, bGenerateBuses, SourceId]()
+						auto RenderOne = [this, SourceId]()
 							{
-								ComputeSourceBuffer(bGenerateBuses, SourceId);
-								ComputePostSourceEffectBuffer(bGenerateBuses, SourceId);
-								ComputeOutputBuffers(bGenerateBuses, SourceId);
+								RenderSource(SourceId);
 							};
 						UE::Tasks::AddNested(UE::Tasks::Launch(TEXT("Render Audio Source"), MoveTemp(RenderOne), LowLevelTasks::ETaskPriority::High));
 					}
@@ -3278,6 +3299,14 @@ namespace Audio
 		}
 	}
 
+	void FMixerSourceManager::RenderSource(const int32 SourceId)
+	{
+		const bool bIsBus = SourceInfos[SourceId].AudioBusId != INDEX_NONE;
+		ComputeSourceBuffer(bIsBus, SourceId);
+		ComputePostSourceEffectBuffer(bIsBus, SourceId);
+		ComputeOutputBuffers(bIsBus, SourceId);
+	}
+
 	void FMixerSourceManager::Get2DChannelMap(const int32 SourceId, int32 InNumOutputChannels, Audio::FAlignedFloatBuffer& OutChannelMap)
 	{
 		AUDIO_MIXER_CHECK_AUDIO_PLAT_THREAD(MixerDevice);
@@ -3534,8 +3563,16 @@ namespace Audio
 		// Update pending tasks and release them if they're finished
 		UpdatePendingReleaseData();
 
-		// First generate non-bus audio (bGenerateBuses = false)
-		GenerateSourceAudio(false);
+		if (UseRenderSchedulerCVar)
+		{
+			RenderThreadPhase = ESourceManagerRenderThreadPhase::RunScheduledRenderSteps;
+			MixerDevice->GetRenderScheduler().RenderBlock(DisableParallelSourceProcessingCvar != 0);
+		}
+		else
+		{
+			// First generate non-bus audio (bGenerateBuses = false)
+			GenerateSourceAudio(false);
+		}
 
 		// Now mix in the non-bus audio into the buses
 		ComputeBuses();
@@ -3939,4 +3976,19 @@ namespace Audio
 			UE_LOG(LogAudioMixer, Warning, TEXT("***** ThreadStackWalkAndDump for ThreadId(%lu) ******\n%s"), InThreadId, ANSI_TO_TCHAR(StackTrace));
 		}
 	}
-}
+
+	void FMixerSourceManager::FSourceInfo::DoRenderStep()
+	{
+		const int32 SourceId = this - SourceManager->SourceInfos.GetData();
+		SourceManager->RenderSource(SourceId);
+	}
+
+	const TCHAR* FMixerSourceManager::FSourceInfo::GetRenderStepName()
+	{
+#if AUDIO_MIXER_ENABLE_DEBUG_MODE
+		return *DebugName;
+#else
+		return TEXT("Render mixer source");
+#endif
+	}
+}
\ No newline at end of file
diff --git a/Engine/Source/Runtime/AudioMixer/Private/AudioMixerSourceManager.h b/Engine/Source/Runtime/AudioMixer/Private/AudioMixerSourceManager.h
--- a/Engine/Source/Runtime/AudioMixer/Private/AudioMixerSourceManager.h
+++ b/Engine/Source/Runtime/AudioMixer/Private/AudioMixerSourceManager.h
@@ -15,6 +15,7 @@
 #include "DSP/ParamInterpolator.h"
 #include "DSP/RuntimeResampler.h"
 #include "IAudioExtensionPlugin.h"
+#include "IAudioMixerRenderStep.h"
 #include "ISoundfieldFormat.h"
 #include "Misc/SpinLock.h"
 #include "Sound/SoundModulationDestination.h"
@@ -53,6 +54,7 @@ enum class ESourceManagerRenderThreadPhase: uint8
 	PumpCmds,
 	ProcessModulators,
 	UpdatePendingReleaseData,
+	RunScheduledRenderSteps,
 	GenerateSrcAudio_WithBusses,
 	ComputeBusses,
 	GenerateSrcAudio_WithoutBusses,
@@ -261,6 +263,7 @@ namespace Audio
 		void UpdateSourceState();
 		void ClearStoppingSounds();
 		void MixOutputBuffers(const int32 SourceId, int32 InNumOutputChannels, const float InSendLevel, EMixerSourceSubmixSendStage InSubmixSendStage, FAlignedFloatBuffer& OutWetBuffer) const;
+		void RenderSource(const int32 SourceId);
 
 		// Retrieves a channel map for the given source ID for the given output channels
 		// can be used even when a source is 3D if the source is doing any kind of bus sending or otherwise needs a channel map
@@ -398,11 +401,14 @@ namespace Audio
 		using FAudioMixerMpscCommand = FAudioMixerThreadCommand;
 		TMpscQueue<FAudioMixerMpscCommand> MpscCommandQueue;
 		
-		struct FSourceInfo
+		struct FSourceInfo : public IAudioMixerRenderStep
 		{
 			FSourceInfo() : Resampler(1 /* NumChannels */) {}
 			~FSourceInfo() {}
 
+			// The owning source manager
+			FMixerSourceManager* SourceManager;
+
 			// Object which handles source buffer decoding
 			TSharedPtr<FMixerSourceBuffer, ESPMode::ThreadSafe> MixerSourceBuffer;
 			ISourceListener* SourceListener;
@@ -578,6 +584,9 @@ namespace Audio
 					return 0;
 				}
 			}
+
+			virtual void DoRenderStep() override;
+			virtual const TCHAR* GetRenderStepName() override;
 		};
 
 		static void ApplyDistanceAttenuation(FSourceInfo& InSourceInfo, int32 NumSamples);
diff --git a/Engine/Source/Runtime/AudioMixer/Private/AudioRenderScheduler.cpp b/Engine/Source/Runtime/AudioMixer/Private/AudioRenderScheduler.cpp
new file mode 100644
--- /dev/null
+++ b/Engine/Source/Runtime/AudioMixer/Private/AudioRenderScheduler.cpp
@@ -0,0 +1,372 @@
+// Copyright Epic Games, Inc. All Rights Reserved.
+
+#include "AudioRenderScheduler.h"
+
+#include "AudioMixerDevice.h"
+#include "AudioMixerTrace.h"
+#include "IAudioMixerRenderStep.h"
+#include "Tasks/Task.h"
+
+// If MixerDevice is null, don't check the thread. This is used in unit tests.
+#define AUDIO_RENDER_SCHEDULER_CHECK_THREAD() check(MixerDevice == nullptr || MixerDevice->IsAudioRenderingThread())
+
+namespace Audio
+{
+	FAudioRenderStepId FAudioRenderStepId::FromWaveInstancePlayOrder(uint32 PlayOrder)
+	{
+		return uint64(PlayOrder);
+	}
+
+	FAudioRenderScheduler::FAudioRenderScheduler(FMixerDevice* InMixerDevice) :
+		MixerDevice(InMixerDevice)
+	{
+#if !WITH_DEV_AUTOMATION_TESTS
+		// A null mixer device is only allowed in unit tests.
+		check(MixerDevice != nullptr);
+#endif
+	}
+
+	FAudioRenderScheduler::FInternalIndex FAudioRenderScheduler::GetInternalIndex(const FAudioRenderStepId Id)
+	{
+		FInternalIndex* IndexInMap = IndexMap.Find(Id);
+		if (IndexInMap)
+		{
+			return *IndexInMap;
+		}
+
+		FInternalIndex Index;
+		if (FirstFreeIndex != INDEX_NONE)
+		{
+			// Pop a step entry off the free list if available
+			Index = FirstFreeIndex;
+			FirstFreeIndex = Steps[Index].NextFreeIndex;
+		}
+		else
+		{
+			// Otherwise grow the steps array
+			check(Steps.Num() < TNumericLimits<FInternalIndex>::Max());
+			Steps.AddDefaulted(1);
+			Index = Steps.Num() - 1;
+		}
+
+		Steps[Index].StepId = Id;
+		IndexMap.Add(Id, Index);
+		ActiveSteps.Add(Index);
+
+		return Index;
+	}
+
+	void FAudioRenderScheduler::RemoveReference(FAudioRenderScheduler::FInternalIndex Index)
+	{
+		FStepEntry& Entry = Steps[Index];
+		check(Entry.ReferenceCount >= 1 && Entry.StepId.IsValid());
+		if (--Entry.ReferenceCount == 0)
+		{
+			IndexMap.Remove(Entry.StepId);
+			ActiveSteps.Remove(Index);
+			Entry = FStepEntry{};
+			Entry.NextFreeIndex = FirstFreeIndex;
+			FirstFreeIndex = Index;
+		}
+	}
+
+	void FAudioRenderScheduler::AddStep(const FAudioRenderStepId Id, IAudioMixerRenderStep* Step)
+	{
+		AUDIO_RENDER_SCHEDULER_CHECK_THREAD();
+
+		check(Id.IsValid() && Step != nullptr);
+		FInternalIndex Index = GetInternalIndex(Id);
+		FStepEntry& Entry = Steps[Index];
+		if (ensureMsgf(Entry.StepInterface == nullptr, TEXT("Duplicate render step added to audio render scheduler")))
+		{
+			Entry.StepInterface = Step;
+			Entry.ReferenceCount++;
+			bDirty = true;
+		}
+	}
+
+	void FAudioRenderScheduler::RemoveStep(const FAudioRenderStepId Id)
+	{
+		AUDIO_RENDER_SCHEDULER_CHECK_THREAD();
+
+		check(Id.IsValid());
+		FInternalIndex Index = GetInternalIndex(Id);
+		FStepEntry& Entry = Steps[Index];
+		if (ensureMsgf(Entry.StepInterface != nullptr, TEXT("Unknown render step removed from audio render scheduler")))
+		{
+			Entry.StepInterface = nullptr;
+			RemoveReference(Index);
+			bDirty = true;
+		}
+	}
+
+	void FAudioRenderScheduler::AddDependency(const FAudioRenderStepId FirstStep, const FAudioRenderStepId SecondStep)
+	{
+		AUDIO_RENDER_SCHEDULER_CHECK_THREAD();
+
+		check(FirstStep.IsValid() && SecondStep.IsValid());
+		if (FirstStep == SecondStep)
+		{
+			// We allow this, but ignore it because this dependency will always need to be broken anyway.
+			return;
+		}
+
+		const FInternalIndex FirstIndex = GetInternalIndex(FirstStep);
+		const FInternalIndex SecondIndex = GetInternalIndex(SecondStep);
+
+		Steps[SecondIndex].Prerequisites.Add(FirstIndex);
+		Steps[SecondIndex].ReferenceCount++;
+		Steps[FirstIndex].ReferenceCount++;
+
+		bDirty = true;
+	}
+
+	void FAudioRenderScheduler::RemoveDependency(const FAudioRenderStepId FirstStep, const FAudioRenderStepId SecondStep)
+	{
+		AUDIO_RENDER_SCHEDULER_CHECK_THREAD();
+
+		check(FirstStep.IsValid() && SecondStep.IsValid());
+		if (FirstStep == SecondStep)
+		{
+			return;
+		}
+
+		const FInternalIndex FirstIndex = GetInternalIndex(FirstStep);
+		const FInternalIndex SecondIndex = GetInternalIndex(SecondStep);
+
+		if (ensureMsgf(Steps[SecondIndex].Prerequisites.RemoveSingleSwap(FirstIndex), TEXT("Unknown dependency removed from audio render scheduler")))
+		{
+			RemoveReference(FirstIndex);
+			RemoveReference(SecondIndex);
+		}
+
+		bDirty = true;
+	}
+
+	void FAudioRenderScheduler::RenderBlock(const bool bSingleThreaded)
+	{
+		AUDIO_RENDER_SCHEDULER_CHECK_THREAD();
+		AUDIO_MIXER_TRACE_CPUPROFILER_EVENT_SCOPE(FAudioRenderScheduler::RenderBlock);
+
+		if (bSingleThreaded)
+		{
+			ScheduleAll(true);
+		}
+		else
+		{
+			UE::Tasks::FTask Task = UE::Tasks::Launch(
+				TEXT("Scheduled audio render steps"),
+				[this]() { ScheduleAll(false); },
+				LowLevelTasks::ETaskPriority::High);
+
+			Task.Wait();
+		}
+	}
+
+	void FAudioRenderScheduler::ScheduleAll(const bool bSingleThreaded)
+	{
+		AUDIO_MIXER_TRACE_CPUPROFILER_EVENT_SCOPE(FAudioRenderScheduler::ScheduleAll);
+
+		TArray<FInternalIndex> OrderedSteps;
+		OrderedSteps.Reserve(ActiveSteps.Num());
+		TArray<UE::Tasks::FTask> PrerequisiteHandles;
+
+		// If we need multiple passes, we avoid having to start from the beginning each time.
+		int FirstPendingStepIndex = 0;
+
+		// Check if any of the previously broken links need to be broken again
+		CheckBrokenLinks();
+
+		while (OrderedSteps.Num() < ActiveSteps.Num())
+		{
+			bool bStuck = true;
+
+			// Go through the steps array, launching any steps that have no prerequisites or prerequisites that have already been launched.
+			for (int ActiveStepIndex = FirstPendingStepIndex; ActiveStepIndex < ActiveSteps.Num(); ++ActiveStepIndex)
+			{
+				FStepEntry& Entry = Steps[ActiveSteps[ActiveStepIndex]];
+				if (Entry.bLaunched)
+				{
+					if (ActiveStepIndex == FirstPendingStepIndex)
+					{
+						++FirstPendingStepIndex;
+					}
+					continue;
+				}
+
+				bool bCanLaunch = true;
+				if (bDirty)
+				{
+					for (FInternalIndex PrerequisiteIndex : Entry.Prerequisites)
+					{
+						if (!Steps[PrerequisiteIndex].bLaunched)
+						{
+							bCanLaunch = false;
+							break;
+						}
+					}
+				}
+
+				if (bCanLaunch)
+				{
+					if (bSingleThreaded)
+					{
+						if (Entry.StepInterface != nullptr)
+						{
+							Entry.StepInterface->DoRenderStep();
+						}
+					}
+					else
+					{
+						// Gather prerequisite task handles
+						PrerequisiteHandles.Reset();
+						for (FInternalIndex PrerequisiteIndex : Entry.Prerequisites)
+						{
+							check(Steps[PrerequisiteIndex].TaskHandle.IsValid());
+							PrerequisiteHandles.Add(Steps[PrerequisiteIndex].TaskHandle);
+						}
+
+						auto TaskLambda = [&Entry]()
+							{
+								if (Entry.StepInterface != nullptr)
+								{
+									Entry.StepInterface->DoRenderStep();
+								}
+							};
+
+						// Launch the task
+						Entry.TaskHandle = UE::Tasks::Launch(
+							Entry.StepInterface ? Entry.StepInterface->GetRenderStepName() : TEXT("Empty audio render step"),
+							MoveTemp(TaskLambda),
+							PrerequisiteHandles,
+							LowLevelTasks::ETaskPriority::High);
+						UE::Tasks::AddNested(Entry.TaskHandle);
+					}
+
+					OrderedSteps.Add(ActiveSteps[ActiveStepIndex]);
+					if (ActiveStepIndex == FirstPendingStepIndex)
+					{
+						++FirstPendingStepIndex;
+					}
+					Entry.bLaunched = true;
+					bStuck = false;
+				}
+			}
+
+			if (bStuck)
+			{
+				// If we go through the whole list and find no steps that are launchable, it means there must be a cycle that needs to be broken.
+				BreakCycle(FirstPendingStepIndex);
+			}
+		}
+
+		// After scheduling everything, restore the broken links
+		for (TPair<FInternalIndex, FInternalIndex> Link : BrokenLinks)
+		{
+			Steps[Link.Get<1>()].Prerequisites.Add(Link.Get<0>());
+		}
+
+		// Reset the launched state for all steps
+		for (int i = 0; i < Steps.Num(); ++i)
+		{
+			Steps[i].TaskHandle = {};
+			Steps[i].bLaunched = false;
+		}
+
+		// Remember the ordering of steps for next block
+		ActiveSteps = MoveTemp(OrderedSteps);
+
+		// If no changes happen before next block, we can skip lots of the checks.
+		bDirty = false;
+	}
+
+	void FAudioRenderScheduler::BreakCycle(const int FirstPendingStepIndex)
+	{
+		FInternalIndex CurrentStep = ActiveSteps[FirstPendingStepIndex];
+		TArray<FInternalIndex> VisitedIndices;
+		VisitedIndices.Add(CurrentStep);
+
+		while (true)
+		{
+			FStepEntry& Entry = Steps[CurrentStep];
+			check(!Entry.bLaunched);
+			FInternalIndex PrerequisiteStep = INDEX_NONE;
+
+			// Find a pending step that is a dependency of CurrentStep
+			for (FInternalIndex PrerequisiteIndex : Entry.Prerequisites)
+			{
+				if (!Steps[PrerequisiteIndex].bLaunched)
+				{
+					PrerequisiteStep = PrerequisiteIndex;
+					break;
+				}
+			}
+
+			// If this fails, BreakCycle() was called unnecessarily (there was a step with no pending prerequisites.)
+			check(PrerequisiteStep != INDEX_NONE);
+
+			if (VisitedIndices.Contains(PrerequisiteStep))
+			{
+				// The link between CurrentStep and PrerequisiteStep is part of a cycle, so we'll break it.
+				Entry.Prerequisites.RemoveSingleSwap(PrerequisiteStep);
+				BrokenLinks.Add(TPair<FInternalIndex, FInternalIndex>{PrerequisiteStep, CurrentStep});
+				return;
+			}
+
+			CurrentStep = PrerequisiteStep;
+			VisitedIndices.Add(CurrentStep);
+		}
+
+		// It shouldn't be possible to get here; each time through the loop we extend VisitedIndices and check for duplicates.
+		checkNoEntry();
+	}
+
+	void FAudioRenderScheduler::CheckBrokenLinks()
+	{
+		for (int LinkIndex = 0; LinkIndex < BrokenLinks.Num(); )
+		{
+			const TPair<FInternalIndex, FInternalIndex>& Link = BrokenLinks[LinkIndex];
+			if (!bDirty || (Steps[Link.Get<1>()].Prerequisites.Contains(Link.Get<0>()) && FindIndirectLink(Link.Get<1>(), Link.Get<0>())))
+			{
+				// This link is still part of a cycle; break it again and leave it in BrokenLinks
+				Steps[Link.Get<1>()].Prerequisites.RemoveSingleSwap(Link.Get<0>());
+				++LinkIndex;
+			}
+			else
+			{
+				// This link is no longer part of a cycle; remove it and move on.
+				BrokenLinks.RemoveAt(LinkIndex);
+			}
+		}
+	}
+
+	bool FAudioRenderScheduler::FindIndirectLink(FInternalIndex FirstStep, FInternalIndex SecondStep)
+	{
+		TArray<FInternalIndex> StepsToSkip;
+		return FindIndirectLinkSkipping(FirstStep, SecondStep, StepsToSkip);
+	}
+
+	bool FAudioRenderScheduler::FindIndirectLinkSkipping(FInternalIndex FirstStep, FInternalIndex SecondStep, TArray<FInternalIndex>& StepsToSkip)
+	{
+		StepsToSkip.Add(SecondStep);
+		for (FInternalIndex Prereq : Steps[SecondStep].Prerequisites)
+		{
+			if (Prereq == FirstStep)
+			{
+				return true;
+			}
+
+			if (StepsToSkip.Contains(Prereq))
+			{
+				continue;
+			}
+
+			if (FindIndirectLinkSkipping(FirstStep, Prereq, StepsToSkip))
+			{
+				return true;
+			}
+		}
+
+		return false;
+	}
+}
diff --git a/Engine/Source/Runtime/AudioMixer/Private/AudioRenderScheduler.h b/Engine/Source/Runtime/AudioMixer/Private/AudioRenderScheduler.h
new file mode 100644
--- /dev/null
+++ b/Engine/Source/Runtime/AudioMixer/Private/AudioRenderScheduler.h
@@ -0,0 +1,139 @@
+// Copyright Epic Games, Inc. All Rights Reserved.
+
+#pragma once
+
+#include "Containers/Map.h"
+#include "Tasks/Task.h"
+
+namespace Audio
+{
+	// Forward declarations
+	struct FAudioBusKey;
+	class FMixerDevice;
+	struct IAudioMixerRenderStep;
+
+	// An identifier that uniquely identifies a given render step to the scheduler.
+	class FAudioRenderStepId
+	{
+	public:
+		FAudioRenderStepId() :
+			Value(INDEX_NONE)
+		{
+		}
+
+		static FAudioRenderStepId FromWaveInstancePlayOrder(uint32 PlayOrder);
+
+		uint64 GetRawValue() const
+		{
+			return Value;
+		}
+
+		friend uint32 GetTypeHash(FAudioRenderStepId Id)
+		{
+			return GetTypeHash(Id.Value);
+		}
+
+		friend bool operator==(const FAudioRenderStepId InLHS, const FAudioRenderStepId InRHS)
+		{
+			return InLHS.Value == InRHS.Value;
+		}
+
+		friend bool operator!=(const FAudioRenderStepId InLHS, const FAudioRenderStepId InRHS)
+		{
+			return InLHS.Value != InRHS.Value;
+		}
+
+		bool IsValid() const
+		{
+			return Value != INDEX_NONE;
+		}
+
+	protected:
+		FAudioRenderStepId(uint64 InValue) :
+			Value(InValue)
+		{
+		}
+
+		uint64 Value;
+	};
+
+	// FAudioRenderScheduler allows individual audio render work items ("steps") to be added to be run as part of each rendering block.
+	//  Dependencies can be declared to enforce ordering and synchronization between steps when needed. There can be cyclical
+	//  dependencies, in which case the scheduler will do its best to consistently break the cycles in the same way every block;
+	//  but cycles should be avoided when possible.
+	//
+	// Render steps are executed on worker threads using the UE Tasks system, unless single-threaded rendering is requested. Currently
+	//  only source rendering uses the scheduler.
+	//
+	// All methods of this class are to be called from the audio render thread only.
+
+	class FAudioRenderScheduler
+	{
+	public:
+		UE_NONCOPYABLE(FAudioRenderScheduler);
+
+		FAudioRenderScheduler(FMixerDevice* InMixerDevice);
+
+		// Add/remove a render step.
+		void AddStep(const FAudioRenderStepId Id, IAudioMixerRenderStep* Step);
+		void RemoveStep(const FAudioRenderStepId Id);
+
+		// Declare that one step is dependent on another. It's possible to add a dependency without
+		//  having previously called AddStep; in this case the dependency will still be respected with
+		//  an empty placeholder step.
+		void AddDependency(const FAudioRenderStepId FirstStep, const FAudioRenderStepId SecondStep);
+
+		// All dependencies declared with AddDependency() must be removed by a corresponding call to RemoveDependency().
+		void RemoveDependency(const FAudioRenderStepId FirstStep, const FAudioRenderStepId SecondStep);
+
+		// Execute all steps and block until they're done.
+		void RenderBlock(const bool bSingleThreaded);
+
+	private:
+		// Internally the step id's are mapped to smaller internal indices to reduce
+		//  memory footprint.
+		using FInternalIndex = int16;
+
+		void ScheduleAll(const bool bSingleThreaded);
+
+		FInternalIndex GetInternalIndex(const FAudioRenderStepId Id);
+		void RemoveReference(FAudioRenderScheduler::FInternalIndex Index);
+
+		void BreakCycle(const int FirstPendingStepIndex);
+		void CheckBrokenLinks();
+
+		// Helpers to verify that a link is still part of a cycle
+		bool FindIndirectLink(const FInternalIndex FirstStep, const FInternalIndex SecondStep);
+		bool FindIndirectLinkSkipping(const FInternalIndex FirstStep, const FInternalIndex SecondStep, TArray<FInternalIndex>& StepsToSkip);
+
+		struct FStepEntry
+		{
+			FAudioRenderStepId StepId;
+			IAudioMixerRenderStep* StepInterface = nullptr;
+			TArray<FInternalIndex, TInlineAllocator<4>> Prerequisites;
+			int32 ReferenceCount = 0;
+			bool bLaunched = false;
+			UE::Tasks::FTask TaskHandle;
+
+			// Unused step entries are added to a free list.
+			FInternalIndex NextFreeIndex = INDEX_NONE;
+		};
+
+		// Only used to check the current thread; may be null.
+		FMixerDevice* MixerDevice;
+
+		TMap<FAudioRenderStepId, FInternalIndex> IndexMap;
+		TArray<FStepEntry> Steps;
+		FInternalIndex FirstFreeIndex = INDEX_NONE;
+
+		// Array of all used step indices. Ends up sorted at the end of RenderBlock().
+		TArray<FInternalIndex> ActiveSteps;
+
+		// Have any changes happened since the previous call to RenderBlock()?
+		bool bDirty = true;
+
+		// Links that have been broken to eliminate cycles. These are retained between
+		//  blocks to improve stability of step ordering.
+		TArray<TPair<FInternalIndex, FInternalIndex>> BrokenLinks;
+	};
+}
diff --git a/Engine/Source/Runtime/AudioMixer/Private/AudioRenderSchedulerTester.cpp b/Engine/Source/Runtime/AudioMixer/Private/AudioRenderSchedulerTester.cpp
new file mode 100644
--- /dev/null
+++ b/Engine/Source/Runtime/AudioMixer/Private/AudioRenderSchedulerTester.cpp
@@ -0,0 +1,80 @@
+// Copyright Epic Games, Inc. All Rights Reserved.
+
+#if WITH_DEV_AUTOMATION_TESTS
+
+#include "AudioRenderSchedulerTester.h"
+
+#include "AudioRenderScheduler.h"
+#include "IAudioMixerRenderStep.h"
+
+namespace Audio
+{
+	// Simple test harness for the audio render scheduler. You can add steps identified by number and dependencies
+	//  between them. Running the scheduler returns an array that shows the order the steps were executed in.
+	class FAudioRenderSchedulerTesterImpl : public IAudioRenderSchedulerTester
+	{
+		class FTestStep : public IAudioMixerRenderStep
+		{
+		public:
+			FTestStep(FAudioRenderSchedulerTesterImpl& InOwner, int InId) :
+				Owner(InOwner),
+				Id(InId)
+			{
+			}
+
+		protected:
+			virtual void DoRenderStep() override
+			{
+				FScopeLock Lock(&Owner.ResultsLock);
+				Owner.Results.Add(Id);
+			}
+
+			virtual const TCHAR* GetRenderStepName() override
+			{
+				return TEXT("");
+			}
+
+			FAudioRenderSchedulerTesterImpl& Owner;
+			int Id;
+		};
+
+		TMap<int, TUniquePtr<FTestStep>> Steps;
+		FAudioRenderScheduler Scheduler{ /*MixerDevice=*/ nullptr};
+
+		FCriticalSection ResultsLock;
+		TArray<int> Results;
+
+		virtual void AddStep(int Id) override
+		{
+			if (!Steps.Contains(Id))
+			{
+				Steps.Emplace(Id, MakeUnique<FTestStep>(*this, Id));
+				Scheduler.AddStep(FAudioRenderStepId::FromWaveInstancePlayOrder(Id), Steps[Id].Get());
+			}
+		}
+
+		virtual void AddDependency(int FirstStep, int SecondStep) override
+		{
+			Scheduler.AddDependency(FAudioRenderStepId::FromWaveInstancePlayOrder(FirstStep), FAudioRenderStepId::FromWaveInstancePlayOrder(SecondStep));
+		}
+
+		virtual void RemoveDependency(int FirstStep, int SecondStep) override
+		{
+			Scheduler.RemoveDependency(FAudioRenderStepId::FromWaveInstancePlayOrder(FirstStep), FAudioRenderStepId::FromWaveInstancePlayOrder(SecondStep));
+		}
+
+		virtual TArray<int> Run() override
+		{
+			Results.Reset();
+			Scheduler.RenderBlock(false);
+			return Results;
+		}
+	};
+}
+
+TUniquePtr<IAudioRenderSchedulerTester> IAudioRenderSchedulerTester::Create()
+{
+	return MakeUnique<Audio::FAudioRenderSchedulerTesterImpl>();
+}
+
+#endif //WITH_DEV_AUTOMATION_TESTS
diff --git a/Engine/Source/Runtime/AudioMixer/Private/IAudioMixerRenderStep.h b/Engine/Source/Runtime/AudioMixer/Private/IAudioMixerRenderStep.h
new file mode 100644
--- /dev/null
+++ b/Engine/Source/Runtime/AudioMixer/Private/IAudioMixerRenderStep.h
@@ -0,0 +1,23 @@
+// Copyright Epic Games, Inc. All Rights Reserved.
+
+#pragma once
+
+#include "CoreMinimal.h"
+
+namespace Audio
+{
+	/** Interface representing an item of audio rendering work that can be scheduled as part
+	of an AudioMixerDevice's rendering pass.
+	*/
+	struct IAudioMixerRenderStep
+	{
+		virtual ~IAudioMixerRenderStep() = default;
+
+		/* Do the work associated with this render step. */
+		virtual void DoRenderStep() = 0;
+
+		/* Get a name for this render step to display for debugging purposes. */
+		virtual const TCHAR* GetRenderStepName() = 0;
+	};
+
+}
diff --git a/Engine/Source/Runtime/AudioMixer/Public/AudioMixerDevice.h b/Engine/Source/Runtime/AudioMixer/Public/AudioMixerDevice.h
--- a/Engine/Source/Runtime/AudioMixer/Public/AudioMixerDevice.h
+++ b/Engine/Source/Runtime/AudioMixer/Public/AudioMixerDevice.h
@@ -40,6 +40,7 @@ namespace Audio
 	class FMixerSourceVoice;
 	class FMixerSubmix;
 	class FAudioFormatSettings;
+	class FAudioRenderScheduler;
 
 	typedef TSharedPtr<FMixerSubmix, ESPMode::ThreadSafe> FMixerSubmixPtr;
 	typedef TWeakPtr<FMixerSubmix, ESPMode::ThreadSafe> FMixerSubmixWeakPtr;
@@ -359,6 +360,11 @@ namespace Audio
 		// Debug Commands
 		AUDIOMIXER_API void DrawSubmixes(FOutputDevice& InOutput, const TArray<FString>& InArgs) const;
 
+		FAudioRenderScheduler& GetRenderScheduler() const
+		{
+			return *RenderScheduler;
+		}
+
 	protected:
 		AUDIOMIXER_API virtual void InitSoundSubmixes() override;
 
@@ -503,6 +509,9 @@ namespace Audio
 		/** Whether or not initialization of the submix system is underway and submixes can be registered */
 		bool bSubmixRegistrationDisabled;
 
+		/** Scheduler responsible for launching rendering tasks */
+		TUniquePtr<FAudioRenderScheduler> RenderScheduler;
+
 	public:
 
 		// Creates a queue for audio decode requests with a specific Id. Tasks
diff --git a/Engine/Source/Runtime/AudioMixer/Public/AudioRenderSchedulerTester.h b/Engine/Source/Runtime/AudioMixer/Public/AudioRenderSchedulerTester.h
new file mode 100644
--- /dev/null
+++ b/Engine/Source/Runtime/AudioMixer/Public/AudioRenderSchedulerTester.h
@@ -0,0 +1,30 @@
+// Copyright Epic Games, Inc. All Rights Reserved.
+
+#pragma once
+
+#if WITH_DEV_AUTOMATION_TESTS
+
+#include "CoreMinimal.h"
+
+// A simple test harness for the scheduler. Numbered steps and dependencies between them can be added and
+//  the order the steps were run in is returned.
+class IAudioRenderSchedulerTester
+{
+public:
+	static AUDIOMIXER_API TUniquePtr<IAudioRenderSchedulerTester> Create();
+
+	virtual ~IAudioRenderSchedulerTester() = default;
+
+	// Add a numbered step to the scheduler.
+	virtual void AddStep(int Step) = 0;
+
+	// Automatically creates the steps if they don't exist yet. 
+	virtual void AddDependency(int FirstStep, int SecondStep) = 0;
+
+	virtual void RemoveDependency(int FirstStep, int SecondStep) = 0;
+
+	// Runs the scheduler, and returns the order steps were executed in.
+	virtual TArray<int> Run() = 0;
+};
+
+#endif //WITH_DEV_AUTOMATION_TESTS
