[AudioMixer] Pass bus name so it can be traced and then received by [REDACTED_PERSON] dashboard in Audio Insights standalone

diff --git a/Engine/Source/Runtime/AudioMixer/Private/AudioMixerSource.cpp b/Engine/Source/Runtime/AudioMixer/Private/AudioMixerSource.cpp
--- a/Engine/Source/Runtime/AudioMixer/Private/AudioMixerSource.cpp
+++ b/Engine/Source/Runtime/AudioMixer/Private/AudioMixerSource.cpp
@@ -406,11 +406,17 @@ namespace Audio
 					{
 						InitParams.AudioBusId = SoundSourceBus->AudioBus->GetUniqueID();
 						InitParams.AudioBusChannels = (int32)SoundSourceBus->AudioBus->GetNumChannels();
+#if ENABLE_AUDIO_DEBUG
+						InitParams.AudioBusName = SoundSourceBus->AudioBus->GetPathName();
+#endif // ENABLE_AUDIO_DEBUG
 					}
 					else
 					{
 						InitParams.AudioBusId = WaveData->GetUniqueID();
 						InitParams.AudioBusChannels = WaveData->NumChannels;
+#if ENABLE_AUDIO_DEBUG
+						InitParams.AudioBusName = WaveData->GetPathName();
+#endif // ENABLE_AUDIO_DEBUG
 					}
 
 					if (!WaveData->IsLooping())
@@ -629,11 +635,14 @@ namespace Audio
 			for (FSoundSourceBusSendInfo& SendInfo : WaveInstance->BusSends[BusSendType])
 			{
 				// Avoid redoing duplicate code for sending audio to source bus or audio bus. Most of it is the same other than the bus id.
-				auto SetupBusSend = [this](TArray<FInitAudioBusSend>* AudioBusSends, const FSoundSourceBusSendInfo& InSendInfo, int32 InBusSendType, uint32 InBusId, bool bEnableBusSends, int32 InBusChannels)
+				auto SetupBusSend = [this](TArray<FInitAudioBusSend>* AudioBusSends, const FSoundSourceBusSendInfo& InSendInfo, int32 InBusSendType, uint32 InBusId, bool bEnableBusSends, int32 InBusChannels, const FString& InBusName)
 				{
 					FInitAudioBusSend BusSend;
 					BusSend.AudioBusId = InBusId;
 					BusSend.BusChannels = InBusChannels;
+#if ENABLE_AUDIO_DEBUG
+					BusSend.AudioBusName = InBusName;
+#endif // ENABLE_AUDIO_DEBUG
 					
 					if(bEnableBusSends)
 					{
@@ -652,6 +661,9 @@ namespace Audio
 					FDynamicBusSendInfo NewDynamicBusSendInfo;
 					NewDynamicBusSendInfo.SendLevel = InSendInfo.SendLevel;
 					NewDynamicBusSendInfo.BusId = BusSend.AudioBusId;
+#if ENABLE_AUDIO_DEBUG
+					NewDynamicBusSendInfo.BusName = BusSend.AudioBusName;
+#endif // ENABLE_AUDIO_DEBUG
 					NewDynamicBusSendInfo.BusSendLevelControlMethod = InSendInfo.SourceBusSendLevelControlMethod;
 					NewDynamicBusSendInfo.BusSendType = (EBusSendType)InBusSendType;
 					NewDynamicBusSendInfo.MinSendLevel = InSendInfo.MinSendLevel;
@@ -687,22 +699,29 @@ namespace Audio
 				{						
 					uint32 BusId;
 					int32 BusChannels;
+					FString BusName;
 
 					// Either use the bus id of the source bus's audio bus id if it was specified
 					if (SendInfo.SoundSourceBus->AudioBus)
 					{
 						BusId = SendInfo.SoundSourceBus->AudioBus->GetUniqueID();
 						BusChannels = (int32)SendInfo.SoundSourceBus->AudioBus->GetNumChannels();
+#if ENABLE_AUDIO_DEBUG
+						BusName = SendInfo.SoundSourceBus->AudioBus.GetPathName();
+#endif // ENABLE_AUDIO_DEBUG
 					}
 					else
 					{
 						// otherwise, use the id of the source bus itself (for an automatic source bus)
 						BusId = SendInfo.SoundSourceBus->GetUniqueID();
 						BusChannels = SendInfo.SoundSourceBus->NumChannels;
+#if ENABLE_AUDIO_DEBUG
+						BusName = SendInfo.SoundSourceBus->GetPathName();
+#endif // ENABLE_AUDIO_DEBUG
 					}
 
 					// Call lambda w/ the correctly derived bus id
-					SetupBusSend(OutAudioBusSends, SendInfo, BusSendType, BusId, bEnableBusSends, BusChannels);
+					SetupBusSend(OutAudioBusSends, SendInfo, BusSendType, BusId, bEnableBusSends, BusChannels, BusName);
 				}
 
 				if (SendInfo.AudioBus)
@@ -710,9 +729,14 @@ namespace Audio
 					// Only need to send audio to just the specified audio bus
 					uint32 BusId = SendInfo.AudioBus->GetUniqueID();
 					int32 BusChannels = (int32)SendInfo.AudioBus->AudioBusChannels + 1;
+					FString BusName;
+
+#if ENABLE_AUDIO_DEBUG
+					BusName = SendInfo.AudioBus->GetPathName();
+#endif // ENABLE_AUDIO_DEBUG
 
 					// Note we will be sending audio to both the specified source bus and the audio bus with the same send level
-					SetupBusSend(OutAudioBusSends, SendInfo, BusSendType, BusId, bEnableBusSends, BusChannels);
+					SetupBusSend(OutAudioBusSends, SendInfo, BusSendType, BusId, bEnableBusSends, BusChannels, BusName);
 				}
 			}
 		}
@@ -1671,13 +1695,25 @@ namespace Audio
 			{
 				if (NewSend.Value.SoundSourceBus)
 				{
-					MixerSourceVoice->SetAudioBusSendInfo(NewSend.Key, NewSend.Value.SoundSourceBus->GetUniqueID(), NewSend.Value.SendLevel);
+					FString BusName;
+
+#if ENABLE_AUDIO_DEBUG
+					BusName = NewSend.Value.SoundSourceBus->GetPathName();
+#endif // if ENABLE_AUDIO_DEBUG
+
+					MixerSourceVoice->SetAudioBusSendInfo(NewSend.Key, NewSend.Value.SoundSourceBus->GetUniqueID(), NewSend.Value.SendLevel, BusName);
 					bSendingAudioToBuses = true;
 				}
 
 				if (NewSend.Value.AudioBus)
 				{
-					MixerSourceVoice->SetAudioBusSendInfo(NewSend.Key, NewSend.Value.AudioBus->GetUniqueID(), NewSend.Value.SendLevel);
+					FString BusName;
+
+#if ENABLE_AUDIO_DEBUG
+					BusName = NewSend.Value.AudioBus->GetPathName();
+#endif // if ENABLE_AUDIO_DEBUG
+
+					MixerSourceVoice->SetAudioBusSendInfo(NewSend.Key, NewSend.Value.AudioBus->GetUniqueID(), NewSend.Value.SendLevel, BusName);
 					bSendingAudioToBuses = true;
 				}
 			}
@@ -1721,7 +1757,13 @@ namespace Audio
 				DynamicBusSendInfo.SendLevel = SendLevel;
 				DynamicBusSendInfo.bIsInit = false;
 
-				MixerSourceVoice->SetAudioBusSendInfo(DynamicBusSendInfo.BusSendType, DynamicBusSendInfo.BusId, SendLevel);
+				FString BusName;
+
+#if ENABLE_AUDIO_DEBUG
+				BusName = DynamicBusSendInfo.BusName;
+#endif // if ENABLE_AUDIO_DEBUG
+
+				MixerSourceVoice->SetAudioBusSendInfo(DynamicBusSendInfo.BusSendType, DynamicBusSendInfo.BusId, SendLevel, BusName);
 
 				bPreviousBusEnablement = WaveInstance->bEnableBusSends;
 			}
diff --git a/Engine/Source/Runtime/AudioMixer/Private/AudioMixerSource.h b/Engine/Source/Runtime/AudioMixer/Private/AudioMixerSource.h
--- a/Engine/Source/Runtime/AudioMixer/Private/AudioMixerSource.h
+++ b/Engine/Source/Runtime/AudioMixer/Private/AudioMixerSource.h
@@ -153,6 +153,9 @@ namespace Audio
 		{
 			float SendLevel = 0.0f;
 			uint32 BusId = 0;
+#if ENABLE_AUDIO_DEBUG
+			FString BusName;
+#endif // ENABLE_AUDIO_DEBUG
 			ESourceBusSendLevelControlMethod BusSendLevelControlMethod = ESourceBusSendLevelControlMethod::Manual;
 			EBusSendType BusSendType = EBusSendType::PreEffect;
 			float MinSendLevel = 0.0f;
diff --git a/Engine/Source/Runtime/AudioMixer/Private/AudioMixerSourceManager.cpp b/Engine/Source/Runtime/AudioMixer/Private/AudioMixerSourceManager.cpp
--- a/Engine/Source/Runtime/AudioMixer/Private/AudioMixerSourceManager.cpp
+++ b/Engine/Source/Runtime/AudioMixer/Private/AudioMixerSourceManager.cpp
@@ -1997,13 +1997,13 @@ namespace Audio
 		}, AUDIO_MIXER_THREAD_COMMAND_STRING("ClearSubmixSendInfo()"));
 	}
 
-	void FMixerSourceManager::SetBusSendInfo(const int32 SourceId, EBusSendType InAudioBusSendType, uint32 AudioBusId, float BusSendLevel)
+	void FMixerSourceManager::SetBusSendInfo(const int32 SourceId, EBusSendType InAudioBusSendType, uint32 AudioBusId, float BusSendLevel, const FString& InBusName)
 	{
 		AUDIO_MIXER_CHECK(SourceId < NumTotalSources);
 		AUDIO_MIXER_CHECK(GameThreadInfo.bIsBusy[SourceId]);
 		AUDIO_MIXER_CHECK_GAME_THREAD(MixerDevice);
 
-		AudioMixerThreadCommand([this, SourceId, InAudioBusSendType, AudioBusId, BusSendLevel]()
+		AudioMixerThreadCommand([this, SourceId, InAudioBusSendType, AudioBusId, BusSendLevel, InBusName]()
 		{
 			// Create mapping of source id to bus send level
 			FAudioBusSend BusSend;
diff --git a/Engine/Source/Runtime/AudioMixer/Private/AudioMixerSourceManager.h b/Engine/Source/Runtime/AudioMixer/Private/AudioMixerSourceManager.h
--- a/Engine/Source/Runtime/AudioMixer/Private/AudioMixerSourceManager.h
+++ b/Engine/Source/Runtime/AudioMixer/Private/AudioMixerSourceManager.h
@@ -119,6 +119,9 @@ namespace Audio
 		uint32 AudioBusId = INDEX_NONE;
 		float SendLevel = 0.0f;
 		int32 BusChannels = 0;
+#if ENABLE_AUDIO_DEBUG
+		FString AudioBusName;
+#endif // ENABLE_AUDIO_DEBUG
 	};
 
 	struct FMixerSourceVoiceInitParams
@@ -129,6 +132,9 @@ namespace Audio
 		TArray<FInitAudioBusSend> AudioBusSends[(int32)EBusSendType::Count];
 		uint32 AudioBusId = INDEX_NONE;
 		int32 AudioBusChannels = 0;
+#if ENABLE_AUDIO_DEBUG
+		FString AudioBusName;
+#endif // ENABLE_AUDIO_DEBUG
 		float SourceBusDuration = 0.0f;
 		uint32 SourceEffectChainId = INDEX_NONE;
 		TArray<FSourceEffectChainEntry> SourceEffectChain;
@@ -280,7 +286,7 @@ namespace Audio
 		void SetSubmixSendInfo(const int32 SourceId, const FMixerSourceSubmixSend& SubmixSend);
 		void ClearSubmixSendInfo(const int32 SourceId, const FMixerSourceSubmixSend& SubmixSend);
 
-		void SetBusSendInfo(const int32 SourceId, EBusSendType InAudioBusSendType, uint32 AudiobusId, float BusSendLevel);
+		void SetBusSendInfo(const int32 SourceId, EBusSendType InAudioBusSendType, uint32 AudiobusId, float BusSendLevel, const FString& InBusName);
 
 		void UpdateDeviceChannelCount(const int32 InNumOutputChannels);
 
diff --git a/Engine/Source/Runtime/AudioMixer/Private/AudioMixerSourceVoice.cpp b/Engine/Source/Runtime/AudioMixer/Private/AudioMixerSourceVoice.cpp
--- a/Engine/Source/Runtime/AudioMixer/Private/AudioMixerSourceVoice.cpp
+++ b/Engine/Source/Runtime/AudioMixer/Private/AudioMixerSourceVoice.cpp
@@ -444,7 +444,7 @@ namespace Audio
 	}
 
 
-	void FMixerSourceVoice::SetAudioBusSendInfo(EBusSendType InBusSendType, uint32 AudioBusId, float BusSendLevel)
+	void FMixerSourceVoice::SetAudioBusSendInfo(EBusSendType InBusSendType, uint32 AudioBusId, float BusSendLevel, const FString& InBusName)
 	{
 		AUDIO_MIXER_CHECK_GAME_THREAD(MixerDevice);
 
@@ -453,7 +453,7 @@ namespace Audio
 			BusSendLevel = 0.0f;
 		}
 
-		SourceManager->SetBusSendInfo(SourceId, InBusSendType, AudioBusId, BusSendLevel);
+		SourceManager->SetBusSendInfo(SourceId, InBusSendType, AudioBusId, BusSendLevel, InBusName);
 	}
 
 	bool FMixerSourceVoice::IsRenderingToSubmixes() const
diff --git a/Engine/Source/Runtime/AudioMixer/Private/AudioMixerSourceVoice.h b/Engine/Source/Runtime/AudioMixer/Private/AudioMixerSourceVoice.h
--- a/Engine/Source/Runtime/AudioMixer/Private/AudioMixerSourceVoice.h
+++ b/Engine/Source/Runtime/AudioMixer/Private/AudioMixerSourceVoice.h
@@ -150,7 +150,7 @@ namespace Audio
 		void SetEnablement(bool bInEnableBusSendRouting, bool bInEnableMainSubmixOutput, bool bInEnableSubmixSendRouting);
 
 		// Set the source bus send levels
-		void SetAudioBusSendInfo(EBusSendType InBusSendType, uint32 AudioBusId, float BusSendLevel);
+		void SetAudioBusSendInfo(EBusSendType InBusSendType, uint32 AudioBusId, float BusSendLevel, const FString& InBusName);
 
 	private:
 
