[AudioMixer] AudioMixerBus: Added envelope follower logic so it can be used by [REDACTED_PERSON] dashboard in Audio Insights standalone

diff --git a/Engine/Source/Runtime/AudioMixer/Private/AudioMixerBus.cpp b/Engine/Source/Runtime/AudioMixer/Private/AudioMixerBus.cpp
--- a/Engine/Source/Runtime/AudioMixer/Private/AudioMixerBus.cpp
+++ b/Engine/Source/Runtime/AudioMixer/Private/AudioMixerBus.cpp
@@ -145,6 +145,13 @@ namespace Audio
 
 		// Send the mix to the patch splitter's outputs
 		PatchSplitter.PushAudio(BusDataBufferPtr, NumSamples);
+
+#if UE_AUDIO_PROFILERTRACE_ENABLED
+		if (bIsEnvelopeFollowing)
+		{
+			ProcessEnvelopeFollower(BusDataBufferPtr);
+		}
+#endif // UE_AUDIO_PROFILERTRACE_ENABLED
 	}
 
 	void FMixerAudioBus::CopyCurrentBuffer(Audio::FAlignedFloatBuffer& InChannelMap, int32 InNumOutputChannels, FAlignedFloatBuffer& OutBuffer, int32 NumOutputFrames) const
@@ -186,4 +193,59 @@ namespace Audio
 	{
 		return PatchMixer.RemovePatch(PatchInput);
 	}
+
+#if UE_AUDIO_PROFILERTRACE_ENABLED
+	void FMixerAudioBus::StartEnvelopeFollower(const float InAttackTime, const float InReleaseTime, const float InSampleRate)
+	{
+		if (!bIsEnvelopeFollowing)
+		{
+			FEnvelopeFollowerInitParams EnvelopeFollowerInitParams;
+
+			EnvelopeFollowerInitParams.SampleRate = InSampleRate;
+			EnvelopeFollowerInitParams.NumChannels = NumChannels;
+			EnvelopeFollowerInitParams.AttackTimeMsec = InAttackTime;
+			EnvelopeFollowerInitParams.ReleaseTimeMsec = InReleaseTime;
+
+			EnvelopeFollower.Init(EnvelopeFollowerInitParams);
+
+			// Zero out any previous envelope values which may have been in the array before starting up
+			for (int32 ChannelIndex = 0; ChannelIndex < AUDIO_MIXER_MAX_OUTPUT_CHANNELS; ++ChannelIndex)
+			{
+				EnvelopeValues[ChannelIndex] = 0.0f;
+			}
+
+			bIsEnvelopeFollowing = true;
+		}
+	}
+
+	void FMixerAudioBus::StopEnvelopeFollower()
+	{
+		bIsEnvelopeFollowing = false;
+	}
+
+	void FMixerAudioBus::ProcessEnvelopeFollower(const float* InBuffer)
+	{
+		FMemory::Memset(EnvelopeValues, sizeof(float) * AUDIO_MIXER_MAX_OUTPUT_CHANNELS);
+
+		if (NumChannels > 0)
+		{
+			if (EnvelopeFollower.GetNumChannels() != NumChannels)
+			{
+				EnvelopeFollower.SetNumChannels(NumChannels);
+			}
+
+			EnvelopeFollower.ProcessAudio(InBuffer, NumFrames);
+
+			const TArray<float>& EnvValues = EnvelopeFollower.GetEnvelopeValues();
+
+			check(EnvValues.Num() == NumChannels);
+
+			FMemory::Memcpy(EnvelopeValues, EnvValues.GetData(), sizeof(float) * NumChannels);
+
+			Audio::ArrayClampInPlace(MakeArrayView(EnvelopeValues, NumChannels), 0.0f, 1.0f);
+		}
+
+		EnvelopeNumChannels = NumChannels;
+	}
+#endif // UE_AUDIO_PROFILERTRACE_ENABLED
 }
diff --git a/Engine/Source/Runtime/AudioMixer/Public/AudioMixerBus.h b/Engine/Source/Runtime/AudioMixer/Public/AudioMixerBus.h
--- a/Engine/Source/Runtime/AudioMixer/Public/AudioMixerBus.h
+++ b/Engine/Source/Runtime/AudioMixer/Public/AudioMixerBus.h
@@ -4,9 +4,14 @@
 
 #include "CoreMinimal.h"
 #include "Audio.h"
+#include "AudioMixerTrace.h"
 #include "DSP/MultithreadedPatching.h"
 #include "Sound/AudioSettings.h"
 
+#if UE_AUDIO_PROFILERTRACE_ENABLED
+#include "DSP/EnvelopeFollower.h"
+#endif // UE_AUDIO_PROFILERTRACE_ENABLED
+
 namespace Audio
 {
 	class FMixerSourceManager;
@@ -80,6 +85,12 @@ namespace Audio
 		// If this bus was constructed before
 		void SetNumOutputChannels(int32 InNumOutputChannels);
 
+#if UE_AUDIO_PROFILERTRACE_ENABLED
+		void StartEnvelopeFollower(const float InAttackTime, const float InReleaseTime, const float InSampleRate);
+		void StopEnvelopeFollower();
+		void ProcessEnvelopeFollower(const float* InBuffer);
+#endif // UE_AUDIO_PROFILERTRACE_ENABLED
+
 		// Array of instance ids. These are sources which are instances of this.
 		// It's possible for this data to have bus sends but no instance ids.
 		// This means a source would send its audio to the bus if the bus had an instance.
@@ -112,6 +123,13 @@ namespace Audio
 		// Was created manually, not via source buses.
 		bool bIsAutomatic;
 
+#if UE_AUDIO_PROFILERTRACE_ENABLED
+		Audio::FEnvelopeFollower EnvelopeFollower;
+		float EnvelopeValues[AUDIO_MIXER_MAX_OUTPUT_CHANNELS];
+		int32 EnvelopeNumChannels = 0;
+		bool bIsEnvelopeFollowing = false;
+#endif // UE_AUDIO_PROFILERTRACE_ENABLED
+
 		friend FMixerSourceManager;
 		friend FMixerSubmix;
 	};
