[AudioMixer] Added Audio Bus traces for Audio Insights standalone Audio Bus dashboard

diff --git a/Engine/Source/Runtime/AudioMixer/Private/AudioMixerSourceManager.cpp b/Engine/Source/Runtime/AudioMixer/Private/AudioMixerSourceManager.cpp
--- a/Engine/Source/Runtime/AudioMixer/Private/AudioMixerSourceManager.cpp
+++ b/Engine/Source/Runtime/AudioMixer/Private/AudioMixerSourceManager.cpp
@@ -275,6 +275,7 @@ UE_TRACE_EVENT_BEGIN(Audio, SoundSourceEnvelope)
 UE_TRACE_EVENT_END()
 
 // Audio Bus messages
+// @TODO UE-297630: remove this trace message once the full dashboard transition is finished (only editor -> editor and standalone)
 UE_TRACE_EVENT_BEGIN(Audio, AudioBusActivate)
 	UE_TRACE_EVENT_FIELD(uint32, DeviceId)
 	UE_TRACE_EVENT_FIELD(uint32, AudioBusId)
@@ -282,18 +283,48 @@ UE_TRACE_EVENT_BEGIN(Audio, AudioBusActivate)
 	UE_TRACE_EVENT_FIELD(UE::Trace::WideString, Name)
 UE_TRACE_EVENT_END()
 
+// @TODO UE-297630: remove this trace message once the full dashboard transition is finished (only editor -> editor and standalone)
 UE_TRACE_EVENT_BEGIN(Audio, AudioBusDeactivate)
 	UE_TRACE_EVENT_FIELD(uint32, DeviceId)
 	UE_TRACE_EVENT_FIELD(uint32, AudioBusId)
 	UE_TRACE_EVENT_FIELD(double, Timestamp)
 UE_TRACE_EVENT_END()
 
+UE_TRACE_EVENT_BEGIN(Audio, AudioBusStartOnTraceStart)
+	UE_TRACE_EVENT_FIELD(uint32, DeviceId)
+	UE_TRACE_EVENT_FIELD(uint32, AudioBusId)
+	UE_TRACE_EVENT_FIELD(double, Timestamp)
+	UE_TRACE_EVENT_FIELD(UE::Trace::WideString, Name)
+	UE_TRACE_EVENT_FIELD(int32, NumChannels)
+UE_TRACE_EVENT_END()
+
+UE_TRACE_EVENT_BEGIN(Audio, AudioBusStart)
+	UE_TRACE_EVENT_FIELD(uint32, DeviceId)
+	UE_TRACE_EVENT_FIELD(uint32, AudioBusId)
+	UE_TRACE_EVENT_FIELD(double, Timestamp)
+	UE_TRACE_EVENT_FIELD(UE::Trace::WideString, Name)
+	UE_TRACE_EVENT_FIELD(int32, NumChannels)
+UE_TRACE_EVENT_END()
+
+UE_TRACE_EVENT_BEGIN(Audio, AudioBusStop)
+	UE_TRACE_EVENT_FIELD(uint32, DeviceId)
+	UE_TRACE_EVENT_FIELD(uint32, AudioBusId)
+	UE_TRACE_EVENT_FIELD(double, Timestamp)
+UE_TRACE_EVENT_END()
+
 UE_TRACE_EVENT_BEGIN(Audio, AudioBusHasActivity)
 	UE_TRACE_EVENT_FIELD(uint32, DeviceId)
 	UE_TRACE_EVENT_FIELD(uint32, AudioBusId)
 	UE_TRACE_EVENT_FIELD(double, Timestamp)
 	UE_TRACE_EVENT_FIELD(bool, HasActivity)
 UE_TRACE_EVENT_END()
+
+UE_TRACE_EVENT_BEGIN(Audio, AudioBusEnvelopeValues)
+	UE_TRACE_EVENT_FIELD(uint32, DeviceId)
+	UE_TRACE_EVENT_FIELD(uint32, AudioBusId)
+	UE_TRACE_EVENT_FIELD(double, Timestamp)
+	UE_TRACE_EVENT_FIELD(float[], EnvelopeValues)
+UE_TRACE_EVENT_END()
 #endif // UE_AUDIO_PROFILERTRACE_ENABLED
 
 
@@ -355,6 +386,34 @@ namespace Audio
 		return bIsSilent;
 	}
 
+#if UE_AUDIO_PROFILERTRACE_ENABLED
+	void TraceAudioBusStartMessage(const uint32 InDeviceId, const uint32 InObjectId, const FString& InAudioBusName, const int32 InNumChannels)
+	{
+		const bool bChannelEnabled = UE_TRACE_CHANNELEXPR_IS_ENABLED(AudioMixerChannel);
+		if (bChannelEnabled)
+		{
+			UE_TRACE_LOG(Audio, AudioBusStart, AudioChannel)
+				<< AudioBusStart.DeviceId(InDeviceId)
+				<< AudioBusStart.AudioBusId(InObjectId)
+				<< AudioBusStart.Timestamp(FPlatformTime::Cycles64())
+				<< AudioBusStart.Name(InAudioBusName.IsEmpty() ? *FString::Printf(TEXT("%d"), InObjectId) : *InAudioBusName)
+				<< AudioBusStart.NumChannels(InNumChannels);
+		}
+	}
+
+	void TraceAudioBusStopMessage(const uint32 InDeviceId, const uint32 InObjectId)
+	{
+		const bool bChannelEnabled = UE_TRACE_CHANNELEXPR_IS_ENABLED(AudioMixerChannel);
+		if (bChannelEnabled)
+		{
+			UE_TRACE_LOG(Audio, AudioBusStop, AudioChannel)
+				<< AudioBusStop.DeviceId(InDeviceId)
+				<< AudioBusStop.AudioBusId(InObjectId)
+				<< AudioBusStop.Timestamp(FPlatformTime::Cycles64());
+		}
+	}
+#endif // UE_AUDIO_PROFILERTRACE_ENABLED
+
 	/*************************************************************************
 	* FMixerSourceManager
 	**************************************************************************/
@@ -383,9 +442,125 @@ namespace Audio
 
 	FMixerSourceManager::~FMixerSourceManager()
 	{
+#if UE_AUDIO_PROFILERTRACE_ENABLED
+		FTraceAuxiliary::OnTraceStarted.RemoveAll(this);
+		UnregisterAudioBusEnvelopeFollowerCVars();
+#endif // UE_AUDIO_PROFILERTRACE_ENABLED
+
 		FPlatformProcess::ReturnSynchEventToPool(CommandsProcessedEvent);
 	}
 
+#if UE_AUDIO_PROFILERTRACE_ENABLED
+	void FMixerSourceManager::OnTraceStarted(FTraceAuxiliary::EConnectionType TraceType, const FString& TraceDestination)
+	{
+		for (const auto& [AudioBusKey, AudioBus] : AudioBuses)
+		{
+			if (AudioBus.IsValid())
+			{
+ 				TraceAudioBusStartMessage(MixerDevice->DeviceID, AudioBusKey.ObjectId, AudioBus->BusName, AudioBus->GetNumChannels());
+			}
+		}
+	}
+
+	void FMixerSourceManager::RegisterAudioBusEnvelopeFollowerCVars()
+	{
+		// StartAudioBusEnvelopeFollowerCVar
+		StartAudioBusEnvelopeFollowerCVar = IConsoleManager::Get().RegisterConsoleVariable(
+			*FString::Printf(TEXT("au.StartAudioBusEnvelopeFollower_AD%d"), MixerDevice->DeviceID),
+			FString(),
+			TEXT("Starts an audio bus envelope follower (args= audio bus id)\n"),
+			ECVF_Default
+		);
+
+		if (StartAudioBusEnvelopeFollowerCVar)
+		{
+			StartAudioBusEnvelopeFollowerCVar->SetOnChangedCallback(FConsoleVariableDelegate::CreateRaw(this, &FMixerSourceManager::OnStartAudioBusEnvelopeFollowerCVarChanged));
+		}
+
+		// StopAudioBusEnvelopeFollowerCVar
+		StopAudioBusEnvelopeFollowerCVar = IConsoleManager::Get().RegisterConsoleVariable(
+			*FString::Printf(TEXT("au.StopAudioBusEnvelopeFollower_AD%d"), MixerDevice->DeviceID),
+			FString(),
+			TEXT("Stops an audio bus envelope follower (args= audio bus id)\n"),
+			ECVF_Default
+		);
+
+		if (StopAudioBusEnvelopeFollowerCVar)
+		{
+			StopAudioBusEnvelopeFollowerCVar->SetOnChangedCallback(FConsoleVariableDelegate::CreateRaw(this, &FMixerSourceManager::OnStopAudioBusEnvelopeFollowerCVarChanged));
+		}
+	}
+
+	void FMixerSourceManager::UnregisterAudioBusEnvelopeFollowerCVars()
+	{
+		IConsoleManager::Get().UnregisterConsoleObject(StartAudioBusEnvelopeFollowerCVar);
+		IConsoleManager::Get().UnregisterConsoleObject(StopAudioBusEnvelopeFollowerCVar);
+	}
+
+	void FMixerSourceManager::OnStartAudioBusEnvelopeFollowerCVarChanged(IConsoleVariable* InStartAudioBusEnvelopeFollowerCVar)
+	{
+		ensure(IsInGameThread());
+
+		if (InStartAudioBusEnvelopeFollowerCVar)
+		{
+			uint32 AudioBusId = INDEX_NONE;
+			LexFromString(AudioBusId, *InStartAudioBusEnvelopeFollowerCVar->GetString());
+
+			FAudioThread::RunCommandOnAudioThread([this, AudioBusId]()
+			{
+				StartAudioBusEnvelopeFollowerFromCVar(AudioBusId);
+			});
+		}
+	}
+
+	void FMixerSourceManager::OnStopAudioBusEnvelopeFollowerCVarChanged(IConsoleVariable* InStopAudioBusEnvelopeFollowerCVar)
+	{
+		ensure(IsInGameThread());
+
+		if (InStopAudioBusEnvelopeFollowerCVar)
+		{
+			uint32 AudioBusId = INDEX_NONE;
+			LexFromString(AudioBusId, *InStopAudioBusEnvelopeFollowerCVar->GetString());
+
+			FAudioThread::RunCommandOnAudioThread([this, AudioBusId]()
+			{
+				StopAudioBusEnvelopeFollowerFromCVar(AudioBusId);
+			});
+		}
+	}
+
+	void FMixerSourceManager::StartAudioBusEnvelopeFollowerFromCVar(const uint32 InAudioBusId)
+	{
+		ensure(IsInAudioThread());
+
+		constexpr float AttackTime  = 10.0f;
+		constexpr float ReleaseTime = 500.0f;
+
+		for (auto& [AudioBusKey, AudioBus] : AudioBuses)
+		{
+			if (AudioBusKey.ObjectId == InAudioBusId && AudioBus.IsValid())
+			{
+				AudioBus->StartEnvelopeFollower(AttackTime, ReleaseTime, MixerDevice->GetSampleRate());
+				break;
+			}
+		}
+	}
+
+	void FMixerSourceManager::StopAudioBusEnvelopeFollowerFromCVar(const uint32 InAudioBusId)
+	{
+		ensure(IsInAudioThread());
+
+		for (auto& [AudioBusKey, AudioBus] : AudioBuses)
+		{
+			if (AudioBusKey.ObjectId == InAudioBusId && AudioBus.IsValid())
+			{
+				AudioBus->StopEnvelopeFollower();
+				break;
+			}
+		}
+	}
+#endif // UE_AUDIO_PROFILERTRACE_ENABLED
+
 	void FMixerSourceManager::Init(const FSourceManagerInitParams& InitParams)
 	{
 		AUDIO_MIXER_CHECK(InitParams.NumSources > 0);
@@ -582,6 +757,11 @@ namespace Audio
 
 		bInitialized = true;
 		bPumpQueue = false;
+
+#if UE_AUDIO_PROFILERTRACE_ENABLED
+		RegisterAudioBusEnvelopeFollowerCVars();
+		FTraceAuxiliary::OnTraceStarted.AddRaw(this, &FMixerSourceManager::OnTraceStarted);
+#endif // UE_AUDIO_PROFILERTRACE_ENABLED
 	}
 
 	void FMixerSourceManager::Update(bool bTimedOut)
@@ -695,6 +875,10 @@ namespace Audio
 				// If this audio bus was automatically created via source bus playback, this this audio bus can be removed
 				if (AudioBusPtr->RemoveInstanceId(SourceId))
 				{
+#if UE_AUDIO_PROFILERTRACE_ENABLED
+					TraceAudioBusStopMessage(MixerDevice->DeviceID, SourceInfo.AudioBusId);
+#endif // UE_AUDIO_PROFILERTRACE_ENABLED
+
 					// Only automatic buses will be getting removed here. Otherwise they need to be manually removed from the source manager.
 					ensure(AudioBusPtr->IsAutomatic());
 					AudioBuses.Remove(SourceInfo.AudioBusId);
@@ -713,6 +897,10 @@ namespace Audio
 				{
 					if (AudioBusPtr->RemoveSend((EBusSendType)AudioBusSendType, SourceId))
 					{
+#if UE_AUDIO_PROFILERTRACE_ENABLED
+						TraceAudioBusStopMessage(MixerDevice->DeviceID, AudioBusId);
+#endif // UE_AUDIO_PROFILERTRACE_ENABLED
+
 						ensure(AudioBusPtr->IsAutomatic());
 						AudioBuses.Remove(AudioBusId);
 					}
@@ -1158,6 +1346,11 @@ namespace Audio
 						TSharedPtr<FMixerAudioBus> NewAudioBus = TSharedPtr<FMixerAudioBus>(new FMixerAudioBus(this, true, InitParams.AudioBusChannels));
 						NewAudioBus->AddInstanceId(SourceId, InitParams.NumInputChannels);
 
+#if UE_AUDIO_PROFILERTRACE_ENABLED
+						NewAudioBus->BusName = InitParams.AudioBusName;
+						TraceAudioBusStartMessage(MixerDevice->DeviceID, InitParams.AudioBusId, InitParams.AudioBusName, InitParams.AudioBusChannels);
+#endif // UE_AUDIO_PROFILERTRACE_ENABLED
+
 						AudioBuses.Add(InitParams.AudioBusId, NewAudioBus);
 					}
 				}
@@ -1191,6 +1384,11 @@ namespace Audio
 						// source's audio to it.
 						NewAudioBus->AddSend((EBusSendType)BusSendType, NewAudioBusSend);
 
+#if UE_AUDIO_PROFILERTRACE_ENABLED
+						NewAudioBus->BusName = AudioBusSend.AudioBusName;
+						TraceAudioBusStartMessage(MixerDevice->DeviceID, AudioBusSend.AudioBusId, AudioBusSend.AudioBusName, AudioBusSend.BusChannels);
+#endif // UE_AUDIO_PROFILERTRACE_ENABLED
+
 						AudioBuses.Add(AudioBusSend.AudioBusId, NewAudioBus);
 					}
 
@@ -1319,14 +1517,17 @@ namespace Audio
 		AudioBusKeys_AudioThread.Add(InAudioBusKey);
 
 #if UE_AUDIO_PROFILERTRACE_ENABLED
+		// @TODO UE-297630: remove this trace message once the full dashboard transition is finished (only editor -> editor and standalone)
 		UE_TRACE_LOG(Audio, AudioBusActivate, AudioChannel)
 			<< AudioBusActivate.DeviceId(MixerDevice->DeviceID)
 			<< AudioBusActivate.AudioBusId(InAudioBusKey.ObjectId)
 			<< AudioBusActivate.Timestamp(FPlatformTime::Cycles64())
 			<< AudioBusActivate.Name(*InAudioBusName);
+
+		TraceAudioBusStartMessage(MixerDevice->DeviceID, InAudioBusKey.ObjectId, InAudioBusName, InNumChannels);
 #endif
 
-		AudioMixerThreadCommand([this, InAudioBusKey, InNumChannels, bInIsAutomatic]()
+		AudioMixerThreadCommand([this, InAudioBusKey, InNumChannels, bInIsAutomatic, InAudioBusName]()
 		{
 			// If this audio bus id already exists, set it to not be automatic and return it
 			TSharedPtr<FMixerAudioBus> AudioBusPtr = AudioBuses.FindRef(InAudioBusKey);
@@ -1340,7 +1541,9 @@ namespace Audio
 			{
 				// If the bus is not registered, make a new entry.
 				TSharedPtr<FMixerAudioBus> NewBusData(new FMixerAudioBus(this, bInIsAutomatic, InNumChannels));
-
+#if UE_AUDIO_PROFILERTRACE_ENABLED
+				NewBusData->BusName = InAudioBusName;
+#endif
 				AudioBuses.Add(InAudioBusKey, NewBusData);
 			}
 		}, AUDIO_MIXER_THREAD_COMMAND_STRING("StartAudioBus()"));
@@ -1356,10 +1559,13 @@ namespace Audio
 		AudioBusKeys_AudioThread.Remove(InAudioBusKey);
 
 #if UE_AUDIO_PROFILERTRACE_ENABLED
+		// @TODO UE-297630: remove this trace message once the full dashboard transition is finished (only editor -> editor and standalone)
 		UE_TRACE_LOG(Audio, AudioBusDeactivate, AudioChannel)
 			<< AudioBusDeactivate.DeviceId(MixerDevice->DeviceID)
 			<< AudioBusDeactivate.AudioBusId(InAudioBusKey.ObjectId)
 			<< AudioBusDeactivate.Timestamp(FPlatformTime::Cycles64());
+
+		TraceAudioBusStopMessage(MixerDevice->DeviceID, InAudioBusKey.ObjectId);
 #endif
 
 		AudioMixerThreadCommand([this, InAudioBusKey]()
@@ -2030,6 +2236,11 @@ namespace Audio
 				// source's audio to it.
 				NewBusData->AddSend(InAudioBusSendType, BusSend);
 
+#if UE_AUDIO_PROFILERTRACE_ENABLED
+				NewBusData->BusName = InBusName;
+				TraceAudioBusStartMessage(MixerDevice->DeviceID, AudioBusId, InBusName, SourceInfo.NumInputChannels);
+#endif // UE_AUDIO_PROFILERTRACE_ENABLED
+
 				AudioBuses.Add(AudioBusId, NewBusData);
 			}
 
@@ -2590,11 +2801,30 @@ namespace Audio
 
 		ConnectBusPatches();
 
+#if UE_AUDIO_PROFILERTRACE_ENABLED
+		const bool bChannelEnabled = UE_TRACE_CHANNELEXPR_IS_ENABLED(AudioChannel);
+#endif // UE_AUDIO_PROFILERTRACE_ENABLED
+
 		// Loop through the bus registry and mix source audio
-		for (auto& Entry : AudioBuses)
+		for (auto& [AudioBusKey, AudioBus] : AudioBuses)
 		{
-			TSharedPtr<FMixerAudioBus>& AudioBus = Entry.Value;
+			if (!AudioBus.IsValid())
+			{
+				continue;
+			}
+
 			AudioBus->MixBuffer();
+
+#if UE_AUDIO_PROFILERTRACE_ENABLED
+			if (bChannelEnabled && AudioBus->bIsEnvelopeFollowing)
+			{
+				UE_TRACE_LOG(Audio, AudioBusEnvelopeValues, AudioChannel)
+					<< AudioBusEnvelopeValues.DeviceId(MixerDevice->DeviceID)
+					<< AudioBusEnvelopeValues.AudioBusId(AudioBusKey.ObjectId)
+					<< AudioBusEnvelopeValues.Timestamp(FPlatformTime::Cycles64())
+					<< AudioBusEnvelopeValues.EnvelopeValues(AudioBus->EnvelopeValues, AudioBus->GetNumChannels());
+			}
+#endif // UE_AUDIO_PROFILERTRACE_ENABLED
 		}
 	}
 
diff --git a/Engine/Source/Runtime/AudioMixer/Private/AudioMixerSourceManager.h b/Engine/Source/Runtime/AudioMixer/Private/AudioMixerSourceManager.h
--- a/Engine/Source/Runtime/AudioMixer/Private/AudioMixerSourceManager.h
+++ b/Engine/Source/Runtime/AudioMixer/Private/AudioMixerSourceManager.h
@@ -22,6 +22,10 @@
 #include "Sound/QuartzQuantizationUtilities.h"
 #include "Stats/Stats.h"
 
+#if UE_AUDIO_PROFILERTRACE_ENABLED
+#include "ProfilingDebugging/TraceAuxiliary.h"
+#endif // UE_AUDIO_PROFILERTRACE_ENABLED
+
 #include "AudioMixerSourceManager.generated.h"
 
 // Default this to on (it's quite a small memory footprint).
@@ -373,6 +377,22 @@ namespace Audio
 
 		void AudioMixerThreadCommand(TFunction<void()>&& InFunction, const char* InDebugString = nullptr, bool bInDeferExecution = false);
 
+#if UE_AUDIO_PROFILERTRACE_ENABLED
+		void OnTraceStarted(FTraceAuxiliary::EConnectionType TraceType, const FString& TraceDestination);
+
+		void RegisterAudioBusEnvelopeFollowerCVars();
+		void UnregisterAudioBusEnvelopeFollowerCVars();
+
+		void OnStartAudioBusEnvelopeFollowerCVarChanged(IConsoleVariable* InStartAudioBusEnvelopeFollowerCVar);
+		void OnStopAudioBusEnvelopeFollowerCVarChanged(IConsoleVariable* InStopAudioBusEnvelopeFollowerCVar);
+
+		void StartAudioBusEnvelopeFollowerFromCVar(const uint32 InAudioBusId);
+		void StopAudioBusEnvelopeFollowerFromCVar(const uint32 InAudioBusId);
+
+		IConsoleVariable* StartAudioBusEnvelopeFollowerCVar = nullptr;
+		IConsoleVariable* StopAudioBusEnvelopeFollowerCVar  = nullptr;
+#endif // UE_AUDIO_PROFILERTRACE_ENABLED
+
 		// Critical section to ensure mutating effect chains is thread-safe
 		FCriticalSection EffectChainMutationCriticalSection;
 
diff --git a/Engine/Source/Runtime/AudioMixer/Public/AudioMixerBus.h b/Engine/Source/Runtime/AudioMixer/Public/AudioMixerBus.h
--- a/Engine/Source/Runtime/AudioMixer/Public/AudioMixerBus.h
+++ b/Engine/Source/Runtime/AudioMixer/Public/AudioMixerBus.h
@@ -123,6 +123,10 @@ namespace Audio
 		// Was created manually, not via source buses.
 		bool bIsAutomatic;
 
+#if ENABLE_AUDIO_DEBUG
+		FString BusName;
+#endif // if ENABLE_AUDIO_DEBUG
+
 #if UE_AUDIO_PROFILERTRACE_ENABLED
 		Audio::FEnvelopeFollower EnvelopeFollower;
 		float EnvelopeValues[AUDIO_MIXER_MAX_OUTPUT_CHANNELS];
