Improvements to xbox resume logic.

#rnx

diff --git a/Engine/Source/Runtime/AudioDeviceEnumeration/Windows/WindowsMMDeviceEnumeration/Private/WindowsMMDeviceInfoCache.cpp b/Engine/Source/Runtime/AudioDeviceEnumeration/Windows/WindowsMMDeviceEnumeration/Private/WindowsMMDeviceInfoCache.cpp
--- a/Engine/Source/Runtime/AudioDeviceEnumeration/Windows/WindowsMMDeviceEnumeration/Private/WindowsMMDeviceInfoCache.cpp
+++ b/Engine/Source/Runtime/AudioDeviceEnumeration/Windows/WindowsMMDeviceEnumeration/Private/WindowsMMDeviceInfoCache.cpp
@@ -88,8 +88,7 @@ namespace Audio
 	{
 		ensure(SUCCEEDED(CoCreateInstance(__uuidof(MMDeviceEnumerator), NULL, CLSCTX_INPROC_SERVER, IID_PPV_ARGS(&DeviceEnumerator))) && DeviceEnumerator);
 		
-		EnumerateEndpoints();
-		EnumerateDefaults();
+		Refresh();
 	}
 
 	FWindowsMMDeviceCache::FWindowsMMDeviceCache() :
@@ -97,6 +96,11 @@ namespace Audio
 	{
 	}
 
+	void FWindowsMMDeviceCache::Refresh()
+	{
+		EnumerateEndpoints();
+		EnumerateDefaults();
+	}
 
 	bool FWindowsMMDeviceCache::EnumerateChannelMask(uint32 InMask, FCacheEntry& OutInfo)
 	{
diff --git a/Engine/Source/Runtime/AudioDeviceEnumeration/Windows/WindowsMMDeviceEnumeration/Public/WindowsMMDeviceInfoCache.h b/Engine/Source/Runtime/AudioDeviceEnumeration/Windows/WindowsMMDeviceEnumeration/Public/WindowsMMDeviceInfoCache.h
--- a/Engine/Source/Runtime/AudioDeviceEnumeration/Windows/WindowsMMDeviceEnumeration/Public/WindowsMMDeviceInfoCache.h
+++ b/Engine/Source/Runtime/AudioDeviceEnumeration/Windows/WindowsMMDeviceEnumeration/Public/WindowsMMDeviceInfoCache.h
@@ -107,6 +107,8 @@ struct FWindowsMMDeviceCache : IAudioMixerDeviceChangedListener, IAudioPlatformD
 	UE_API virtual TArray<FAudioPlatformDeviceInfo> GetLogicalAggregateDevices(const FName InHardwareId, const EDeviceEndpointType InEndpointType) const override;
 
 	UE_API virtual TArray<FAudioPlatformDeviceInfo> SynthesizeAggregateDeviceList(const EDeviceEndpointType InType) const;
+
+	UE_API virtual void Refresh() override;
 	
 	UE_API FName GetDefaultOutputDevice_NoLock() const;
 
diff --git a/Engine/Source/Runtime/AudioMixer/Private/AudioMixerDevice.cpp b/Engine/Source/Runtime/AudioMixer/Private/AudioMixerDevice.cpp
--- a/Engine/Source/Runtime/AudioMixer/Private/AudioMixerDevice.cpp
+++ b/Engine/Source/Runtime/AudioMixer/Private/AudioMixerDevice.cpp
@@ -947,44 +947,51 @@ namespace Audio
 
 		AudioMixerPlatform->OnHardwareUpdate();
 
-		if (AudioMixerPlatform->CheckAudioDeviceChange())
+		// Device swaps span multiple calls beginning with CheckAudioDeviceChange() and ending with
+		// ResumePlaybackOnNewDevice(). We protect the associated resources by locking and unlocking around
+		// the entire call chain.
 		{
-			// Get the platform device info we're using
-			PlatformInfo = AudioMixerPlatform->GetPlatformDeviceInfo();
+			FScopedDeviceSwapLock Lock(AudioMixerPlatform);
 
-			// Initialize some data that depends on speaker configuration, etc.
-			InitializeChannelAzimuthMap(PlatformInfo.NumChannels);
+			if (AudioMixerPlatform->CheckAudioDeviceChange())
+			{
+				// Get the platform device info we're using
+				PlatformInfo = AudioMixerPlatform->GetPlatformDeviceInfo();
 
-			// Update the channel device count in case it changed
-			SourceManager->UpdateDeviceChannelCount(PlatformInfo.NumChannels);
+				// Initialize some data that depends on speaker configuration, etc.
+				InitializeChannelAzimuthMap(PlatformInfo.NumChannels);
 
-			// Reset rendering thread ID to this thread ID so that commands can
-			// be flushed. Audio Rendering Thread ID will be reset again in call
-			// to FMixerDevice::OnProcessAudio
-			ResetAudioRenderingThreadId();
-			
-			// Cache the audio platform thread id for debugging purposes. This is
-			// an attempt to narrow down the causes of UE-209237. The theory is
-			// that the there are multiple threads attempting to run audio rendering
-			// commands. To catch the issue we attempt to cache the thread id before
-			// running audio rendering commands. We then check that the thread id
-			// has not changed for the duration that render thread commands have 
-			// run.
-			std::atomic<int32> CurrentAudioPlatformThreadId = AudioPlatformThreadId.load();
+				// Update the channel device count in case it changed
+				SourceManager->UpdateDeviceChannelCount(PlatformInfo.NumChannels);
 
-			// Force source manager to incorporate device channel count change.
-			FlushAudioRenderingCommands(true /* bPumpSynchronously */);
+				// Reset rendering thread ID to this thread ID so that commands can
+				// be flushed. Audio Rendering Thread ID will be reset again in call
+				// to FMixerDevice::OnProcessAudio
+				ResetAudioRenderingThreadId();
 
-			if (UAudioDeviceNotificationSubsystem* AudioDeviceNotifSubsystem = UAudioDeviceNotificationSubsystem::Get())
-			{
-				AudioDeviceNotifSubsystem->OnDeviceSwitched(PlatformInfo.DeviceId);
-			}
+				// Cache the audio platform thread id for debugging purposes. This is
+				// an attempt to narrow down the causes of UE-209237. The theory is
+				// that the there are multiple threads attempting to run audio rendering
+				// commands. To catch the issue we attempt to cache the thread id before
+				// running audio rendering commands. We then check that the thread id
+				// has not changed for the duration that render thread commands have 
+				// run.
+				std::atomic<int32> CurrentAudioPlatformThreadId = AudioPlatformThreadId.load();
+
+				// Force source manager to incorporate device channel count change.
+				FlushAudioRenderingCommands(true /* bPumpSynchronously */);
+
+				if (UAudioDeviceNotificationSubsystem* AudioDeviceNotifSubsystem = UAudioDeviceNotificationSubsystem::Get())
+				{
+					AudioDeviceNotifSubsystem->OnDeviceSwitched(PlatformInfo.DeviceId);
+				}
 
-			// Related to earlier mention of UE-209237
-			UE_CLOG(CurrentAudioPlatformThreadId != AudioPlatformThreadId, LogAudioMixer, Error, TEXT("Platform audio thread id changed while flushing render commands. Expected %d, found %d. May result in corrupt internal audio source state."), CurrentAudioPlatformThreadId.load(), AudioPlatformThreadId.load());
+				// Related to earlier mention of UE-209237
+				UE_CLOG(CurrentAudioPlatformThreadId != AudioPlatformThreadId, LogAudioMixer, Error, TEXT("Platform audio thread id changed while flushing render commands. Expected %d, found %d. May result in corrupt internal audio source state."), CurrentAudioPlatformThreadId.load(), AudioPlatformThreadId.load());
 
-			// Audio rendering was suspended in CheckAudioDeviceChange if it changed.
-			AudioMixerPlatform->ResumePlaybackOnNewDevice();
+				// Audio rendering was suspended in CheckAudioDeviceChange if it changed.
+				AudioMixerPlatform->ResumePlaybackOnNewDevice();
+			}
 		}
 
 		// Device must be initialized prior to call as submix graph may not be ready yet otherwise.
diff --git a/Engine/Source/Runtime/AudioMixerCore/Private/AudioMixer.cpp b/Engine/Source/Runtime/AudioMixerCore/Private/AudioMixer.cpp
--- a/Engine/Source/Runtime/AudioMixerCore/Private/AudioMixer.cpp
+++ b/Engine/Source/Runtime/AudioMixerCore/Private/AudioMixer.cpp
@@ -723,13 +723,11 @@ namespace Audio
 			{
 				check(AudioStreamInfo.StreamState == EAudioOutputStreamState::Stopped);
 			}
-			else
-			{
-				AudioStreamInfo.StreamState = EAudioOutputStreamState::Stopped;
-			}
 
 			AudioRenderThread.Reset();
 		}
+		
+		AudioStreamInfo.StreamState = EAudioOutputStreamState::Stopped;
 
 		if (AudioRenderEvent != nullptr)
 		{
@@ -954,35 +952,39 @@ namespace Audio
 
 	bool FAudioMixerPlatformSwappable::RequestDeviceSwap(const FString& DeviceID, const bool bInForce, const TCHAR* InReason)
 	{
-		if (!AllowDeviceSwap(bInForce))
-		{
-			UE_LOG(LogAudioMixer, Display, TEXT("NOT-ALLOWING attempt to swap audio render device to new device: '%s', because: '%s', force=%d"),
-				!DeviceID.IsEmpty() ? *DeviceID : TEXT("[System Default]"),
-				InReason ? InReason : TEXT("None specified"),
-				(int32)bInForce
-			);
-			return false;
-		}
+		bool bRequestSucceeded = false;
 
 		// This critical section protects that device swap context below.
 		if (DeviceSwapCriticalSection.TryLock())
 		{
-			UE_LOG(LogAudioMixer, Display, TEXT("Attempt to swap audio render device to new device: '%s', because: '%s', force=%d"),
-				!DeviceID.IsEmpty() ? *DeviceID : TEXT("[System Default]"),
-				InReason ? InReason : TEXT("None specified"),
-				(int32)bInForce
-			);
-
-			// Set up the context for this device swap
-			if (InitializeDeviceSwapContext(DeviceID, InReason))
+			if (AllowDeviceSwap(bInForce))
 			{
-				// Set the flag indicating we wish to begin a device swap
-				AudioStreamInfo.StreamState = EAudioOutputStreamState::SwappingDevice;
+				UE_LOG(LogAudioMixer, Display, TEXT("Attempt to swap audio render device to new device: '%s', because: '%s', force=%d"),
+					!DeviceID.IsEmpty() ? *DeviceID : TEXT("[System Default]"),
+					InReason ? InReason : TEXT("None specified"),
+					(int32)bInForce
+				);
+
+				// Set up the context for this device swap
+				if (InitializeDeviceSwapContext(DeviceID, InReason))
+				{
+					// Set the flag indicating we wish to begin a device swap
+					AudioStreamInfo.StreamState = EAudioOutputStreamState::SwappingDevice;
+					bRequestSucceeded = true;
+				}
+			}
+			else
+			{
+				UE_LOG(LogAudioMixer, Display, TEXT("NOT-ALLOWING attempt to swap audio render device to new device: '%s', because: '%s', force=%d [normal if occurring while suspended on consoles]"),
+					!DeviceID.IsEmpty() ? *DeviceID : TEXT("[System Default]"),
+					InReason ? InReason : TEXT("None specified"),
+					(int32)bInForce
+				);
 			}
 		
 			DeviceSwapCriticalSection.Unlock();
 
-			return true;
+			return bRequestSucceeded;
 		}
 		
 		UE_LOG(LogAudioMixer, Display, TEXT("Dropping device swap request; device: '%s', because: '%s', force=%d"),
@@ -990,7 +992,7 @@ namespace Audio
 		       InReason ? InReason : TEXT("None specified"),
 		       (int32)bInForce);
 
-		return false;
+		return bRequestSucceeded;
 	}
 
 	bool FAudioMixerPlatformSwappable::CheckAudioDeviceChange()
@@ -1024,15 +1026,16 @@ namespace Audio
 			return false;
 		}
 		
-		// If the device is not running or if a device swap is already in progress, don't "double-trigger" a swap.
-		if (ActiveDeviceSwap.IsValid())
+		// If the device is not running or if a device swap is already in progress, don't "double-trigger" a swap unless 
+		// bInForceSwap is true, in which case we'll replace the existing request.
+		if (ActiveDeviceSwap.IsValid() && !bInForceSwap)
 		{
-			UE_LOG(LogAudioMixer, Display, TEXT("Ignoring device swap request, swap already requested."));
+			UE_LOG(LogAudioMixer, Display, TEXT("Ignoring device swap request, swap already requested. bInForceSwap: %d"), bInForceSwap);
 			return false;
 		}
 
-		// Only device swap if the audio device is running
-		if (AudioStreamInfo.StreamState != EAudioOutputStreamState::Running)
+		// Only device swap if the audio device is running or we're mid-swap
+		if (AudioStreamInfo.StreamState != EAudioOutputStreamState::Running && AudioStreamInfo.StreamState != EAudioOutputStreamState::SwappingDevice)
 		{
 			UE_LOG(LogAudioMixer, Display, TEXT("Ignoring device swap request, AudioStreamInfo.StreamState: %d"), AudioStreamInfo.StreamState.load());
 			return false;
diff --git a/Engine/Source/Runtime/AudioMixerCore/Public/AudioMixer.h b/Engine/Source/Runtime/AudioMixerCore/Public/AudioMixer.h
--- a/Engine/Source/Runtime/AudioMixerCore/Public/AudioMixer.h
+++ b/Engine/Source/Runtime/AudioMixerCore/Public/AudioMixer.h
@@ -243,6 +243,9 @@ namespace Audio
 
 		/** Returns an array of logical audio devices which, when combined, form an aggregate device. */
 		virtual TArray<FAudioPlatformDeviceInfo> GetLogicalAggregateDevices(const FName InHardwareId, const EDeviceEndpointType InEndpointType) const = 0;
+
+		/** Refresh the data in the cache */
+		virtual void Refresh() = 0;
 	};
 
 	/** Defines parameters needed for opening a new audio stream to device. */
@@ -545,6 +548,12 @@ namespace Audio
 		/** Sends a command to swap which output device is being used */
 		virtual bool RequestDeviceSwap(const FString& DeviceID, bool bInForce, const TCHAR* InReason = nullptr) { return false; }
 
+		/** Protects device swap related resources during device swaps */
+		virtual void LockDeviceSwap() {}
+
+		/** Unprotects device swap resources. Balance with LockDeviceSwap() */
+		virtual void UnlockDeviceSwap() {}
+
 		/** Returns the platform device info of the currently open audio stream. */
 		virtual FAudioPlatformDeviceInfo GetPlatformDeviceInfo() const = 0;
 
@@ -805,6 +814,8 @@ namespace Audio
 
 		//~ Begin IAudioMixerPlatformInterface
 		UE_API virtual bool RequestDeviceSwap(const FString& DeviceID, const bool bInForce, const TCHAR* InReason) override;
+		UE_API virtual void LockDeviceSwap() override { DeviceSwapCriticalSection.Lock(); }
+		UE_API virtual void UnlockDeviceSwap() override { DeviceSwapCriticalSection.Unlock(); }
 		UE_API virtual bool CheckAudioDeviceChange() override;
 		UE_API virtual bool MoveAudioStreamToNewAudioDevice() override;
 		UE_API virtual void ResumePlaybackOnNewDevice() override;
@@ -857,6 +868,30 @@ namespace Audio
 		double LastDeviceSwapTime = 0.0;
 	};
 
+	class FScopedDeviceSwapLock final
+	{
+	public:
+		FScopedDeviceSwapLock() = delete;
+		explicit FScopedDeviceSwapLock(IAudioMixerPlatformInterface* InPlatform)
+			: AudioMixerPlatform(InPlatform)
+		{
+			if (AudioMixerPlatform)
+			{
+				AudioMixerPlatform->LockDeviceSwap();
+			}
+		}
+
+		~FScopedDeviceSwapLock()
+		{
+			if (AudioMixerPlatform)
+			{
+				AudioMixerPlatform->UnlockDeviceSwap();
+			}
+		}
+
+	private:
+		IAudioMixerPlatformInterface* AudioMixerPlatform = nullptr;
+	};
 }
 
 /**
diff --git a/Engine/Source/Runtime/Windows/AudioMixerWasapi/Private/AudioMixerWasapi.cpp b/Engine/Source/Runtime/Windows/AudioMixerWasapi/Private/AudioMixerWasapi.cpp
--- a/Engine/Source/Runtime/Windows/AudioMixerWasapi/Private/AudioMixerWasapi.cpp
+++ b/Engine/Source/Runtime/Windows/AudioMixerWasapi/Private/AudioMixerWasapi.cpp
@@ -503,11 +503,11 @@ namespace Audio
 		
 		if (DeviceSwapContext.IsValid())
 		{
-			UE_LOG(LogAudioMixer, Display, TEXT("FAudioMixerWasapi::InitDeviceSwapContextInternal - DeviceSwapContext in-flight, ignoring"));
-			return false;
+			UE_LOG(LogAudioMixer, Display, TEXT("FAudioMixerWasapi::InitDeviceSwapContextInternal - DeviceSwapContext in-flight, replacing"));
 		}
 
-		// Create the device swap context which will be valid over the course of the swap
+		// Create the device swap context which will be valid over the course of the 
+		DeviceSwapContext.Reset();
 		DeviceSwapContext = MakeUnique<FWasapiDeviceSwapContext>(InRequestedDeviceID, InReason);
 		if (!DeviceSwapContext.IsValid())
 		{
diff --git a/Engine/Source/Runtime/Windows/AudioMixerWasapi/Private/AudioMixerWasapiRenderStream.cpp b/Engine/Source/Runtime/Windows/AudioMixerWasapi/Private/AudioMixerWasapiRenderStream.cpp
--- a/Engine/Source/Runtime/Windows/AudioMixerWasapi/Private/AudioMixerWasapiRenderStream.cpp
+++ b/Engine/Source/Runtime/Windows/AudioMixerWasapi/Private/AudioMixerWasapiRenderStream.cpp
@@ -50,6 +50,7 @@ namespace Audio
 			return false;
 		}
 
+#if PLATFORM_WINDOWS
 		AudioClientProperties ClientProperties = {};
 		ClientProperties.cbSize = sizeof(AudioClientProperties);
 		ClientProperties.bIsOffload = false;
@@ -63,6 +64,7 @@ namespace Audio
 			// Not all clients support raw mode. Log warning and continue.
 			UE_LOG(LogAudioMixerWasapi, Warning, TEXT("InitializeHardware failed SetClientProperties %s"), *AudioClientErrorToFString(Result));
 		}
+#endif //PLATFORM_WINDOWS
 		
 		WAVEFORMATEX* MixFormat = nullptr;
 		Result = TempAudioClient->GetMixFormat(&MixFormat);
@@ -327,7 +329,7 @@ namespace Audio
 		EAudioOutputStreamState::Type CurrStreamState = StreamState.load(std::memory_order_acquire);
 		if (CurrStreamState != EAudioOutputStreamState::Open)
 		{
-			UE_LOG(LogAudioMixerWasapi, Warning, TEXT("FAudioMixerWasapiRenderStream::StartAudioStream stream not in open sate. bIsInitialized: %d, CurrStreamState: %d"),
+			UE_LOG(LogAudioMixerWasapi, Warning, TEXT("FAudioMixerWasapiRenderStream::StartAudioStream stream not in open state. bIsInitialized: %d, CurrStreamState: %d"),
 				bIsInitialized.load(std::memory_order_relaxed), CurrStreamState);
 			return false;
 		}
