[Audio Insights] Removed legacy audio trace mesages + misc. trace messages adjustments

diff --git a/Engine/Plugins/AudioInsights/Source/AudioInsights/Private/Providers/AudioBusTraceProvider.cpp b/Engine/Plugins/AudioInsights/Source/AudioInsights/Private/Providers/AudioBusTraceProvider.cpp
--- a/Engine/Plugins/AudioInsights/Source/AudioInsights/Private/Providers/AudioBusTraceProvider.cpp
+++ b/Engine/Plugins/AudioInsights/Source/AudioInsights/Private/Providers/AudioBusTraceProvider.cpp
@@ -76,9 +76,8 @@ namespace UE::Audio::Insights
 
 				UE::Trace::IAnalyzer::FInterfaceBuilder& Builder = Context.InterfaceBuilder;
 
-				Builder.RouteEvent(RouteId_Start,          "Audio", "AudioBusStartOnTraceStart");
-				Builder.RouteEvent(RouteId_HasActivity,    "Audio", "AudioBusHasActivity");
 				Builder.RouteEvent(RouteId_Start,          "Audio", "AudioBusStart");
+				Builder.RouteEvent(RouteId_HasActivity,    "Audio", "AudioBusHasActivity");
 				Builder.RouteEvent(RouteId_EnvelopeValues, "Audio", "AudioBusEnvelopeValues");
 				Builder.RouteEvent(RouteId_Stop,           "Audio", "AudioBusStop");
 			}
diff --git a/Engine/Source/Runtime/AudioMixer/Private/AudioMixerSourceManager.cpp b/Engine/Source/Runtime/AudioMixer/Private/AudioMixerSourceManager.cpp
--- a/Engine/Source/Runtime/AudioMixer/Private/AudioMixerSourceManager.cpp
+++ b/Engine/Source/Runtime/AudioMixer/Private/AudioMixerSourceManager.cpp
@@ -275,29 +275,6 @@ UE_TRACE_EVENT_BEGIN(Audio, SoundSourceEnvelope)
 UE_TRACE_EVENT_END()
 
 // Audio Bus messages
-// @TODO UE-297630: remove this trace message once the full dashboard transition is finished (only editor -> editor and standalone)
-UE_TRACE_EVENT_BEGIN(Audio, AudioBusActivate)
-	UE_TRACE_EVENT_FIELD(uint32, DeviceId)
-	UE_TRACE_EVENT_FIELD(uint32, AudioBusId)
-	UE_TRACE_EVENT_FIELD(double, Timestamp)
-	UE_TRACE_EVENT_FIELD(UE::Trace::WideString, Name)
-UE_TRACE_EVENT_END()
-
-// @TODO UE-297630: remove this trace message once the full dashboard transition is finished (only editor -> editor and standalone)
-UE_TRACE_EVENT_BEGIN(Audio, AudioBusDeactivate)
-	UE_TRACE_EVENT_FIELD(uint32, DeviceId)
-	UE_TRACE_EVENT_FIELD(uint32, AudioBusId)
-	UE_TRACE_EVENT_FIELD(double, Timestamp)
-UE_TRACE_EVENT_END()
-
-UE_TRACE_EVENT_BEGIN(Audio, AudioBusStartOnTraceStart)
-	UE_TRACE_EVENT_FIELD(uint32, DeviceId)
-	UE_TRACE_EVENT_FIELD(uint32, AudioBusId)
-	UE_TRACE_EVENT_FIELD(double, Timestamp)
-	UE_TRACE_EVENT_FIELD(UE::Trace::WideString, Name)
-	UE_TRACE_EVENT_FIELD(int32, NumChannels)
-UE_TRACE_EVENT_END()
-
 UE_TRACE_EVENT_BEGIN(Audio, AudioBusStart)
 	UE_TRACE_EVENT_FIELD(uint32, DeviceId)
 	UE_TRACE_EVENT_FIELD(uint32, AudioBusId)
@@ -1517,13 +1494,6 @@ namespace Audio
 		AudioBusKeys_AudioThread.Add(InAudioBusKey);
 
 #if UE_AUDIO_PROFILERTRACE_ENABLED
-		// @TODO UE-297630: remove this trace message once the full dashboard transition is finished (only editor -> editor and standalone)
-		UE_TRACE_LOG(Audio, AudioBusActivate, AudioChannel)
-			<< AudioBusActivate.DeviceId(MixerDevice->DeviceID)
-			<< AudioBusActivate.AudioBusId(InAudioBusKey.ObjectId)
-			<< AudioBusActivate.Timestamp(FPlatformTime::Cycles64())
-			<< AudioBusActivate.Name(*InAudioBusName);
-
 		TraceAudioBusStartMessage(MixerDevice->DeviceID, InAudioBusKey.ObjectId, InAudioBusName, InNumChannels);
 #endif
 
@@ -1559,12 +1529,6 @@ namespace Audio
 		AudioBusKeys_AudioThread.Remove(InAudioBusKey);
 
 #if UE_AUDIO_PROFILERTRACE_ENABLED
-		// @TODO UE-297630: remove this trace message once the full dashboard transition is finished (only editor -> editor and standalone)
-		UE_TRACE_LOG(Audio, AudioBusDeactivate, AudioChannel)
-			<< AudioBusDeactivate.DeviceId(MixerDevice->DeviceID)
-			<< AudioBusDeactivate.AudioBusId(InAudioBusKey.ObjectId)
-			<< AudioBusDeactivate.Timestamp(FPlatformTime::Cycles64());
-
 		TraceAudioBusStopMessage(MixerDevice->DeviceID, InAudioBusKey.ObjectId);
 #endif
 
@@ -2801,10 +2765,6 @@ namespace Audio
 
 		ConnectBusPatches();
 
-#if UE_AUDIO_PROFILERTRACE_ENABLED
-		const bool bChannelEnabled = UE_TRACE_CHANNELEXPR_IS_ENABLED(AudioChannel);
-#endif // UE_AUDIO_PROFILERTRACE_ENABLED
-
 		// Loop through the bus registry and mix source audio
 		for (auto& [AudioBusKey, AudioBus] : AudioBuses)
 		{
@@ -2814,17 +2774,6 @@ namespace Audio
 			}
 
 			AudioBus->MixBuffer();
-
-#if UE_AUDIO_PROFILERTRACE_ENABLED
-			if (bChannelEnabled && AudioBus->bIsEnvelopeFollowing)
-			{
-				UE_TRACE_LOG(Audio, AudioBusEnvelopeValues, AudioChannel)
-					<< AudioBusEnvelopeValues.DeviceId(MixerDevice->DeviceID)
-					<< AudioBusEnvelopeValues.AudioBusId(AudioBusKey.ObjectId)
-					<< AudioBusEnvelopeValues.Timestamp(FPlatformTime::Cycles64())
-					<< AudioBusEnvelopeValues.EnvelopeValues(AudioBus->EnvelopeValues, AudioBus->GetNumChannels());
-			}
-#endif // UE_AUDIO_PROFILERTRACE_ENABLED
 		}
 	}
 
@@ -2832,22 +2781,43 @@ namespace Audio
 	{
 		RenderThreadPhase = ESourceManagerRenderThreadPhase::UpdateBusses;
 
+#if UE_AUDIO_PROFILERTRACE_ENABLED
+		const bool bChannelEnabled = UE_TRACE_CHANNELEXPR_IS_ENABLED(AudioMixerChannel);
+#endif // UE_AUDIO_PROFILERTRACE_ENABLED
+
 		// Update the bus states post mixing. This flips the current/previous buffer indices.
 		for (auto& [AudioBusKey, AudioBus] : AudioBuses)
 		{
+			if (!AudioBus.IsValid())
+			{
+				continue;
+			}
+
 			AudioBus->Update();
 			
 #if UE_AUDIO_PROFILERTRACE_ENABLED
-			const float* AudioBuffer = AudioBus->GetCurrentBusBuffer();
-			const int32 NumSamples   = AudioBus->NumChannels * AudioBus->NumFrames;
+			if (bChannelEnabled)
+			{
+				const float* AudioBuffer = AudioBus->GetCurrentBusBuffer();
+				const int32 NumSamples   = AudioBus->NumChannels * AudioBus->NumFrames;
+
+				const bool bIsAudioBufferSilent = IsAudioBufferSilent(AudioBuffer, NumSamples);
 
-			const bool bIsAudioBufferSilent = IsAudioBufferSilent(AudioBuffer, NumSamples);
+				UE_TRACE_LOG(Audio, AudioBusHasActivity, AudioChannel)
+					<< AudioBusHasActivity.DeviceId(MixerDevice->DeviceID)
+					<< AudioBusHasActivity.AudioBusId(AudioBusKey.ObjectId)
+					<< AudioBusHasActivity.Timestamp(FPlatformTime::Cycles64())
+					<< AudioBusHasActivity.HasActivity(!bIsAudioBufferSilent);
 
-			UE_TRACE_LOG(Audio, AudioBusHasActivity, AudioChannel)
-				<< AudioBusHasActivity.DeviceId(MixerDevice->DeviceID)
-				<< AudioBusHasActivity.AudioBusId(AudioBusKey.ObjectId)
-				<< AudioBusHasActivity.Timestamp(FPlatformTime::Cycles64())
-				<< AudioBusHasActivity.HasActivity(!bIsAudioBufferSilent);
+				if (AudioBus->bIsEnvelopeFollowing)
+				{
+					UE_TRACE_LOG(Audio, AudioBusEnvelopeValues, AudioChannel)
+						<< AudioBusEnvelopeValues.DeviceId(MixerDevice->DeviceID)
+						<< AudioBusEnvelopeValues.AudioBusId(AudioBusKey.ObjectId)
+						<< AudioBusEnvelopeValues.Timestamp(FPlatformTime::Cycles64())
+						<< AudioBusEnvelopeValues.EnvelopeValues(AudioBus->EnvelopeValues, AudioBus->GetNumChannels());
+				}
+			}
 #endif // UE_AUDIO_PROFILERTRACE_ENABLED
 		}
 	}
diff --git a/Engine/Source/Runtime/AudioMixer/Private/AudioMixerSubmix.cpp b/Engine/Source/Runtime/AudioMixer/Private/AudioMixerSubmix.cpp
--- a/Engine/Source/Runtime/AudioMixer/Private/AudioMixerSubmix.cpp
+++ b/Engine/Source/Runtime/AudioMixer/Private/AudioMixerSubmix.cpp
@@ -59,14 +59,6 @@ DEFINE_STAT(STAT_AudioMixerSubmixSoundfieldProcessors);
 
 #if UE_AUDIO_PROFILERTRACE_ENABLED
 
-// @TODO UE-297630: remove this trace message once the full dashboard transition is finished (only editor -> editor and standalone)
-UE_TRACE_EVENT_BEGIN(Audio, SoundSubmixHasActivity)
-	UE_TRACE_EVENT_FIELD(uint32, DeviceId)
-	UE_TRACE_EVENT_FIELD(uint32, SubmixId)
-	UE_TRACE_EVENT_FIELD(double, Timestamp)
-	UE_TRACE_EVENT_FIELD(bool, HasActivity)
-UE_TRACE_EVENT_END()
-
 UE_TRACE_EVENT_BEGIN(Audio, SubmixLoaded)
 	UE_TRACE_EVENT_FIELD(uint32, DeviceId)
 	UE_TRACE_EVENT_FIELD(uint32, SubmixId)
@@ -1676,22 +1668,16 @@ namespace Audio
 			}
 
 #if UE_AUDIO_PROFILERTRACE_ENABLED
-			if (OwningSubmixObject.IsValid())
+			const bool bChannelEnabled = UE_TRACE_CHANNELEXPR_IS_ENABLED(AudioChannel);
+			if (bChannelEnabled && OwningSubmixObject.IsValid())
 			{
-				// @TODO UE-297630: remove this trace message once the full dashboard transition is finished (only editor -> editor and standalone)
-				UE_TRACE_LOG(Audio, SoundSubmixHasActivity, AudioChannel)
-					<< SoundSubmixHasActivity.DeviceId(MixerDevice->DeviceID)
-					<< SoundSubmixHasActivity.SubmixId(GetTypeHash(OwningSubmixObject->GetPathName()))
-					<< SoundSubmixHasActivity.Timestamp(FPlatformTime::Cycles64())
-					<< SoundSubmixHasActivity.HasActivity(!bIsNowSilent);
-
 				UE_TRACE_LOG(Audio, SubmixHasActivity, AudioChannel)
 					<< SubmixHasActivity.DeviceId(MixerDevice->DeviceID)
 					<< SubmixHasActivity.SubmixId(GetTypeHash(OwningSubmixObject->GetPathName()))
 					<< SubmixHasActivity.Timestamp(FPlatformTime::Cycles64())
 					<< SubmixHasActivity.HasActivity(!bIsNowSilent);
 			}
-#endif
+#endif // UE_AUDIO_PROFILERTRACE_ENABLED
 		}
 
 	}
