Remove Submixes From Root and Add them as referenced objects so that Dynamic Submixes are automatically GC'd

diff --git a/Engine/Source/Runtime/AudioMixer/Private/AudioMixerDevice.cpp b/Engine/Source/Runtime/AudioMixer/Private/AudioMixerDevice.cpp
--- a/Engine/Source/Runtime/AudioMixer/Private/AudioMixerDevice.cpp
+++ b/Engine/Source/Runtime/AudioMixer/Private/AudioMixerDevice.cpp
@@ -528,6 +528,7 @@ namespace Audio
 
 	void FMixerDevice::AddReferencedObjects(FReferenceCollector& Collector)
 	{
+		Collector.AddReferencedObjects(LoadedSubmixes);
 	}
 
 	void FMixerDevice::CheckAudioThread() const
@@ -1328,7 +1329,7 @@ namespace Audio
 		}
 
 		check(NewSubmix);
-		NewSubmix->AddToRoot();
+		LoadedSubmixes.Add(NewSubmix->GetUniqueID(), NewSubmix);
 
 		// If sharing submix with other explicitly defined MainSubmix, create
 		// shared pointer directed to already existing submix instance. Otherwise,
@@ -1364,7 +1365,6 @@ namespace Audio
 			LLM_SCOPE(ELLMTag::AudioMixerPlugins);
 			USoundSubmix* ReverbPluginSubmix = ReverbPluginInterface->LoadSubmix();
 			check(ReverbPluginSubmix);
-			ReverbPluginSubmix->AddToRoot();
 
 			LoadSoundSubmix(*ReverbPluginSubmix);
 			InitSoundfieldAndEndpointDataForSubmix(*ReverbPluginSubmix, GetSubmixInstance(ReverbPluginSubmix).Pin(), false);
@@ -2149,8 +2149,8 @@ namespace Audio
 		FMixerSubmixPtr MixerSubmix = GetSubmixInstance(&InSoundSubmix).Pin();
 		if (!MixerSubmix.IsValid())
 		{
-			InSoundSubmix.AddToRoot();
-
+			LoadedSubmixes.Add(InSoundSubmix.GetUniqueID(), TObjectPtr<USoundSubmixBase>(&InSoundSubmix));
+			
 			MixerSubmix = MakeShared<FMixerSubmix, ESPMode::ThreadSafe>(this);
 			Submixes.Add(InSoundSubmix.GetUniqueID(), MixerSubmix);
 		}
@@ -2269,6 +2269,7 @@ namespace Audio
 		}
 
 		Submixes.Remove(InSoundSubmix.GetUniqueID());
+		LoadedSubmixes.Remove(InSoundSubmix.GetUniqueID());
 	}
 
 	FMixerSubmixPtr FMixerDevice::FindSubmixInstanceByObjectId(uint32 InObjectId)
diff --git a/Engine/Source/Runtime/AudioMixer/Public/AudioMixerDevice.h b/Engine/Source/Runtime/AudioMixer/Public/AudioMixerDevice.h
--- a/Engine/Source/Runtime/AudioMixer/Public/AudioMixerDevice.h
+++ b/Engine/Source/Runtime/AudioMixer/Public/AudioMixerDevice.h
@@ -462,6 +462,9 @@ namespace Audio
 		/** Map of USoundSubmix static data objects to the dynamic audio mixer submix. */
 		FSubmixMap Submixes;
 
+		/** A map of all currently loaded submixes, so that they can be tracked for GC */
+		TMap<uint32, TObjectPtr<USoundSubmixBase>> LoadedSubmixes;
+		
 		// Submixes that will sum their audio and send it directly to AudioMixerPlatform.
 		// Submixes are added to this list in RegisterSoundSubmix, and removed in UnregisterSoundSubmix.
 		TArray<FMixerSubmixPtr> DefaultEndpointSubmixes;
