Replace remaning RemoveFromRoot call with removal from the LoadedSubmixes list

diff --git a/Engine/Source/Runtime/AudioMixer/Private/AudioMixerDevice.cpp b/Engine/Source/Runtime/AudioMixer/Private/AudioMixerDevice.cpp
--- a/Engine/Source/Runtime/AudioMixer/Private/AudioMixerDevice.cpp
+++ b/Engine/Source/Runtime/AudioMixer/Private/AudioMixerDevice.cpp
@@ -1289,7 +1289,7 @@ namespace Audio
 			{
 				return;
 			}
-			OldSubmix->RemoveFromRoot();
+			LoadedSubmixes.Remove(OldSubmix->GetUniqueID());
 			FMixerSubmixPtr OldSubmixPtr = RequiredSubmixInstances[TypeIndex];
 			if (OldSubmixPtr.IsValid())
 			{
