[AudioMixer] Added SendSubmixesAlivePing method so we can keep Audio Insights informed every second about the currently active submixes.

diff --git a/Engine/Plugins/AudioInsights/Source/AudioInsights/Private/Providers/SubmixTraceProvider.cpp b/Engine/Plugins/AudioInsights/Source/AudioInsights/Private/Providers/SubmixTraceProvider.cpp
--- a/Engine/Plugins/AudioInsights/Source/AudioInsights/Private/Providers/SubmixTraceProvider.cpp
+++ b/Engine/Plugins/AudioInsights/Source/AudioInsights/Private/Providers/SubmixTraceProvider.cpp
@@ -76,9 +76,9 @@ namespace UE::Audio::Insights
 
 				UE::Trace::IAnalyzer::FInterfaceBuilder& Builder = Context.InterfaceBuilder;
 
-				Builder.RouteEvent(RouteId_Loaded,         "Audio", "SubmixLoadedOnTraceStart");
-				Builder.RouteEvent(RouteId_HasActivity,    "Audio", "SubmixHasActivity");
+				Builder.RouteEvent(RouteId_Loaded,         "Audio", "SubmixAlivePing");
 				Builder.RouteEvent(RouteId_Loaded,         "Audio", "SubmixLoaded");
+				Builder.RouteEvent(RouteId_HasActivity,    "Audio", "SubmixHasActivity");
 				Builder.RouteEvent(RouteId_EnvelopeValues, "Audio", "SubmixEnvelopeValues");
 				Builder.RouteEvent(RouteId_Unloaded,       "Audio", "SubmixUnloaded");
 			}
diff --git a/Engine/Source/Runtime/AudioMixer/Private/AudioMixerDevice.cpp b/Engine/Source/Runtime/AudioMixer/Private/AudioMixerDevice.cpp
--- a/Engine/Source/Runtime/AudioMixer/Private/AudioMixerDevice.cpp
+++ b/Engine/Source/Runtime/AudioMixer/Private/AudioMixerDevice.cpp
@@ -128,7 +128,7 @@ FAutoConsoleVariableRef CVarMuteAudio(
 
 #if UE_AUDIO_PROFILERTRACE_ENABLED
 
-UE_TRACE_EVENT_BEGIN(Audio, SubmixLoadedOnTraceStart)
+UE_TRACE_EVENT_BEGIN(Audio, SubmixAlivePing)
 	UE_TRACE_EVENT_FIELD(uint32, DeviceId)
 	UE_TRACE_EVENT_FIELD(uint32, SubmixId)
 	UE_TRACE_EVENT_FIELD(double, Timestamp)
@@ -381,18 +381,23 @@ namespace Audio
 
 #if UE_AUDIO_PROFILERTRACE_ENABLED
 	void FMixerDevice::OnTraceStarted(FTraceAuxiliary::EConnectionType TraceType, const FString& TraceDestination)
+	{
+		SendSubmixesAlivePing();
+	}
+
+	void FMixerDevice::SendSubmixesAlivePing()
 	{
 		const bool bChannelEnabled = UE_TRACE_CHANNELEXPR_IS_ENABLED(AudioChannel);
 		if (bChannelEnabled)
 		{
 			auto TraceSubmixLoadedMessage = [this](const FMixerSubmixPtr& InSubmix)
 			{
-				UE_TRACE_LOG(Audio, SubmixLoadedOnTraceStart, AudioChannel)
-					<< SubmixLoadedOnTraceStart.DeviceId(DeviceID)
-					<< SubmixLoadedOnTraceStart.SubmixId(GetTypeHash(InSubmix->GetPathName()))
-					<< SubmixLoadedOnTraceStart.Timestamp(FPlatformTime::Cycles64())
-					<< SubmixLoadedOnTraceStart.Name(*InSubmix->GetPathName())
-					<< SubmixLoadedOnTraceStart.NumChannels(InSubmix->NumChannels);
+				UE_TRACE_LOG(Audio, SubmixAlivePing, AudioChannel)
+					<< SubmixAlivePing.DeviceId(DeviceID)
+					<< SubmixAlivePing.SubmixId(GetTypeHash(InSubmix->GetPathName()))
+					<< SubmixAlivePing.Timestamp(FPlatformTime::Cycles64())
+					<< SubmixAlivePing.Name(*InSubmix->GetPathName())
+					<< SubmixAlivePing.NumChannels(InSubmix->NumChannels);
 			};
 
 			Submixes.Iterate([this, &TraceSubmixLoadedMessage](const FSubmixMap::FPair& SubmixPair)
@@ -414,6 +419,17 @@ namespace Audio
 		}
 	}
 
+	void FMixerDevice::UpdateSubmixesAlivePingTime()
+	{
+		SubmixesAlivePingTimeSec += AudioClockDelta;
+
+		if (SubmixesAlivePingTimeSec > SubmixesAlivePingWaitTimeSec)
+		{
+			SendSubmixesAlivePing();
+			SubmixesAlivePingTimeSec = 0.0;
+		}
+	}
+
 	void FMixerDevice::RegisterSubmixEnvelopeFollowerCVars()
 	{
 		// StartSubmixEnvelopeFollowerCVar
@@ -789,6 +805,10 @@ namespace Audio
 				AudioClockDelta = (double)OpenStreamParams.NumFrames / OpenStreamParams.SampleRate;
 				AudioClockTimingData.UpdateTime = 0.0;
 
+#if UE_AUDIO_PROFILERTRACE_ENABLED
+				SubmixesAlivePingTimeSec = 0.0;
+#endif //UE_AUDIO_PROFILERTRACE_ENABLED
+
 				PluginInitializationParams.NumSources = SourceManagerInitParams.NumSources;
 				PluginInitializationParams.SampleRate = SampleRate;
 				PluginInitializationParams.BufferLength = OpenStreamParams.NumFrames;
@@ -1229,6 +1249,10 @@ namespace Audio
 		// Update the audio clock
 		UpdateAudioClock();
 
+#if UE_AUDIO_PROFILERTRACE_ENABLED
+		UpdateSubmixesAlivePingTime();
+#endif //UE_AUDIO_PROFILERTRACE_ENABLED
+
 		// notify interested parties
 		NotifyAudioDevicePostRender(RenderInfo);
 
diff --git a/Engine/Source/Runtime/AudioMixer/Public/AudioMixerDevice.h b/Engine/Source/Runtime/AudioMixer/Public/AudioMixerDevice.h
--- a/Engine/Source/Runtime/AudioMixer/Public/AudioMixerDevice.h
+++ b/Engine/Source/Runtime/AudioMixer/Public/AudioMixerDevice.h
@@ -409,6 +409,8 @@ namespace Audio
 
 #if UE_AUDIO_PROFILERTRACE_ENABLED
 		void OnTraceStarted(FTraceAuxiliary::EConnectionType TraceType, const FString& TraceDestination);
+		void SendSubmixesAlivePing();
+		void UpdateSubmixesAlivePingTime();
 
 		void RegisterSubmixEnvelopeFollowerCVars();
 		void UnregisterSubmixEnvelopeFollowerCVars();
@@ -501,6 +503,9 @@ namespace Audio
 #if UE_AUDIO_PROFILERTRACE_ENABLED
 		IConsoleVariable* StartSubmixEnvelopeFollowerCVar = nullptr;
 		IConsoleVariable* StopSubmixEnvelopeFollowerCVar  = nullptr;
+
+		double SubmixesAlivePingTimeSec = 0.0;
+		static constexpr double SubmixesAlivePingWaitTimeSec = 1.0;
 #endif // UE_AUDIO_PROFILERTRACE_ENABLED
 
 		/** Whether or not we generate output audio to test multi-platform mixer. */
