[CrashReport][Assert] UnrealEditor-AudioMixer!UE::Core::Private::ReallocGrow<3,TSizedHeapAllocator<32,FMemory>::ForAnyElementType>(unsigned int,unsigned int,int,TSizedHeapAllocator<32,FMemory>::ForAnyElementType &,int &,int &) [Engine/Source/Runtime/Core/

diff --git a/Engine/Source/Runtime/AudioMixer/Private/SoundFileIO/SoundFileIO.cpp b/Engine/Source/Runtime/AudioMixer/Private/SoundFileIO/SoundFileIO.cpp
--- a/Engine/Source/Runtime/AudioMixer/Private/SoundFileIO/SoundFileIO.cpp
+++ b/Engine/Source/Runtime/AudioMixer/Private/SoundFileIO/SoundFileIO.cpp
@@ -148,6 +148,7 @@ namespace Audio::SoundFileUtils
 		}
 
 		SoundFileCount SamplesWritten = 0;
+		SoundFileCount TotalSamplesRead = InputSamplesRead;
 
 		while (InputSamplesRead == ProcessBuffer.Num())
 		{
@@ -174,6 +175,15 @@ namespace Audio::SoundFileUtils
 					ProcessBuffer[Sample] = FMath::Clamp(ProcessBuffer[Sample], -1.0f, 1.0f);
 				}
 			}
+
+			TotalSamplesRead += InputSamplesRead;
+
+			if (TotalSamplesRead > static_cast<SoundFileCount>(INT32_MAX))
+			{
+				UE_LOG(LogAudioMixer, Error, TEXT("Error encountered while converting audio to wav, file too large... skipping"));
+
+				return false;
+			}
 		}
 
 		// Write final block of samples
diff --git a/Engine/Source/Runtime/AudioMixer/Private/SoundFileIO/SoundFileIOManagerImpl.cpp b/Engine/Source/Runtime/AudioMixer/Private/SoundFileIO/SoundFileIOManagerImpl.cpp
--- a/Engine/Source/Runtime/AudioMixer/Private/SoundFileIO/SoundFileIOManagerImpl.cpp
+++ b/Engine/Source/Runtime/AudioMixer/Private/SoundFileIO/SoundFileIOManagerImpl.cpp
@@ -1425,10 +1425,19 @@ namespace Audio
 
 			SoundFileCount BulkDataLength = BulkData.Num();
 
-			// If we need more room, we add it here.
+			// If we need more room, we add it here. 
 			int64 NumExtraBytesNeeded = (CurrentIndexBytes + NumBytes) - BulkDataLength;
+
+			if (BulkData.Num() + NumExtraBytesNeeded > INT32_MAX)
+			{
+				UE_LOG(LogAudioMixer, Error, TEXT("Error encountered while writing bytes, file too large... skipping"));
+
+				return ESoundFileError::Type::INVALID_DATA;
+			}
+
 			if (NumExtraBytesNeeded > 0)
 			{
+				check(BulkData.Num() + NumExtraBytesNeeded <= INT32_MAX);
 				BulkData.AddUninitialized(NumExtraBytesNeeded);
 			}
 
