Remove From Root on Submix unload (requires some deprecation)

diff --git a/Engine/Source/Runtime/AudioMixer/Private/AudioMixerDevice.cpp b/Engine/Source/Runtime/AudioMixer/Private/AudioMixerDevice.cpp
--- a/Engine/Source/Runtime/AudioMixer/Private/AudioMixerDevice.cpp
+++ b/Engine/Source/Runtime/AudioMixer/Private/AudioMixerDevice.cpp
@@ -2175,6 +2175,14 @@ namespace Audio
 		FMixerSubmixPtr MixerSubmix = GetSubmixInstance(&InSoundSubmix).Pin();
 		if (!MixerSubmix.IsValid())
 		{
+			// We add all submixes to root, including dynamic submixes. Since the static graph is set up early (during Engine init),
+			// all of the static submixes are implicitly added to a pool of objects which are never checked for GC.
+			// Those objects are not allowed to reference objects which are eligible for GC, because they are comepletely excluded from the GC evaluation
+			// and may cause objects to GC at unexpected times (e.g. if the new submix is only referenced by the static parent, it may be immediately GC'd.
+			// Therefore, with the current implementation of dynamic submixes which creates a reference between the UObjects,
+			// we can only connect dynamic submixes to the main graph if they are rooted objects.
+			// We will remove them from root when the submix is unloaded (FMixerDevice::UnloadSoundSubmix).
+			// This will still allow the submix to be GC'd once it is no longer in use.
 			InSoundSubmix.AddToRoot();
 
 			MixerSubmix = MakeShared<FMixerSubmix, ESPMode::ThreadSafe>(this);
@@ -2233,6 +2241,32 @@ namespace Audio
 
 			const TWeakObjectPtr<const USoundSubmixBase> SubmixToUnload = InSoundSubmix;
 			FAudioThread::RunCommandOnAudioThread([this, SubmixToUnload, bReparentChildren]()
+			{
+				CSV_SCOPED_TIMING_STAT(Audio, UnregisterSubmix);
+				if (SubmixToUnload.IsValid())
+				{
+					UnloadSoundSubmixWithoutRemovingFromRoot(*SubmixToUnload.Get(), bReparentChildren);
+				}
+			}, GET_STATID(STAT_AudioUnregisterSoundSubmix));
+			return;
+		}
+
+		UnloadSoundSubmixWithoutRemovingFromRoot(*InSoundSubmix, bReparentChildren);
+	}
+
+	void FMixerDevice::UnregisterSoundSubmix(USoundSubmixBase* InSoundSubmix, const bool bReparentChildren)
+	{
+		if (!InSoundSubmix || bSubmixRegistrationDisabled || IsRequiredSubmixType(InSoundSubmix))
+		{
+			return;
+		}
+
+		if (!IsInAudioThread())
+		{
+			DECLARE_CYCLE_STAT(TEXT("FAudioThreadTask.UnregisterSoundSubmix"), STAT_AudioUnregisterSoundSubmix, STATGROUP_AudioThreadCommands);
+
+			const TWeakObjectPtr<USoundSubmixBase> SubmixToUnload = InSoundSubmix;
+			FAudioThread::RunCommandOnAudioThread([this, SubmixToUnload, bReparentChildren]()
 			{
 				CSV_SCOPED_TIMING_STAT(Audio, UnregisterSubmix);
 				if (SubmixToUnload.IsValid())
@@ -2246,7 +2280,7 @@ namespace Audio
 		UnloadSoundSubmix(*InSoundSubmix, bReparentChildren);
 	}
 
-	void FMixerDevice::UnloadSoundSubmix(const USoundSubmixBase& InSoundSubmix, const bool bReparentChildren)
+	void FMixerDevice::UnloadSoundSubmixWithoutRemovingFromRoot(const USoundSubmixBase& InSoundSubmix, const bool bReparentChildren)
 	{
 		check(IsInAudioThread());
 
@@ -2295,8 +2329,15 @@ namespace Audio
 		}
 
 		Submixes.Remove(InSoundSubmix.GetUniqueID());
+		
 	}
 
+	void FMixerDevice::UnloadSoundSubmix(USoundSubmixBase& InSoundSubmix, const bool bReparentChildren)
+	{
+		UnloadSoundSubmixWithoutRemovingFromRoot(InSoundSubmix, bReparentChildren);
+		InSoundSubmix.RemoveFromRoot();
+	}
+	
 	FMixerSubmixPtr FMixerDevice::FindSubmixInstanceByObjectId(uint32 InObjectId)
 	{
 		for (int32 i = 0; i < RequiredSubmixes.Num(); i++)
diff --git a/Engine/Source/Runtime/AudioMixer/Public/AudioMixerDevice.h b/Engine/Source/Runtime/AudioMixer/Public/AudioMixerDevice.h
--- a/Engine/Source/Runtime/AudioMixer/Public/AudioMixerDevice.h
+++ b/Engine/Source/Runtime/AudioMixer/Public/AudioMixerDevice.h
@@ -147,7 +147,11 @@ namespace Audio
 		AUDIOMIXER_API virtual void EnableDebugAudioOutput() override;
 		AUDIOMIXER_API virtual FAudioPlatformSettings GetPlatformSettings() const override;
 		AUDIOMIXER_API virtual void RegisterSoundSubmix(USoundSubmixBase* SoundSubmix, bool bInit = true) override;
+
+		UE_DEPRECATED(5.7, "Deprecated in favor of the version of this function which takes a non-const USoundSubmix.")
 		AUDIOMIXER_API virtual void UnregisterSoundSubmix(const USoundSubmixBase* SoundSubmix, const bool bReparentChildren) override;
+		
+		AUDIOMIXER_API virtual void UnregisterSoundSubmix(USoundSubmixBase* SoundSubmix, const bool bReparentChildren) override;
 
 		AUDIOMIXER_API virtual int32 GetNumActiveSources() const override;
 
@@ -393,8 +397,9 @@ namespace Audio
 		void LoadSoundSubmix(USoundSubmixBase& SoundSubmix);
 
 		void InitSoundfieldAndEndpointDataForSubmix(const USoundSubmixBase& InSoundSubmix, FMixerSubmixPtr MixerSubmix, bool bAllowReInit);
-
-		void UnloadSoundSubmix(const USoundSubmixBase& SoundSubmix, const bool bReparentChildren);
+		
+		void UnloadSoundSubmixWithoutRemovingFromRoot(const USoundSubmixBase& SoundSubmix, const bool bReparentChildren);
+		void UnloadSoundSubmix(USoundSubmixBase& SoundSubmix, const bool bReparentChildren);
 
 		bool IsRequiredSubmixType(const USoundSubmixBase* InSubmix) const;
 		FMixerSubmixPtr GetRequiredSubmixInstance(uint32 InSubmixId) const;
diff --git a/Engine/Source/Runtime/Engine/Private/AudioDeviceManager.cpp b/Engine/Source/Runtime/Engine/Private/AudioDeviceManager.cpp
--- a/Engine/Source/Runtime/Engine/Private/AudioDeviceManager.cpp
+++ b/Engine/Source/Runtime/Engine/Private/AudioDeviceManager.cpp
@@ -921,6 +921,20 @@ void FAudioDeviceManager::RegisterSoundSubmix(USoundSubmixBase* SoundSubmix)
 }
 
 void FAudioDeviceManager::UnregisterSoundSubmix(const USoundSubmixBase* SoundSubmix)
+{
+	// const version is deprecated
+	PRAGMA_DISABLE_DEPRECATION_WARNINGS
+	IterateOverAllDevices(
+		[&SoundSubmix](Audio::FDeviceId, FAudioDevice* InDevice)
+		{
+			InDevice->UnregisterSoundSubmix(SoundSubmix, true);
+		}
+	);
+	PRAGMA_ENABLE_DEPRECATION_WARNINGS
+	
+}
+
+void FAudioDeviceManager::UnregisterSoundSubmix(USoundSubmixBase* SoundSubmix)
 {
 	IterateOverAllDevices(
 		[&SoundSubmix](Audio::FDeviceId, FAudioDevice* InDevice)
diff --git a/Engine/Source/Runtime/Engine/Public/AudioDevice.h b/Engine/Source/Runtime/Engine/Public/AudioDevice.h
--- a/Engine/Source/Runtime/Engine/Public/AudioDevice.h
+++ b/Engine/Source/Runtime/Engine/Public/AudioDevice.h
@@ -963,9 +963,12 @@ public:
 	/** Registers the sound submix */
 	virtual void RegisterSoundSubmix(USoundSubmixBase* SoundSubmix, bool bInit) {}
 
-	/** Unregisters the sound submix */
+	UE_DEPRECATED(5.7, "Deprecated in favor of the version of this function which takes a non-const USoundSubmix.")
 	virtual void UnregisterSoundSubmix(const USoundSubmixBase* SoundSubmix, const bool bReparentChildren) {}
 
+	/** Unregisters the sound submix */
+	virtual void UnregisterSoundSubmix(USoundSubmixBase* SoundSubmix, const bool bReparentChildren) {}
+
 	ENGINE_API virtual USoundSubmix& GetMainSubmixObject() const;
 
 	UE_DEPRECATED(5.4, "Use RegisterSubmixBufferListener version that requires a shared reference to a listener and provide explicit reference to a submix: use GetMainSubmixObject to register with the Main Output Submix (rather than nullptr for safety), and instantiate buffer listener via the shared pointer API.")
diff --git a/Engine/Source/Runtime/Engine/Public/AudioDeviceManager.h b/Engine/Source/Runtime/Engine/Public/AudioDeviceManager.h
--- a/Engine/Source/Runtime/Engine/Public/AudioDeviceManager.h
+++ b/Engine/Source/Runtime/Engine/Public/AudioDeviceManager.h
@@ -216,8 +216,11 @@ public:
 	/** Registers the Sound Mix for all active devices. */
 	ENGINE_API void RegisterSoundSubmix(USoundSubmixBase* SoundSubmix);
 
-	/** Registers the Sound Mix for all active devices. */
+	UE_DEPRECATED(5.7, "Deprecated in favor of the version of this function which takes a non-const USoundSubmix.")
 	ENGINE_API void UnregisterSoundSubmix(const USoundSubmixBase* SoundSubmix);
+	
+	/** Registers the Sound Mix for all active devices. */
+	ENGINE_API void UnregisterSoundSubmix(USoundSubmixBase* SoundSubmix);
 
 	/** Initializes the sound mixes for all active devices. */
 	ENGINE_API void InitSoundSubmixes();
