[Audio Render Scheduler] Move bus mixing and source bus rendering into scheduled steps, and add dependencies so all these steps are ordered and synchronized correctly.

#rnx
#tests replay runs, with and without the scheduler enabled.

diff --git a/Engine/Plugins/Tests/AudioCQTests/Source/AudioCQTests/Private/AudioRenderSchedulerTests.cpp b/Engine/Plugins/Tests/AudioCQTests/Source/AudioCQTests/Private/AudioRenderSchedulerTests.cpp
--- a/Engine/Plugins/Tests/AudioCQTests/Source/AudioCQTests/Private/AudioRenderSchedulerTests.cpp
+++ b/Engine/Plugins/Tests/AudioCQTests/Source/AudioCQTests/Private/AudioRenderSchedulerTests.cpp
@@ -113,6 +113,29 @@ TEST_CLASS(AudioRenderSchedulerTests, "Audio.RenderScheduler")
 			});
 	}
 
+	TEST_METHOD(SinglePrereqDoubled)
+	{
+		TestCommandBuilder.Do([this]()
+			{
+				// Make sure that multiple redundant dependencies work just as well as single dependencies
+				AddStepRange(10);
+				for (int i = 1; i < 10; ++i)
+				{
+					Tester->AddDependency(0, i);
+					Tester->AddDependency(0, i);
+				}
+
+				Results = Tester->Run();
+
+				ASSERT_THAT(AreEqual(Results.Num(), 10));
+				for (int i = 0; i < 10; ++i)
+				{
+					ASSERT_THAT(IsTrue(Results.Contains(i)));
+				}
+				ASSERT_THAT(AreEqual(Results[0], 0));
+			});
+	}
+
 	TEST_METHOD(SinglePostreq)
 	{
 		TestCommandBuilder.Do([this]()
@@ -264,6 +287,31 @@ TEST_CLASS(AudioRenderSchedulerTests, "Audio.RenderScheduler")
 			});
 	}
 
+	TEST_METHOD(DenseGraphWithCycles)
+	{
+		TestCommandBuilder.Do([this]()
+			{
+				AddStepRange(10);
+				for (int i = 0; i < 10; ++i)
+				{
+					for (int j = 0; j < 10; ++j)
+					{
+						Tester->AddDependency(i, j);
+					}
+				}
+
+				Results = Tester->Run();
+
+				// The order is completely arbitrary here, but we check that the scheduler managed to chew through all the cycles
+				//  and schedule every step.
+				ASSERT_THAT(AreEqual(Results.Num(), 10));
+				for (int i = 0; i < 10; ++i)
+				{
+					ASSERT_THAT(IsTrue(Results.Contains(i)));
+				}
+			});
+	}
+
 	TEST_METHOD(TwoCyclesRepeated)
 	{
 		TestCommandBuilder.Do([this]()
diff --git a/Engine/Source/Runtime/AudioMixer/Private/AudioMixerBus.cpp b/Engine/Source/Runtime/AudioMixer/Private/AudioMixerBus.cpp
--- a/Engine/Source/Runtime/AudioMixer/Private/AudioMixerBus.cpp
+++ b/Engine/Source/Runtime/AudioMixer/Private/AudioMixerBus.cpp
@@ -3,20 +3,36 @@
 #include "AudioMixerBus.h"
 
 #include "Algo/ForEach.h"
+#include "AudioMixerCVars.h"
 #include "AudioMixerSourceManager.h"
+#include "AudioRenderScheduler.h"
 #include "DSP/AlignedBuffer.h"
 #include "DSP/FloatArrayMath.h"
 
 namespace Audio
 {
-	FMixerAudioBus::FMixerAudioBus(FMixerSourceManager* InSourceManager, bool bInIsAutomatic, int32 InNumChannels)
+	FMixerAudioBus::FMixerAudioBus(FMixerSourceManager* InSourceManager, bool bInIsAutomatic, int32 InNumChannels, const FAudioBusKey InBusKey)
 		: CurrentBufferIndex(1)
 		, NumChannels(InNumChannels)
 		, NumFrames(InSourceManager->GetNumOutputFrames())
 		, SourceManager(InSourceManager)
+		, BusKey(InBusKey)
 		, bIsAutomatic(bInIsAutomatic)
 	{
 		SetNumOutputChannels(NumChannels);
+
+		if (AudioMixerCVars::UseRenderScheduler)
+		{
+			SourceManager->GetRenderScheduler().AddStep(FAudioRenderStepId::FromAudioBusKey(BusKey), this);
+		}
+	}
+
+	FMixerAudioBus::~FMixerAudioBus()
+	{
+		if (AudioMixerCVars::UseRenderScheduler)
+		{
+			SourceManager->GetRenderScheduler().RemoveStep(FAudioRenderStepId::FromAudioBusKey(BusKey));
+		}
 	}
 
 	void FMixerAudioBus::SetNumOutputChannels(int32 InNumOutputChannels)
@@ -32,17 +48,31 @@ namespace Audio
 
 	void FMixerAudioBus::Update()
 	{
-		CurrentBufferIndex = !CurrentBufferIndex;
+		// When using the scheduler we flip buffers at the beginning of MixBuffer() instead.
+		check(!AudioMixerCVars::UseRenderScheduler);
+		CurrentBufferIndex = 1 - CurrentBufferIndex;
 	}
 
-	void FMixerAudioBus::AddInstanceId(const int32 InSourceId, int32 InNumOutputChannels)
+	void FMixerAudioBus::AddInstanceId(const int32 InSourceId, const uint64 InTransmitterID, int32 InNumOutputChannels)
 	{
 		InstanceIds.Add(InSourceId);
+		if (AudioMixerCVars::UseRenderScheduler)
+		{
+			// Make sure this bus gets mixed before the source from it gets rendered
+			const FAudioRenderStepId SourceStepId = FAudioRenderStepId::FromTransmitterID(InTransmitterID);
+			const FAudioRenderStepId BusStepId = FAudioRenderStepId::FromAudioBusKey(BusKey);
+			SourceManager->GetRenderScheduler().AddDependency(BusStepId, SourceStepId);
+		}
 	}
 
-	bool FMixerAudioBus::RemoveInstanceId(const int32 InSourceId)
+	bool FMixerAudioBus::RemoveInstanceId(const int32 InSourceId, const uint64 InTransmitterID)
 	{
-		InstanceIds.Remove(InSourceId);
+		if (InstanceIds.Remove(InSourceId) > 0 && AudioMixerCVars::UseRenderScheduler)
+		{
+			const FAudioRenderStepId SourceStepId = FAudioRenderStepId::FromTransmitterID(InTransmitterID);
+			const FAudioRenderStepId BusStepId = FAudioRenderStepId::FromAudioBusKey(BusKey);
+			SourceManager->GetRenderScheduler().RemoveDependency(BusStepId, SourceStepId);
+		}
 
 		// Return true if there is no more instances or sends
 		return bIsAutomatic && !InstanceIds.Num() && !AudioBusSends[(int32)EBusSendType::PreEffect].Num() && !AudioBusSends[(int32)EBusSendType::PostEffect].Num();
@@ -63,6 +93,14 @@ namespace Audio
 
 		// It's a new source id so just add it
 		AudioBusSends[(int32)BusSendType].Add(InAudioBusSend);
+
+		if (AudioMixerCVars::UseRenderScheduler)
+		{
+			// Make sure the source sending to the bus gets rendered before the bus is mixed
+			const FAudioRenderStepId SourceStepId = FAudioRenderStepId::FromTransmitterID(InAudioBusSend.TransmitterID);
+			const FAudioRenderStepId BusStepId = FAudioRenderStepId::FromAudioBusKey(BusKey);
+			SourceManager->GetRenderScheduler().AddDependency(SourceStepId, BusStepId);
+		}
 	}
 
 	bool FMixerAudioBus::RemoveSend(EBusSendType BusSendType, const int32 InSourceId)
@@ -74,6 +112,13 @@ namespace Audio
 			// Remove this source id's send
 			if (Sends[i].SourceId == InSourceId)
 			{
+				if (AudioMixerCVars::UseRenderScheduler)
+				{
+					const FAudioRenderStepId SourceStepId = FAudioRenderStepId::FromTransmitterID(Sends[i].TransmitterID);
+					const FAudioRenderStepId BusStepId = FAudioRenderStepId::FromAudioBusKey(BusKey);
+					SourceManager->GetRenderScheduler().RemoveDependency(SourceStepId, BusStepId);
+				}
+
 				Sends.RemoveAtSwap(i, EAllowShrinking::No);
 
 				// There will only be one entry
@@ -91,6 +136,11 @@ namespace Audio
 		const int32 NumSamples = NumFrames * NumChannels;
 		const int32 NumOutputFrames = SourceManager->GetNumOutputFrames();
 
+		if (AudioMixerCVars::UseRenderScheduler)
+		{
+			CurrentBufferIndex = 1 - CurrentBufferIndex;
+		}
+
 		FAlignedFloatBuffer& MixBuffer = MixedSourceData[CurrentBufferIndex];
 		float* BusDataBufferPtr = MixBuffer.GetData();
 
@@ -103,8 +153,10 @@ namespace Audio
 			{
 				const float* SourceBufferPtr = nullptr;
 
-				// If the audio source mixing to this audio bus is itself a source bus, we need to use the previous renderer buffer to avoid infinite recursion
-				if (SourceManager->IsSourceBus(AudioBusSend.SourceId))
+				// If the audio source mixing to this audio bus is itself a source bus, we need to use the previous renderer buffer to avoid infinite recursion.
+				// With the render scheduler we don't need to treat source buses any differently from other sources -- the source bus should already have been
+				//  rendered.
+				if (!AudioMixerCVars::UseRenderScheduler && SourceManager->IsSourceBus(AudioBusSend.SourceId))
 				{
 					SourceBufferPtr = SourceManager->GetPreviousSourceBusBuffer(AudioBusSend.SourceId);
 				}
@@ -176,7 +228,7 @@ namespace Audio
 
 	const float* FMixerAudioBus::GetPreviousBusBuffer() const
 	{
-		return MixedSourceData[!CurrentBufferIndex].GetData();
+		return MixedSourceData[1 - CurrentBufferIndex].GetData();
 	}
 
 	void FMixerAudioBus::AddNewPatchOutput(const FPatchOutputStrongPtr& InPatchOutputStrongPtr)
@@ -194,6 +246,17 @@ namespace Audio
 		return PatchMixer.RemovePatch(PatchInput);
 	}
 
+
+	void FMixerAudioBus::DoRenderStep()
+	{
+		MixBuffer();
+	}
+
+	const TCHAR* FMixerAudioBus::GetRenderStepName()
+	{
+		return TEXT("FMixerAudioBus mixing");
+	}
+
 #if UE_AUDIO_PROFILERTRACE_ENABLED
 	void FMixerAudioBus::StartEnvelopeFollower(const float InAttackTime, const float InReleaseTime, const float InSampleRate)
 	{
diff --git a/Engine/Source/Runtime/AudioMixer/Private/AudioMixerCVars.cpp b/Engine/Source/Runtime/AudioMixer/Private/AudioMixerCVars.cpp
new file mode 100644
--- /dev/null
+++ b/Engine/Source/Runtime/AudioMixer/Private/AudioMixerCVars.cpp
@@ -0,0 +1,16 @@
+// Copyright Epic Games, Inc. All Rights Reserved.
+
+#include "AudioMixerCVars.h"
+
+#include "HAL/IConsoleManager.h"
+
+namespace AudioMixerCVars
+{
+	int32 UseRenderScheduler = 0;
+	FAutoConsoleVariableRef CVarUseAudioRenderScheduler(
+		TEXT("au.UseRenderScheduler"),
+		UseRenderScheduler,
+		TEXT("Use audio rendering scheduler to handle render dependencies.\n")
+		TEXT("0: Disabled, 1: Enabled"),
+		ECVF_Default);
+}
diff --git a/Engine/Source/Runtime/AudioMixer/Private/AudioMixerCVars.h b/Engine/Source/Runtime/AudioMixer/Private/AudioMixerCVars.h
new file mode 100644
--- /dev/null
+++ b/Engine/Source/Runtime/AudioMixer/Private/AudioMixerCVars.h
@@ -0,0 +1,10 @@
+// Copyright Epic Games, Inc. All Rights Reserved.
+
+#pragma once
+
+#include "CoreTypes.h"
+
+namespace AudioMixerCVars
+{
+	extern int32 UseRenderScheduler;
+}
diff --git a/Engine/Source/Runtime/AudioMixer/Private/AudioMixerDevice.cpp b/Engine/Source/Runtime/AudioMixer/Private/AudioMixerDevice.cpp
--- a/Engine/Source/Runtime/AudioMixer/Private/AudioMixerDevice.cpp
+++ b/Engine/Source/Runtime/AudioMixer/Private/AudioMixerDevice.cpp
@@ -378,6 +378,10 @@ namespace Audio
 		{
 			delete AudioMixerPlatform;
 		}
+
+		// The rendering thread has been stopped now, so it is safe to run any cleanup code that
+		//  would normally be run only on the rendering thread.
+		ResetAudioRenderingThreadId();
 	}
 
 #if UE_AUDIO_PROFILERTRACE_ENABLED
diff --git a/Engine/Source/Runtime/AudioMixer/Private/AudioMixerSourceBuffer.cpp b/Engine/Source/Runtime/AudioMixer/Private/AudioMixerSourceBuffer.cpp
--- a/Engine/Source/Runtime/AudioMixer/Private/AudioMixerSourceBuffer.cpp
+++ b/Engine/Source/Runtime/AudioMixer/Private/AudioMixerSourceBuffer.cpp
@@ -718,6 +718,11 @@ namespace Audio
 		return true;
 	}
 
+	uint64 FMixerSourceBuffer::GetInstanceID() const
+	{
+		return InstanceID;
+	}
+
 	float FMixerSourceBuffer::GetRelativeRenderCost() const
 	{
 		return RelativeRenderCost.load(std::memory_order_relaxed);
diff --git a/Engine/Source/Runtime/AudioMixer/Private/AudioMixerSourceBuffer.h b/Engine/Source/Runtime/AudioMixer/Private/AudioMixerSourceBuffer.h
--- a/Engine/Source/Runtime/AudioMixer/Private/AudioMixerSourceBuffer.h
+++ b/Engine/Source/Runtime/AudioMixer/Private/AudioMixerSourceBuffer.h
@@ -110,6 +110,9 @@ namespace Audio
 		// Returns true if the async task is done
 		bool IsAsyncTaskDone() const;
 
+		// Get the unique instance ID (transmitter ID) for the source
+		uint64 GetInstanceID() const;
+
 		// Returns some diagnostic state
 		struct FDiagnosticState
 		{
diff --git a/Engine/Source/Runtime/AudioMixer/Private/AudioMixerSourceManager.cpp b/Engine/Source/Runtime/AudioMixer/Private/AudioMixerSourceManager.cpp
--- a/Engine/Source/Runtime/AudioMixer/Private/AudioMixerSourceManager.cpp
+++ b/Engine/Source/Runtime/AudioMixer/Private/AudioMixerSourceManager.cpp
@@ -4,8 +4,10 @@
 
 #include "Audio/AudioTimingLog.h"
 #include "AudioDefines.h"
-#include "AudioMixerSourceBuffer.h"
+#include "AudioMixer.h"
+#include "AudioMixerCVars.h"
 #include "AudioMixerDevice.h"
+#include "AudioMixerSourceBuffer.h"
 #include "AudioMixerSourceVoice.h"
 #include "AudioMixerSubmix.h"
 #include "AudioMixerTrace.h"
@@ -13,7 +15,6 @@
 #include "AudioThread.h"
 #include "DSP/FloatArrayMath.h"
 #include "IAudioExtensionPlugin.h"
-#include "AudioMixer.h"
 #include "Sound/SoundModulationDestination.h"
 #include "SoundFieldRendering.h"
 #include "ProfilingDebugging/CsvProfiler.h"
@@ -192,13 +193,7 @@ FAutoConsoleVariableRef CVarAudioMixerPerSourceResampling(
 	TEXT("0: Disabled, 1: Enabled"),
 	ECVF_Default);
 
-static int32 UseRenderSchedulerCVar = 0;
-FAutoConsoleVariableRef CVarUseAudioRenderScheduler(
-	TEXT("au.UseRenderScheduler"),
-	UseRenderSchedulerCVar,
-	TEXT("Use audio rendering scheduler to handle render dependencies.\n")
-	TEXT("0: Disabled, 1: Enabled"),
-	ECVF_Default);
+
 
 #define ENVELOPE_TAIL_THRESHOLD (1.58489e-5f) // -96 dB
 
@@ -850,7 +845,7 @@ namespace Audio
 			if (AudioBusPtr.IsValid())
 			{
 				// If this audio bus was automatically created via source bus playback, this this audio bus can be removed
-				if (AudioBusPtr->RemoveInstanceId(SourceId))
+				if (AudioBusPtr->RemoveInstanceId(SourceId, SourceInfo.MixerSourceBuffer->GetInstanceID()))
 				{
 #if UE_AUDIO_PROFILERTRACE_ENABLED
 					TraceAudioBusStopMessage(MixerDevice->DeviceID, SourceInfo.AudioBusId);
@@ -893,6 +888,10 @@ namespace Audio
 		// Free the mixer source buffer data
 		if (SourceInfo.MixerSourceBuffer.IsValid())
 		{
+			if (AudioMixerCVars::UseRenderScheduler)
+			{
+				RemoveBusPatchDependencies(SourceInfo.MixerSourceBuffer->GetInstanceID());
+			}
 			PendingSourceBuffers.Add(SourceInfo.MixerSourceBuffer);
 			SourceInfo.MixerSourceBuffer = nullptr;
 		}
@@ -1310,26 +1309,18 @@ namespace Audio
 						SourceInfo.SourceBusDurationFrames = InitParams.SourceBusDuration * MixerDevice->GetSampleRate();
 					}
 
-					// Register this bus as an instance
-					TSharedPtr<FMixerAudioBus> AudioBusPtr = AudioBuses.FindRef(SourceInfo.AudioBusId);
-					if (AudioBusPtr.IsValid())
-					{
-						// If this bus is already registered, add this as a source id
-						AudioBusPtr->AddInstanceId(SourceId, InitParams.NumInputChannels);
-					}
-					else
-					{
-						// If the bus is not registered, make a new entry. This will default to an automatic audio bus until explicitly made manual later.
-						TSharedPtr<FMixerAudioBus> NewAudioBus = TSharedPtr<FMixerAudioBus>(new FMixerAudioBus(this, true, InitParams.AudioBusChannels));
-						NewAudioBus->AddInstanceId(SourceId, InitParams.NumInputChannels);
+					// If the bus is not registered, make a new entry. This will default to an automatic audio bus until explicitly made manual later.
+					bool bAlreadyExisted = false;
+					TSharedRef<FMixerAudioBus> AudioBusPtr = FindOrAddAudioBus(SourceInfo.AudioBusId, InitParams.NumInputChannels, true, &bAlreadyExisted);
+					AudioBusPtr->AddInstanceId(SourceId, SourceInfo.MixerSourceBuffer->GetInstanceID(), InitParams.NumInputChannels);
 
 #if UE_AUDIO_PROFILERTRACE_ENABLED
-						NewAudioBus->BusName = InitParams.AudioBusName;
+					if (!bAlreadyExisted)
+					{
+						AudioBusPtr->BusName = InitParams.AudioBusName;
 						TraceAudioBusStartMessage(MixerDevice->DeviceID, InitParams.AudioBusId, InitParams.AudioBusName, InitParams.AudioBusChannels);
-#endif // UE_AUDIO_PROFILERTRACE_ENABLED
-
-						AudioBuses.Add(InitParams.AudioBusId, NewAudioBus);
 					}
+#endif // UE_AUDIO_PROFILERTRACE_ENABLED
 				}
 
 			}
@@ -1343,32 +1334,23 @@ namespace Audio
 					// New struct to map which source (SourceId) is sending to the bus
 					FAudioBusSend NewAudioBusSend;
 					NewAudioBusSend.SourceId = SourceId;
+					NewAudioBusSend.TransmitterID = SourceInfo.MixerSourceBuffer->GetInstanceID();
 					NewAudioBusSend.SendLevel = AudioBusSend.SendLevel;
 
-					// Get existing BusId and add the send, or create new bus registration
-					TSharedPtr<FMixerAudioBus> AudioBusPtr = AudioBuses.FindRef(AudioBusSend.AudioBusId);
-					if (AudioBusPtr.IsValid())
-					{
-						AudioBusPtr->AddSend((EBusSendType)BusSendType, NewAudioBusSend);
-					}
-					else
-					{
-						// If the bus is not registered, make a new entry. This will default to an automatic audio bus until explicitly made manual later.
-						TSharedPtr<FMixerAudioBus> NewAudioBus(new FMixerAudioBus(this, true, AudioBusSend.BusChannels));
-
-						// Add a send to it. This will not have a bus instance id (i.e. won't output audio), but 
-						// we register the send anyway in the event that this bus does play, we'll know to send this
-						// source's audio to it.
-						NewAudioBus->AddSend((EBusSendType)BusSendType, NewAudioBusSend);
-
+					// Get existing BusId and add the send, or create new bus registration. A newly created bus won't output audio, but 
+					// we register the send anyway in the event that this bus does play, we'll know to send this
+					// source's audio to it.
+					bool bAlreadyExisted = false;
+					TSharedRef<FMixerAudioBus> AudioBusPtr = FindOrAddAudioBus(AudioBusSend.AudioBusId, AudioBusSend.BusChannels, true, &bAlreadyExisted);
+					AudioBusPtr->AddSend((EBusSendType)BusSendType, NewAudioBusSend);
 #if UE_AUDIO_PROFILERTRACE_ENABLED
-						NewAudioBus->BusName = AudioBusSend.AudioBusName;
+					if (!bAlreadyExisted)
+					{
+						AudioBusPtr->BusName = AudioBusSend.AudioBusName;
 						TraceAudioBusStartMessage(MixerDevice->DeviceID, AudioBusSend.AudioBusId, AudioBusSend.AudioBusName, AudioBusSend.BusChannels);
-#endif // UE_AUDIO_PROFILERTRACE_ENABLED
-
-						AudioBuses.Add(AudioBusSend.AudioBusId, NewAudioBus);
 					}
-
+#endif // UE_AUDIO_PROFILERTRACE_ENABLED
+					
 					// Store on this source, which buses its sending its audio to
 					SourceInfo.AudioBusSends[BusSendType].Add(AudioBusSend.AudioBusId);
 				}
@@ -1426,10 +1408,9 @@ namespace Audio
 
 			SourceSubmixOutputBuffer.Reset(SourceSubmixOutputResetSettings);
 
-			const bool bIsBus = InitParams.AudioBusId != INDEX_NONE;
-			if (UseRenderSchedulerCVar && !bIsBus)
+			if (AudioMixerCVars::UseRenderScheduler)
 			{
-				MixerDevice->GetRenderScheduler().AddStep(FAudioRenderStepId::FromWaveInstancePlayOrder(InitParams.PlayOrder), &SourceInfo);
+				MixerDevice->GetRenderScheduler().AddStep(FAudioRenderStepId::FromTransmitterID(SourceInfo.MixerSourceBuffer->GetInstanceID()), &SourceInfo);
 			}
 
 #if AUDIO_MIXER_ENABLE_DEBUG_MODE
@@ -1469,10 +1450,11 @@ namespace Audio
 		{
 			AUDIO_MIXER_CHECK_AUDIO_PLAT_THREAD(MixerDevice);
 
-			const bool bIsBus = SourceInfos[SourceId].AudioBusId != INDEX_NONE;
-			if (UseRenderSchedulerCVar && !bIsBus)
+			if (AudioMixerCVars::UseRenderScheduler)
 			{
-				MixerDevice->GetRenderScheduler().RemoveStep(FAudioRenderStepId::FromWaveInstancePlayOrder(SourceInfos[SourceId].PlayOrder));
+				FSourceInfo& SourceInfo = SourceInfos[SourceId];
+				check(SourceInfo.MixerSourceBuffer.IsValid());
+				MixerDevice->GetRenderScheduler().RemoveStep(FAudioRenderStepId::FromTransmitterID(SourceInfo.MixerSourceBuffer->GetInstanceID()));
 			}
 
 			ReleaseSource(SourceId);
@@ -1500,8 +1482,9 @@ namespace Audio
 		AudioMixerThreadCommand([this, InAudioBusKey, InNumChannels, bInIsAutomatic, InAudioBusName]()
 		{
 			// If this audio bus id already exists, set it to not be automatic and return it
-			TSharedPtr<FMixerAudioBus> AudioBusPtr = AudioBuses.FindRef(InAudioBusKey);
-			if (AudioBusPtr.IsValid())
+			bool bAlreadyExisted = false;
+			TSharedRef<FMixerAudioBus> AudioBusPtr = FindOrAddAudioBus(InAudioBusKey, InNumChannels, bInIsAutomatic, &bAlreadyExisted);
+			if (bAlreadyExisted)
 			{
 				// If this audio bus already existed, make sure the num channels lines up
 				ensure(AudioBusPtr->GetNumChannels() == InNumChannels);
@@ -1509,12 +1492,9 @@ namespace Audio
 			}
 			else
 			{
-				// If the bus is not registered, make a new entry.
-				TSharedPtr<FMixerAudioBus> NewBusData(new FMixerAudioBus(this, bInIsAutomatic, InNumChannels));
 #if UE_AUDIO_PROFILERTRACE_ENABLED
-				NewBusData->BusName = InAudioBusName;
+				AudioBusPtr->BusName = InAudioBusName;
 #endif
-				AudioBuses.Add(InAudioBusKey, NewBusData);
 			}
 		}, AUDIO_MIXER_THREAD_COMMAND_STRING("StartAudioBus()"));
 	}
@@ -2175,38 +2155,28 @@ namespace Audio
 
 		AudioMixerThreadCommand([this, SourceId, InAudioBusSendType, AudioBusId, BusSendLevel, InBusName]()
 		{
+			FSourceInfo& SourceInfo = SourceInfos[SourceId];
+
 			// Create mapping of source id to bus send level
 			FAudioBusSend BusSend;
 			BusSend.SourceId = SourceId;
+			BusSend.TransmitterID = SourceInfo.MixerSourceBuffer->GetInstanceID();
 			BusSend.SendLevel = BusSendLevel;
 
-			FSourceInfo& SourceInfo = SourceInfos[SourceId];
-
-			// Retrieve the bus we want to send audio to
-			TSharedPtr<FMixerAudioBus>* AudioBusPtr = AudioBuses.Find(AudioBusId);
-
-			// If we already have a bus, we update the amount of audio we want to send to it
-			if (AudioBusPtr)
-			{
-				(*AudioBusPtr)->AddSend(InAudioBusSendType, BusSend);
-			}
-			else
-			{
-				// If the bus is not registered, make a new entry on the send
-				TSharedPtr<FMixerAudioBus> NewBusData(new FMixerAudioBus(this, true, SourceInfo.NumInputChannels));
-
-				// Add a send to it. This will not have a bus instance id (i.e. won't output audio), but 
-				// we register the send anyway in the event that this bus does play, we'll know to send this
-				// source's audio to it.
-				NewBusData->AddSend(InAudioBusSendType, BusSend);
+			// Retrieve the bus we want to send audio to. A newly created bus won't output audio, but 
+			// we register the send anyway in the event that this bus does play, we'll know to send this
+			// source's audio to it.
+			bool bAlreadyExisted = false;
+			TSharedRef<FMixerAudioBus> AudioBusPtr = FindOrAddAudioBus(AudioBusId, SourceInfo.NumInputChannels, true, &bAlreadyExisted);
+			AudioBusPtr->AddSend(InAudioBusSendType, BusSend);
 
 #if UE_AUDIO_PROFILERTRACE_ENABLED
-				NewBusData->BusName = InBusName;
+			if (!bAlreadyExisted)
+			{
+				AudioBusPtr->BusName = InBusName;
 				TraceAudioBusStartMessage(MixerDevice->DeviceID, AudioBusId, InBusName, SourceInfo.NumInputChannels);
-#endif // UE_AUDIO_PROFILERTRACE_ENABLED
-
-				AudioBuses.Add(AudioBusId, NewBusData);
 			}
+#endif // UE_AUDIO_PROFILERTRACE_ENABLED
 
 			// Check to see if we need to create new bus data. If we are not playing a bus with this id, then we
 			// need to create a slot for it such that when a bus does play, it'll start rendering audio from this source
@@ -2732,27 +2702,41 @@ namespace Audio
 			int32 NumChannels = PendingAudioBusConnection->NumChannels;
 			bool bIsAutomatic = PendingAudioBusConnection->bIsAutomatic;
 
-			// If this audio bus id already exists, set it to not be automatic and return it
-			TSharedPtr<FMixerAudioBus> AudioBusPtr = AudioBuses.FindRef(AudioBusKey);
-			if (AudioBusPtr.IsValid())
+			bool bAlreadyExisted = false;
+			TSharedRef<FMixerAudioBus> AudioBusPtr = FindOrAddAudioBus(AudioBusKey, NumChannels, bIsAutomatic, &bAlreadyExisted);
+			if (bAlreadyExisted)
 			{
 				// If this audio bus already existed, make sure the num channels lines up
 				ensure(AudioBusPtr->GetNumChannels() == NumChannels);
 				AudioBusPtr->SetAutomatic(bIsAutomatic);
 			}
-			else
-			{
-				// If the bus is not registered, make a new entry.
-				AudioBusPtr = TSharedPtr<FMixerAudioBus>(new FMixerAudioBus(this, bIsAutomatic, NumChannels));
-				AudioBuses.Add(AudioBusKey, AudioBusPtr);
-			}
 
 			switch (PendingAudioBusConnection->PatchVariant.GetIndex())
 			{
 			case FPendingAudioBusConnection::FPatchVariant::IndexOfType<FPatchInput>():
+				if (AudioMixerCVars::UseRenderScheduler)
+				{
+					const uint64 InstanceID = PendingAudioBusConnection->InstanceID;
+					if (InstanceID != INDEX_NONE)
+					{
+						FAudioRenderStepId FirstStep = FAudioRenderStepId::FromTransmitterID(InstanceID);
+						FAudioRenderStepId SecondStep = FAudioRenderStepId::FromAudioBusKey(AudioBusKey);
+						AddBusPatchDependency(InstanceID, FirstStep, SecondStep);
+					}
+				}
 				AudioBusPtr->AddNewPatchInput(PendingAudioBusConnection->PatchVariant.Get<FPatchInput>());
 				break;
 			case FPendingAudioBusConnection::FPatchVariant::IndexOfType<FPatchOutputStrongPtr>():
+				if (AudioMixerCVars::UseRenderScheduler)
+				{
+					const uint64 InstanceID = PendingAudioBusConnection->InstanceID;
+					if (InstanceID != INDEX_NONE)
+					{
+						FAudioRenderStepId FirstStep = FAudioRenderStepId::FromAudioBusKey(AudioBusKey);
+						FAudioRenderStepId SecondStep = FAudioRenderStepId::FromTransmitterID(InstanceID);
+						AddBusPatchDependency(InstanceID, FirstStep, SecondStep);
+					}
+				}
 				AudioBusPtr->AddNewPatchOutput(PendingAudioBusConnection->PatchVariant.Get<FPatchOutputStrongPtr>());
 				break;
 			}
@@ -2822,6 +2806,31 @@ namespace Audio
 		}
 	}
 
+	void FMixerSourceManager::AddBusPatchDependency(const uint64 InstanceID, const FAudioRenderStepId FirstStep, const FAudioRenderStepId SecondStep)
+	{
+		check(AudioMixerCVars::UseRenderScheduler);
+		if (FirstStep.IsValid() && SecondStep.IsValid())
+		{
+			GetRenderScheduler().AddDependency(FirstStep, SecondStep);
+			BusPatchDependencies.Add(InstanceID, { FirstStep, SecondStep });
+		}
+	}
+
+	void FMixerSourceManager::RemoveBusPatchDependencies(const uint64 InstanceID)
+	{
+		check(AudioMixerCVars::UseRenderScheduler);
+		TArray<TPair<FAudioRenderStepId, FAudioRenderStepId>> Dependencies;
+		BusPatchDependencies.MultiFind(InstanceID, Dependencies, false);
+		if (!Dependencies.IsEmpty())
+		{
+			for (TPair<FAudioRenderStepId, FAudioRenderStepId>& Dependency : Dependencies)
+			{
+				GetRenderScheduler().RemoveDependency(Dependency.Get<0>(), Dependency.Get<1>());
+			}
+			BusPatchDependencies.Remove(InstanceID);
+		}
+	}
+
 	float FMixerSourceManager::GetFloatCompareTolerance() const
 	{
 		const float Percent = GetCommandQueueFillPercentage();
@@ -3509,7 +3518,8 @@ namespace Audio
 
 	void FMixerSourceManager::Get2DChannelMap(const int32 SourceId, int32 InNumOutputChannels, Audio::FAlignedFloatBuffer& OutChannelMap)
 	{
-		AUDIO_MIXER_CHECK_AUDIO_PLAT_THREAD(MixerDevice);
+		// This check had to be disabled for now. It is safe to call this both directly from the render thread and also from render tasks.
+		// AUDIO_MIXER_CHECK_AUDIO_PLAT_THREAD(MixerDevice);
 
 		const FSourceInfo& SourceInfo = SourceInfos[SourceId];
 		MixerDevice->Get2DChannelMap(SourceInfo.bIsVorbis, SourceInfo.NumInputChannels, InNumOutputChannels, SourceInfo.bIsCenterChannelOnly, OutChannelMap);
@@ -3763,8 +3773,10 @@ namespace Audio
 		// Update pending tasks and release them if they're finished
 		UpdatePendingReleaseData();
 
-		if (UseRenderSchedulerCVar)
+		if (AudioMixerCVars::UseRenderScheduler)
 		{
+			ConnectBusPatches();
+
 			RenderThreadPhase = ESourceManagerRenderThreadPhase::RunScheduledRenderSteps;
 			MixerDevice->GetRenderScheduler().RenderBlock(DisableParallelSourceProcessingCvar != 0);
 		}
@@ -3772,16 +3784,16 @@ namespace Audio
 		{
 			// First generate non-bus audio (bGenerateBuses = false)
 			GenerateSourceAudio(false);
-		}
 
-		// Now mix in the non-bus audio into the buses
-		ComputeBuses();
+			// Now mix in the non-bus audio into the buses
+			ComputeBuses();
 
-		// Now generate bus audio (bGenerateBuses = true)
-		GenerateSourceAudio(true);
+			// Now generate bus audio (bGenerateBuses = true)
+			GenerateSourceAudio(true);
 
-		// Update the buses now
-		UpdateBuses();
+			// Update the buses now
+			UpdateBuses();
+		}
 
 		// Let the plugin know we finished processing all sources
 		if (bUsingSpatializationPlugin)
@@ -4166,6 +4178,22 @@ namespace Audio
 		LogCallstack(AudioRenderThreadId);
 	}
 
+	TSharedRef<FMixerAudioBus> FMixerSourceManager::FindOrAddAudioBus(const FAudioBusKey BusKey, const int32 InNumChannels, const bool bIsAutomatic, bool* bOutAlreadyExisted)
+	{
+		TSharedPtr<FMixerAudioBus> AudioBusPtr = AudioBuses.FindRef(BusKey);
+		if (bOutAlreadyExisted != nullptr)
+		{
+			*bOutAlreadyExisted = AudioBusPtr.IsValid();
+		}
+		if (!AudioBusPtr.IsValid())
+		{
+			AudioBusPtr = MakeShared<FMixerAudioBus>(this, bIsAutomatic, InNumChannels, BusKey);
+			AudioBuses.Add(BusKey, AudioBusPtr);
+		}
+
+		return AudioBusPtr.ToSharedRef();
+	}
+
 	void FMixerSourceManager::LogCallstack(uint32 InThreadId)
 	{
 		if (InThreadId != INVALID_AUDIO_RENDER_THREAD_ID)
diff --git a/Engine/Source/Runtime/AudioMixer/Private/AudioMixerSourceManager.h b/Engine/Source/Runtime/AudioMixer/Private/AudioMixerSourceManager.h
--- a/Engine/Source/Runtime/AudioMixer/Private/AudioMixerSourceManager.h
+++ b/Engine/Source/Runtime/AudioMixer/Private/AudioMixerSourceManager.h
@@ -8,6 +8,7 @@
 #include "AudioMixerSourceOutputBuffer.h"
 #include "AudioMixerSubmix.h"
 #include "AudioMixerTrace.h"
+#include "AudioRenderScheduler.h"
 #include "Containers/MpscQueue.h"
 #include "DSP/BufferVectorOperations.h"
 #include "DSP/EnvelopeFollower.h"
@@ -317,14 +318,19 @@ namespace Audio
 		// Pushes a TFUnction command into an MPSC queue from an arbitrary thread to the audio render thread
 		void AudioMixerThreadMPSCCommand(TFunction<void()>&& InCommand, const char* InDebugString=nullptr);
 		
-		void AddPendingAudioBusConnection(FAudioBusKey AudioBusKey, int32 NumChannels, bool bIsAutomatic, FPatchInput PatchInput)
+		void AddPendingAudioBusConnection(FAudioBusKey AudioBusKey, int32 NumChannels, bool bIsAutomatic, FPatchInput PatchInput, uint64 InstanceID = INDEX_NONE)
 		{
-			PendingAudioBusConnections.Enqueue(FPendingAudioBusConnection{ FPendingAudioBusConnection::FPatchVariant(TInPlaceType<FPatchInput>(), MoveTemp(PatchInput)), MoveTemp(AudioBusKey), NumChannels, bIsAutomatic });
+			PendingAudioBusConnections.Enqueue(FPendingAudioBusConnection{ FPendingAudioBusConnection::FPatchVariant(TInPlaceType<FPatchInput>(), MoveTemp(PatchInput)), MoveTemp(AudioBusKey), NumChannels, bIsAutomatic, InstanceID });
 		}
 
-		void AddPendingAudioBusConnection(FAudioBusKey AudioBusKey, int32 NumChannels, bool bIsAutomatic, FPatchOutputStrongPtr PatchOutputStrongPtr)
+		void AddPendingAudioBusConnection(FAudioBusKey AudioBusKey, int32 NumChannels, bool bIsAutomatic, FPatchOutputStrongPtr PatchOutputStrongPtr, uint64 InstanceID = INDEX_NONE)
 		{
-			PendingAudioBusConnections.Enqueue(FPendingAudioBusConnection{ FPendingAudioBusConnection::FPatchVariant(TInPlaceType<FPatchOutputStrongPtr>(), MoveTemp(PatchOutputStrongPtr)), MoveTemp(AudioBusKey), NumChannels, bIsAutomatic });
+			PendingAudioBusConnections.Enqueue(FPendingAudioBusConnection{ FPendingAudioBusConnection::FPatchVariant(TInPlaceType<FPatchOutputStrongPtr>(), MoveTemp(PatchOutputStrongPtr)), MoveTemp(AudioBusKey), NumChannels, bIsAutomatic, InstanceID });
+		}
+
+		FAudioRenderScheduler& GetRenderScheduler() const
+		{
+			return MixerDevice->GetRenderScheduler();
 		}
 
 	private:
@@ -345,6 +351,8 @@ namespace Audio
 		void ConnectBusPatches();
 		void ComputeBuses();
 		void UpdateBuses();
+		void AddBusPatchDependency(const uint64 InstanceID, const FAudioRenderStepId FirstStep, const FAudioRenderStepId SecondStep);
+		void RemoveBusPatchDependencies(const uint64 InstanceID);
 
 		float GetFloatCompareTolerance() const;
 		float GetCommandQueueFillPercentage() const;
@@ -626,6 +634,8 @@ namespace Audio
 		void LogCallstacks();
 		void LogCallstack(uint32 InThreadId);
 
+		TSharedRef<FMixerAudioBus> FindOrAddAudioBus(const FAudioBusKey BusKey, const int32 InNumChannels, const bool bIsAutomatic, bool* bOutAlreadyExisted = nullptr);
+
 		// Array of listener transforms
 		TArray<FTransform> ListenerTransforms;
 
@@ -683,10 +693,13 @@ namespace Audio
 			FAudioBusKey AudioBusKey;
 			int32 NumChannels = 0;
 			bool bIsAutomatic = false;
+			uint64 InstanceID = INDEX_NONE;
 		};
 
 		TMpscQueue<FPendingAudioBusConnection> PendingAudioBusConnections;
 
+		TMultiMap<uint64, TPair<FAudioRenderStepId, FAudioRenderStepId>> BusPatchDependencies;
+
 		friend class FMixerSourceVoice;
 	};
 }
diff --git a/Engine/Source/Runtime/AudioMixer/Private/AudioRenderScheduler.cpp b/Engine/Source/Runtime/AudioMixer/Private/AudioRenderScheduler.cpp
--- a/Engine/Source/Runtime/AudioMixer/Private/AudioRenderScheduler.cpp
+++ b/Engine/Source/Runtime/AudioMixer/Private/AudioRenderScheduler.cpp
@@ -5,16 +5,24 @@
 #include "AudioMixerDevice.h"
 #include "AudioMixerTrace.h"
 #include "IAudioMixerRenderStep.h"
+#include "Logging/StructuredLog.h"
 #include "Tasks/Task.h"
 
+DEFINE_LOG_CATEGORY(LogAudioRenderScheduler);
+
 // If MixerDevice is null, don't check the thread. This is used in unit tests.
 #define AUDIO_RENDER_SCHEDULER_CHECK_THREAD() check(MixerDevice == nullptr || MixerDevice->IsAudioRenderingThread())
 
 namespace Audio
 {
-	FAudioRenderStepId FAudioRenderStepId::FromWaveInstancePlayOrder(uint32 PlayOrder)
+	FAudioRenderStepId FAudioRenderStepId::FromTransmitterID(const uint64 TransmitterID)
+	{
+		return TransmitterID;
+	}
+
+	FAudioRenderStepId FAudioRenderStepId::FromAudioBusKey(const FAudioBusKey BusKey)
 	{
-		return uint64(PlayOrder);
+		return (uint64(BusKey.InstanceId + 1) << 32) | BusKey.ObjectId;
 	}
 
 	FAudioRenderScheduler::FAudioRenderScheduler(FMixerDevice* InMixerDevice) :
@@ -79,6 +87,7 @@ namespace Audio
 		FStepEntry& Entry = Steps[Index];
 		if (ensureMsgf(Entry.StepInterface == nullptr, TEXT("Duplicate render step added to audio render scheduler")))
 		{
+			UE_LOGFMT(LogAudioRenderScheduler, Verbose, "Adding render step '{Name}' with id {Id}", Step->GetRenderStepName(), Id.GetRawValue());
 			Entry.StepInterface = Step;
 			Entry.ReferenceCount++;
 			bDirty = true;
@@ -94,6 +103,7 @@ namespace Audio
 		FStepEntry& Entry = Steps[Index];
 		if (ensureMsgf(Entry.StepInterface != nullptr, TEXT("Unknown render step removed from audio render scheduler")))
 		{
+			UE_LOGFMT(LogAudioRenderScheduler, Verbose, "Removing render step '{Name}' with id {Id}", Entry.StepInterface->GetRenderStepName(), Id.GetRawValue());
 			Entry.StepInterface = nullptr;
 			RemoveReference(Index);
 			bDirty = true;
@@ -226,6 +236,23 @@ namespace Audio
 							PrerequisiteHandles.Add(Steps[PrerequisiteIndex].TaskHandle);
 						}
 
+						if (Entry.bCheckBrokenLinks)
+						{
+							// If this step was one end of a broken link, we prevent it from running concurrently with the other
+							//  end of the link.
+							for (const TPair<FInternalIndex, FInternalIndex>& BrokenLink : BrokenLinks)
+							{
+								if (BrokenLink.Get<0>() == ActiveSteps[ActiveStepIndex] && Steps[BrokenLink.Get<1>()].bLaunched)
+								{
+									PrerequisiteHandles.Add(Steps[BrokenLink.Get<1>()].TaskHandle);
+								}
+								if (BrokenLink.Get<1>() == ActiveSteps[ActiveStepIndex] && Steps[BrokenLink.Get<0>()].bLaunched)
+								{
+									PrerequisiteHandles.Add(Steps[BrokenLink.Get<0>()].TaskHandle);
+								}
+							}
+						}
+
 						auto TaskLambda = [&Entry]()
 							{
 								if (Entry.StepInterface != nullptr)
@@ -266,11 +293,13 @@ namespace Audio
 			Steps[Link.Get<1>()].Prerequisites.Add(Link.Get<0>());
 		}
 
-		// Reset the launched state for all steps
+		// Reset the transient state for all steps
 		for (int i = 0; i < Steps.Num(); ++i)
 		{
-			Steps[i].TaskHandle = {};
-			Steps[i].bLaunched = false;
+			FStepEntry& Entry = Steps[i];
+			Entry.TaskHandle = {};
+			Entry.bLaunched = false;
+			Entry.bCheckBrokenLinks = false;
 		}
 
 		// Remember the ordering of steps for next block
@@ -308,7 +337,8 @@ namespace Audio
 			if (VisitedIndices.Contains(PrerequisiteStep))
 			{
 				// The link between CurrentStep and PrerequisiteStep is part of a cycle, so we'll break it.
-				Entry.Prerequisites.RemoveSingleSwap(PrerequisiteStep);
+				UE_LOGFMT(LogAudioRenderScheduler, Verbose, "Breaking cycle between step {Id1} and step {Id2}", Steps[CurrentStep].StepId.GetRawValue(), Steps[PrerequisiteStep].StepId.GetRawValue());
+				BreakLink(PrerequisiteStep, CurrentStep);
 				BrokenLinks.Add(TPair<FInternalIndex, FInternalIndex>{PrerequisiteStep, CurrentStep});
 				return;
 			}
@@ -329,17 +359,24 @@ namespace Audio
 			if (!bDirty || (Steps[Link.Get<1>()].Prerequisites.Contains(Link.Get<0>()) && FindIndirectLink(Link.Get<1>(), Link.Get<0>())))
 			{
 				// This link is still part of a cycle; break it again and leave it in BrokenLinks
-				Steps[Link.Get<1>()].Prerequisites.RemoveSingleSwap(Link.Get<0>());
+				BreakLink(Link.Get<0>(), Link.Get<1>());
 				++LinkIndex;
 			}
 			else
 			{
-				// This link is no longer part of a cycle; remove it and move on.
+				// This link is no longer part of a cycle; remove it and continue.
 				BrokenLinks.RemoveAt(LinkIndex);
 			}
 		}
 	}
 
+	void FAudioRenderScheduler::BreakLink(FInternalIndex FirstStep, FInternalIndex SecondStep)
+	{
+		Steps[SecondStep].Prerequisites.RemoveSingleSwap(FirstStep);
+		Steps[FirstStep].bCheckBrokenLinks = true;
+		Steps[SecondStep].bCheckBrokenLinks = true;
+	}
+
 	bool FAudioRenderScheduler::FindIndirectLink(FInternalIndex FirstStep, FInternalIndex SecondStep)
 	{
 		TArray<FInternalIndex> StepsToSkip;
diff --git a/Engine/Source/Runtime/AudioMixer/Private/AudioRenderScheduler.h b/Engine/Source/Runtime/AudioMixer/Private/AudioRenderScheduler.h
--- a/Engine/Source/Runtime/AudioMixer/Private/AudioRenderScheduler.h
+++ b/Engine/Source/Runtime/AudioMixer/Private/AudioRenderScheduler.h
@@ -2,9 +2,12 @@
 
 #pragma once
 
+#include "AudioBusSubsystem.h"
 #include "Containers/Map.h"
 #include "Tasks/Task.h"
 
+DECLARE_LOG_CATEGORY_EXTERN(LogAudioRenderScheduler, Display, All);
+
 namespace Audio
 {
 	// Forward declarations
@@ -21,7 +24,8 @@ namespace Audio
 		{
 		}
 
-		static FAudioRenderStepId FromWaveInstancePlayOrder(uint32 PlayOrder);
+		static FAudioRenderStepId FromTransmitterID(const uint64 TransmitterID);
+		static FAudioRenderStepId FromAudioBusKey(const FAudioBusKey BusKey);
 
 		uint64 GetRawValue() const
 		{
@@ -101,6 +105,7 @@ namespace Audio
 
 		void BreakCycle(const int FirstPendingStepIndex);
 		void CheckBrokenLinks();
+		void BreakLink(FInternalIndex FirstStep, FInternalIndex SecondStep);
 
 		// Helpers to verify that a link is still part of a cycle
 		bool FindIndirectLink(const FInternalIndex FirstStep, const FInternalIndex SecondStep);
@@ -113,6 +118,7 @@ namespace Audio
 			TArray<FInternalIndex, TInlineAllocator<4>> Prerequisites;
 			int32 ReferenceCount = 0;
 			bool bLaunched = false;
+			bool bCheckBrokenLinks = false;
 			UE::Tasks::FTask TaskHandle;
 
 			// Unused step entries are added to a free list.
diff --git a/Engine/Source/Runtime/AudioMixer/Private/AudioRenderSchedulerTester.cpp b/Engine/Source/Runtime/AudioMixer/Private/AudioRenderSchedulerTester.cpp
--- a/Engine/Source/Runtime/AudioMixer/Private/AudioRenderSchedulerTester.cpp
+++ b/Engine/Source/Runtime/AudioMixer/Private/AudioRenderSchedulerTester.cpp
@@ -49,18 +49,18 @@ namespace Audio
 			if (!Steps.Contains(Id))
 			{
 				Steps.Emplace(Id, MakeUnique<FTestStep>(*this, Id));
-				Scheduler.AddStep(FAudioRenderStepId::FromWaveInstancePlayOrder(Id), Steps[Id].Get());
+				Scheduler.AddStep(FAudioRenderStepId::FromTransmitterID(Id), Steps[Id].Get());
 			}
 		}
 
 		virtual void AddDependency(int FirstStep, int SecondStep) override
 		{
-			Scheduler.AddDependency(FAudioRenderStepId::FromWaveInstancePlayOrder(FirstStep), FAudioRenderStepId::FromWaveInstancePlayOrder(SecondStep));
+			Scheduler.AddDependency(FAudioRenderStepId::FromTransmitterID(FirstStep), FAudioRenderStepId::FromTransmitterID(SecondStep));
 		}
 
 		virtual void RemoveDependency(int FirstStep, int SecondStep) override
 		{
-			Scheduler.RemoveDependency(FAudioRenderStepId::FromWaveInstancePlayOrder(FirstStep), FAudioRenderStepId::FromWaveInstancePlayOrder(SecondStep));
+			Scheduler.RemoveDependency(FAudioRenderStepId::FromTransmitterID(FirstStep), FAudioRenderStepId::FromTransmitterID(SecondStep));
 		}
 
 		virtual TArray<int> Run() override
diff --git a/Engine/Source/Runtime/AudioMixer/Public/AudioMixerBus.h b/Engine/Source/Runtime/AudioMixer/Public/AudioMixerBus.h
--- a/Engine/Source/Runtime/AudioMixer/Public/AudioMixerBus.h
+++ b/Engine/Source/Runtime/AudioMixer/Public/AudioMixerBus.h
@@ -2,10 +2,11 @@
 
 #pragma once
 
-#include "CoreMinimal.h"
 #include "Audio.h"
+#include "AudioBusSubsystem.h"
 #include "AudioMixerTrace.h"
 #include "DSP/MultithreadedPatching.h"
+#include "IAudioMixerRenderStep.h"
 #include "Sound/AudioSettings.h"
 
 #if UE_AUDIO_PROFILERTRACE_ENABLED
@@ -20,13 +21,24 @@ namespace Audio
 	struct FAudioBusSend
 	{
 		int32 SourceId = INDEX_NONE;
+		uint64 TransmitterID = INDEX_NONE;
 		float SendLevel = 0.0f;
 	};
 
 	// Bus instance data. Holds source id bus instances and bus sends data
-	class FMixerAudioBus
+	class FMixerAudioBus : public IAudioMixerRenderStep
 	{
 	public:
+		// Creates an audio bus.
+		// SourceMnager		The owning source manager object.
+		// bInIsAutomatic	Whether or not this audio bus was created automatically via source buses.
+		// InNumChannels	The number of channels of the source bus.
+		// InNumFrames		The number of frames to mix per mix buffer call. I.e. source manager render size.
+		// InBusKey         The unique key used to identify this bus.
+		FMixerAudioBus(FMixerSourceManager* SourceManager, bool bInIsAutomatic, int32 InNumChannels, const FAudioBusKey InBusKey);
+
+		~FMixerAudioBus();
+
 		// Allow anybody to add a pre-existing patch output object to the audio bus
 		void AddNewPatchOutput(const FPatchOutputStrongPtr& InPatchOutputStrongPtr);
 
@@ -37,14 +49,7 @@ namespace Audio
 		void RemovePatchInput(const FPatchInput& InPatchInput);
 
 	private:
-
-		// Creates an audio bus.
-		// SourceMnager		The owning source manager object.
-		// bInIsAutomatic	Whether or not this audio bus was created automatically via source buses.
-		// InNumChannels	The number of channels of the source bus.
-		// InNumFrames		The number of frames to mix per mix buffer call. I.e. source manager render size.
-		FMixerAudioBus(FMixerSourceManager* SourceManager, bool bInIsAutomatic, int32 InNumChannels);
-		
+	
 		// Sets whether or not this audio bus is automatic.
 		void SetAutomatic(bool bInIsAutomatic) { bIsAutomatic = bInIsAutomatic; }
 
@@ -58,10 +63,10 @@ namespace Audio
 		void Update();
 
 		// Adds a source id for instances of this bus
-		void AddInstanceId(const int32 InSourceInstanceId, int32 InNumOutputChannels);
+		void AddInstanceId(const int32 InSourceInstanceId, const uint64 InTransmitterID, int32 InNumOutputChannels);
 
 		// Removes the source id from this bus. Returns true if there are no more instances or sends.
-		bool RemoveInstanceId(const int32 InSourceId);
+		bool RemoveInstanceId(const int32 InSourceId, const uint64 InTransmitterID);
 
 		// Adds a bus send to the bus
 		void AddSend(EBusSendType BusSendType, const FAudioBusSend& InBusSend);
@@ -85,6 +90,10 @@ namespace Audio
 		// If this bus was constructed before
 		void SetNumOutputChannels(int32 InNumOutputChannels);
 
+		// IAudioMixerRenderStep interface
+		virtual void DoRenderStep() override;
+		virtual const TCHAR* GetRenderStepName() override;
+
 #if UE_AUDIO_PROFILERTRACE_ENABLED
 		void StartEnvelopeFollower(const float InAttackTime, const float InReleaseTime, const float InSampleRate);
 		void StopEnvelopeFollower();
@@ -113,13 +122,16 @@ namespace Audio
 		// The number of output frames
 		int32 NumFrames;
 
-		// Owning soruce manager
+		// Owning source manager
 		FMixerSourceManager* SourceManager;
 
 		// Multiple places can produce and consume from audio buses
 		Audio::FPatchMixer PatchMixer;
 		Audio::FPatchSplitter PatchSplitter;
 
+		// Unique bus identifier
+		FAudioBusKey BusKey;
+
 		// Was created manually, not via source buses.
 		bool bIsAutomatic;
 
diff --git a/Engine/Source/Runtime/AudioMixer/Public/AudioMixerDevice.h b/Engine/Source/Runtime/AudioMixer/Public/AudioMixerDevice.h
--- a/Engine/Source/Runtime/AudioMixer/Public/AudioMixerDevice.h
+++ b/Engine/Source/Runtime/AudioMixer/Public/AudioMixerDevice.h
@@ -483,7 +483,10 @@ namespace Audio
 
 		TMap<uint32, TArray<FSourceEffectChainEntry>> SourceEffectChainOverrides;
 
-		/** The mixer source manager. */
+		/** Scheduler responsible for launching rendering tasks */
+		TUniquePtr<FAudioRenderScheduler> RenderScheduler;
+
+		/** The mixer source manager. Should be cleaned up before the scheduler so dependencies can be removed correctly. */
 		TUniquePtr<FMixerSourceManager> SourceManager;
 
 		/** ThreadId for the game thread (or if audio is running a separate thread, that ID) */
@@ -514,9 +517,6 @@ namespace Audio
 		/** Whether or not initialization of the submix system is underway and submixes can be registered */
 		bool bSubmixRegistrationDisabled;
 
-		/** Scheduler responsible for launching rendering tasks */
-		TUniquePtr<FAudioRenderScheduler> RenderScheduler;
-
 	public:
 
 		// Creates a queue for audio decode requests with a specific Id. Tasks
