Add explicit cast to try to fix static analysis warning

diff --git a/Engine/Source/Runtime/AudioMixer/Private/AudioRenderScheduler.cpp b/Engine/Source/Runtime/AudioMixer/Private/AudioRenderScheduler.cpp
--- a/Engine/Source/Runtime/AudioMixer/Private/AudioRenderScheduler.cpp
+++ b/Engine/Source/Runtime/AudioMixer/Private/AudioRenderScheduler.cpp
@@ -22,7 +22,8 @@ namespace Audio
 
 	FAudioRenderStepId FAudioRenderStepId::FromAudioBusKey(const FAudioBusKey BusKey)
 	{
-		return (uint64(BusKey.InstanceId + 1) << 32) | BusKey.ObjectId;
+		// BusKey.InstanceId is usually -1. We add 1 just to make this equal to the object id in that usual case.
+		return (uint64(uint32(BusKey.InstanceId + 1)) << 32) | BusKey.ObjectId;
 	}
 
 	FAudioRenderScheduler::FAudioRenderScheduler(FMixerDevice* InMixerDevice) :
