[Audio Render Scheduler] Fix comments, channel count bug

#rnx

diff --git a/Engine/Source/Runtime/AudioMixer/Private/AudioMixerSourceManager.cpp b/Engine/Source/Runtime/AudioMixer/Private/AudioMixerSourceManager.cpp
--- a/Engine/Source/Runtime/AudioMixer/Private/AudioMixerSourceManager.cpp
+++ b/Engine/Source/Runtime/AudioMixer/Private/AudioMixerSourceManager.cpp
@@ -1311,7 +1311,7 @@ namespace Audio
 
 					// If the bus is not registered, make a new entry. This will default to an automatic audio bus until explicitly made manual later.
 					bool bAlreadyExisted = false;
-					TSharedRef<FMixerAudioBus> AudioBusPtr = FindOrAddAudioBus(SourceInfo.AudioBusId, InitParams.NumInputChannels, true, &bAlreadyExisted);
+					TSharedRef<FMixerAudioBus> AudioBusPtr = FindOrAddAudioBus(SourceInfo.AudioBusId, InitParams.AudioBusChannels, true, &bAlreadyExisted);
 					AudioBusPtr->AddInstanceId(SourceId, SourceInfo.MixerSourceBuffer->GetInstanceID(), InitParams.NumInputChannels);
 
 #if UE_AUDIO_PROFILERTRACE_ENABLED
diff --git a/Engine/Source/Runtime/AudioMixer/Public/AudioMixerBus.h b/Engine/Source/Runtime/AudioMixer/Public/AudioMixerBus.h
--- a/Engine/Source/Runtime/AudioMixer/Public/AudioMixerBus.h
+++ b/Engine/Source/Runtime/AudioMixer/Public/AudioMixerBus.h
@@ -30,10 +30,9 @@ namespace Audio
 	{
 	public:
 		// Creates an audio bus.
-		// SourceMnager		The owning source manager object.
+		// SourceManager		The owning source manager object.
 		// bInIsAutomatic	Whether or not this audio bus was created automatically via source buses.
 		// InNumChannels	The number of channels of the source bus.
-		// InNumFrames		The number of frames to mix per mix buffer call. I.e. source manager render size.
 		// InBusKey         The unique key used to identify this bus.
 		FMixerAudioBus(FMixerSourceManager* SourceManager, bool bInIsAutomatic, int32 InNumChannels, const FAudioBusKey InBusKey);
 
