[Audio Render Scheduler] Enable scheduler by default. Can be disabled by setting "au.UseRenderScheduler 0"

#tests Replay runs plus manual testing across several experiences
#rnx

diff --git a/Engine/Source/Runtime/AudioMixer/Private/AudioMixerCVars.cpp b/Engine/Source/Runtime/AudioMixer/Private/AudioMixerCVars.cpp
--- a/Engine/Source/Runtime/AudioMixer/Private/AudioMixerCVars.cpp
+++ b/Engine/Source/Runtime/AudioMixer/Private/AudioMixerCVars.cpp
@@ -6,7 +6,7 @@
 
 namespace AudioMixerCVars
 {
-	int32 UseRenderScheduler = 0;
+	int32 UseRenderScheduler = 1;
 	FAutoConsoleVariableRef CVarUseAudioRenderScheduler(
 		TEXT("au.UseRenderScheduler"),
 		UseRenderScheduler,
