[AudioMixer] Added EnvelopeFollowerEnabled trace messages for submixes and audio buses, these trace messages will be used in a further CL.

diff --git a/Engine/Source/Runtime/AudioMixer/Private/AudioMixerSourceManager.cpp b/Engine/Source/Runtime/AudioMixer/Private/AudioMixerSourceManager.cpp
--- a/Engine/Source/Runtime/AudioMixer/Private/AudioMixerSourceManager.cpp
+++ b/Engine/Source/Runtime/AudioMixer/Private/AudioMixerSourceManager.cpp
@@ -291,6 +291,13 @@ UE_TRACE_EVENT_BEGIN(Audio, AudioBusHasActivity)
 	UE_TRACE_EVENT_FIELD(bool, HasActivity)
 UE_TRACE_EVENT_END()
 
+UE_TRACE_EVENT_BEGIN(Audio, AudioBusEnvelopeFollowerEnabled)
+	UE_TRACE_EVENT_FIELD(uint32, DeviceId)
+	UE_TRACE_EVENT_FIELD(uint32, AudioBusId)
+	UE_TRACE_EVENT_FIELD(uint64, Timestamp)
+	UE_TRACE_EVENT_FIELD(bool, EnvelopeFollowerEnabled)
+UE_TRACE_EVENT_END()
+
 UE_TRACE_EVENT_BEGIN(Audio, AudioBusEnvelopeValues)
 	UE_TRACE_EVENT_FIELD(uint32, DeviceId)
 	UE_TRACE_EVENT_FIELD(uint32, AudioBusId)
@@ -513,6 +520,17 @@ namespace Audio
 			if (AudioBusKey.ObjectId == InAudioBusId && AudioBus.IsValid())
 			{
 				AudioBus->StartEnvelopeFollower(AttackTime, ReleaseTime, MixerDevice->GetSampleRate());
+
+				const bool bChannelEnabled = UE_TRACE_CHANNELEXPR_IS_ENABLED(AudioChannel);
+				if (bChannelEnabled && AudioBus->bIsEnvelopeFollowing)
+				{
+					UE_TRACE_LOG(Audio, AudioBusEnvelopeFollowerEnabled, AudioChannel)
+						<< AudioBusEnvelopeFollowerEnabled.DeviceId(MixerDevice->DeviceID)
+						<< AudioBusEnvelopeFollowerEnabled.AudioBusId(AudioBusKey.ObjectId)
+						<< AudioBusEnvelopeFollowerEnabled.Timestamp(FPlatformTime::Cycles64())
+						<< AudioBusEnvelopeFollowerEnabled.EnvelopeFollowerEnabled(true);
+				}
+
 				break;
 			}
 		}
@@ -527,6 +545,17 @@ namespace Audio
 			if (AudioBusKey.ObjectId == InAudioBusId && AudioBus.IsValid())
 			{
 				AudioBus->StopEnvelopeFollower();
+
+				const bool bChannelEnabled = UE_TRACE_CHANNELEXPR_IS_ENABLED(AudioChannel);
+				if (bChannelEnabled && !AudioBus->bIsEnvelopeFollowing)
+				{
+					UE_TRACE_LOG(Audio, AudioBusEnvelopeFollowerEnabled, AudioChannel)
+						<< AudioBusEnvelopeFollowerEnabled.DeviceId(MixerDevice->DeviceID)
+						<< AudioBusEnvelopeFollowerEnabled.AudioBusId(AudioBusKey.ObjectId)
+						<< AudioBusEnvelopeFollowerEnabled.Timestamp(FPlatformTime::Cycles64())
+						<< AudioBusEnvelopeFollowerEnabled.EnvelopeFollowerEnabled(false);
+				}
+
 				break;
 			}
 		}
diff --git a/Engine/Source/Runtime/AudioMixer/Private/AudioMixerSubmix.cpp b/Engine/Source/Runtime/AudioMixer/Private/AudioMixerSubmix.cpp
--- a/Engine/Source/Runtime/AudioMixer/Private/AudioMixerSubmix.cpp
+++ b/Engine/Source/Runtime/AudioMixer/Private/AudioMixerSubmix.cpp
@@ -75,6 +75,13 @@ UE_TRACE_EVENT_BEGIN(Audio, SubmixHasActivity)
 	UE_TRACE_EVENT_FIELD(bool, HasActivity)
 UE_TRACE_EVENT_END()
 
+UE_TRACE_EVENT_BEGIN(Audio, SubmixEnvelopeFollowerEnabled)
+	UE_TRACE_EVENT_FIELD(uint32, DeviceId)
+	UE_TRACE_EVENT_FIELD(uint32, SubmixId)
+	UE_TRACE_EVENT_FIELD(uint64, Timestamp)
+	UE_TRACE_EVENT_FIELD(bool, EnvelopeFollowerEnabled)
+UE_TRACE_EVENT_END()
+
 UE_TRACE_EVENT_BEGIN(Audio, SubmixEnvelopeValues)
 	UE_TRACE_EVENT_FIELD(uint32, DeviceId)
 	UE_TRACE_EVENT_FIELD(uint32, SubmixId)
@@ -2300,12 +2307,36 @@ namespace Audio
 			}
 
 			bIsEnvelopeFollowing = true;
+
+#if UE_AUDIO_PROFILERTRACE_ENABLED
+			const bool bChannelEnabled = UE_TRACE_CHANNELEXPR_IS_ENABLED(AudioChannel);
+			if (bChannelEnabled && OwningSubmixObject.IsValid())
+			{
+				UE_TRACE_LOG(Audio, SubmixEnvelopeFollowerEnabled, AudioChannel)
+					<< SubmixEnvelopeFollowerEnabled.DeviceId(MixerDevice->DeviceID)
+					<< SubmixEnvelopeFollowerEnabled.SubmixId(GetTypeHash(OwningSubmixObject->GetPathName()))
+					<< SubmixEnvelopeFollowerEnabled.Timestamp(FPlatformTime::Cycles64())
+					<< SubmixEnvelopeFollowerEnabled.EnvelopeFollowerEnabled(true);
+			}
+#endif // UE_AUDIO_PROFILERTRACE_ENABLED
 		}
 	}
 
 	void FMixerSubmix::StopEnvelopeFollowing()
 	{
 		bIsEnvelopeFollowing = false;
+
+#if UE_AUDIO_PROFILERTRACE_ENABLED
+		const bool bChannelEnabled = UE_TRACE_CHANNELEXPR_IS_ENABLED(AudioChannel);
+		if (bChannelEnabled && OwningSubmixObject.IsValid())
+		{
+			UE_TRACE_LOG(Audio, SubmixEnvelopeFollowerEnabled, AudioChannel)
+				<< SubmixEnvelopeFollowerEnabled.DeviceId(MixerDevice->DeviceID)
+				<< SubmixEnvelopeFollowerEnabled.SubmixId(GetTypeHash(OwningSubmixObject->GetPathName()))
+				<< SubmixEnvelopeFollowerEnabled.Timestamp(FPlatformTime::Cycles64())
+				<< SubmixEnvelopeFollowerEnabled.EnvelopeFollowerEnabled(false);
+		}
+#endif // UE_AUDIO_PROFILERTRACE_ENABLED
 	}
 
 	void FMixerSubmix::AddEnvelopeFollowerDelegate(const FOnSubmixEnvelopeBP& OnSubmixEnvelopeBP)
