Fix copy-paste errors

#rnx

#DisabledValidations CheckUnacceptableWords

diff --git a/Engine/Source/Runtime/AudioMixer/Private/AudioMixerSubmix.cpp b/Engine/Source/Runtime/AudioMixer/Private/AudioMixerSubmix.cpp
--- a/Engine/Source/Runtime/AudioMixer/Private/AudioMixerSubmix.cpp
+++ b/Engine/Source/Runtime/AudioMixer/Private/AudioMixerSubmix.cpp
@@ -454,7 +454,7 @@ namespace Audio
 				{
 					EndpointSettings = SoundfieldFactory->GetDefaultEndpointSettings();
 
-					if (!ensureMsgf(EncodingSettings, TEXT("Soundfield Endpoint %s did not return default encoding settings! Is ISoundfieldEndpointFactory::GetDefaultEndpointSettings() implemented?"), *SoundfieldFactory->GetEndpointTypeName().ToString()))
+					if (!ensureMsgf(EndpointSettings, TEXT("Soundfield Endpoint %s did not return default endpoint settings! Is ISoundfieldEndpointFactory::GetDefaultEndpointSettings() implemented?"), *SoundfieldFactory->GetEndpointTypeName().ToString()))
 					{
 						return;
 					}
@@ -2544,7 +2544,7 @@ namespace Audio
 		}
 		else
 		{
-			UE_LOG(LogAudioMixer, Warning, TEXT("Call StartSpectrumAnalysis before calling GetMagnitudeForFrequencies."));
+			UE_LOG(LogAudioMixer, Warning, TEXT("Call StartSpectrumAnalysis before calling GetPhaseForFrequencies."));
 		}
 	}
 
