[AudioMixer] Fix for audio meter getting stuck on submix disable (a silent envelope trace message is sent when a submix is disabled so we prevent audio meters in Audio Insights from getting stuck at its previous envelope values)

diff --git a/Engine/Source/Runtime/AudioMixer/Private/AudioMixerSubmix.cpp b/Engine/Source/Runtime/AudioMixer/Private/AudioMixerSubmix.cpp
--- a/Engine/Source/Runtime/AudioMixer/Private/AudioMixerSubmix.cpp
+++ b/Engine/Source/Runtime/AudioMixer/Private/AudioMixerSubmix.cpp
@@ -1307,6 +1307,21 @@ namespace Audio
 				WetLevelMod.ClearModulators();
 				DryLevelMod.ClearModulators();
 				
+#if UE_AUDIO_PROFILERTRACE_ENABLED
+				const bool bChannelEnabled = UE_TRACE_CHANNELEXPR_IS_ENABLED(AudioChannel);
+				if (bChannelEnabled && OwningSubmixObject.IsValid())
+				{
+					TArray<float> SilentEnvelopeValues;
+					SilentEnvelopeValues.Init(MIN_VOLUME_DECIBELS, NumChannels);
+
+					UE_TRACE_LOG(Audio, SubmixEnvelopeValues, AudioChannel)
+						<< SubmixEnvelopeValues.DeviceId(MixerDevice->DeviceID)
+						<< SubmixEnvelopeValues.SubmixId(GetTypeHash(OwningSubmixObject->GetPathName()))
+						<< SubmixEnvelopeValues.Timestamp(FPlatformTime::Cycles64())
+						<< SubmixEnvelopeValues.EnvelopeValues(SilentEnvelopeValues.GetData(), NumChannels);
+				}
+#endif // UE_AUDIO_PROFILERTRACE_ENABLED
+
 				UE_CLOG(LogSubmixEnablementCVar == 1, LogAudioMixer, Verbose,
 					TEXT("Submix Disabled. Num Sources: %d, Time Silent: %.2f, Disablement Threshold: %.2f, Submix Name: %s"), 
 					MixerSourceVoices.Num(),
