[AudioMixer] Fix for issue detected by BugHawk in CL 46203528

diff --git a/Engine/Source/Runtime/AudioMixer/Private/AudioMixerSubmix.cpp b/Engine/Source/Runtime/AudioMixer/Private/AudioMixerSubmix.cpp
--- a/Engine/Source/Runtime/AudioMixer/Private/AudioMixerSubmix.cpp
+++ b/Engine/Source/Runtime/AudioMixer/Private/AudioMixerSubmix.cpp
@@ -1312,7 +1312,7 @@ namespace Audio
 				if (bChannelEnabled && OwningSubmixObject.IsValid())
 				{
 					TArray<float> SilentEnvelopeValues;
-					SilentEnvelopeValues.Init(MIN_VOLUME_DECIBELS, NumChannels);
+					SilentEnvelopeValues.Init(0.0f, NumChannels);
 
 					UE_TRACE_LOG(Audio, SubmixEnvelopeValues, AudioChannel)
 						<< SubmixEnvelopeValues.DeviceId(MixerDevice->DeviceID)
