Add a CVar to control whether the callstack is dumped to the log during Audio stall diagnostics.

#rnx
#FailedValidations Preflight

diff --git a/Engine/Source/Runtime/AudioMixer/Private/AudioMixerSourceManager.cpp b/Engine/Source/Runtime/AudioMixer/Private/AudioMixerSourceManager.cpp
--- a/Engine/Source/Runtime/AudioMixer/Private/AudioMixerSourceManager.cpp
+++ b/Engine/Source/Runtime/AudioMixer/Private/AudioMixerSourceManager.cpp
@@ -163,6 +163,14 @@ FAutoConsoleVariableRef CVarFloatCompareMaxScale(
 	TEXT("Maximum parameter float compare tolerance"),
 	ECVF_Default);
 
+
+static bool bDumpCallstackOnStallDiagnostics = true;
+FAutoConsoleVariableRef CVarDumpCallstackOnStallDiagnostics(
+	TEXT("au.DumpCallstackOnStallDiagnostics"),
+	bDumpCallstackOnStallDiagnostics,
+	TEXT("Whether to dump the callstack to the log as part of stall diagnostics (default = true)"),
+	ECVF_Default);
+
 // +/- 4 Octaves (default)
 static float MaxModulationPitchRangeFreqCVar = 16.0f;
 static float MinModulationPitchRangeFreqCVar = 0.0625f;
@@ -4145,7 +4153,10 @@ namespace Audio
 	{
 		LogRenderThreadStall();
 		LogInflightAsyncTasks();
-		LogCallstacks();
+		if (bDumpCallstackOnStallDiagnostics)
+		{
+			LogCallstacks();
+		}
 	}
 
 	void FMixerSourceManager::LogRenderThreadStall()
