1. Avoid preallocating memory for SourceVoiceBuffers TArray in FMixerSourceBuffer
2. Preallocate TouchedSubmixes TSet in a FMixerSource::UpdateEffects to a know max number

diff --git a/Engine/Source/Runtime/AudioMixer/Private/AudioMixerSource.cpp b/Engine/Source/Runtime/AudioMixer/Private/AudioMixerSource.cpp
--- a/Engine/Source/Runtime/AudioMixer/Private/AudioMixerSource.cpp
+++ b/Engine/Source/Runtime/AudioMixer/Private/AudioMixerSource.cpp
@@ -1609,8 +1609,8 @@ namespace Audio
 		// Including the submix itself. 
 		// It's okay to use "previous" submix here as it's set above or from a previous setting.
 		TSet<FMixerSubmixWeakPtr> TouchedSubmixes;
+		TouchedSubmixes.Reserve(1 + WaveInstance->AttenuationSubmixSends.Num() + WaveInstance->SoundSubmixSends.Num());
 		TouchedSubmixes.Add(PreviousSubmixResolved);
-	
 
 		// Attenuation Submix Sends. (these come from Attenuation assets).
 		// These are largely identical to SoundSubmix Sends, but don't specify a send stage, so we pass one here.		
diff --git a/Engine/Source/Runtime/AudioMixer/Private/AudioMixerSourceBuffer.cpp b/Engine/Source/Runtime/AudioMixer/Private/AudioMixerSourceBuffer.cpp
--- a/Engine/Source/Runtime/AudioMixer/Private/AudioMixerSourceBuffer.cpp
+++ b/Engine/Source/Runtime/AudioMixer/Private/AudioMixerSourceBuffer.cpp
@@ -129,13 +129,10 @@ namespace Audio
 		// Only allocate one buffer to render into when doing direct rendering
 		const int32 NumBuffers = bDirectRendering ? 1 : Audio::MAX_BUFFERS_QUEUED;
 
-		const uint32 TotalSamples = (bDirectRendering && SoundGenerator.IsValid()) ? SoundGenerator->GetDesiredNumSamplesToRenderPerCallback() : MONO_PCM_BUFFER_SAMPLES * NumChannels;
+		SourceVoiceBuffers.Reserve(NumBuffers);
 		for (int32 BufferIndex = 0; BufferIndex < NumBuffers; ++BufferIndex)
 		{
 			SourceVoiceBuffers.Add(MakeShared<FAlignedFloatBuffer, ESPMode::ThreadSafe>());
-
-			// Prepare the memory to fit the max number of samples
-			SourceVoiceBuffers[BufferIndex]->Reset(TotalSamples);
 		}
 	}
 
