[AudioMixer] Moved silent audio buffer check logic to its own function since it is needed in more than one place. This check now applies to submixes where autodisable is not checked (only if audio traces are enabled).

diff --git a/Engine/Source/Runtime/AudioMixer/Private/AudioMixerSubmix.cpp b/Engine/Source/Runtime/AudioMixer/Private/AudioMixerSubmix.cpp
--- a/Engine/Source/Runtime/AudioMixer/Private/AudioMixerSubmix.cpp
+++ b/Engine/Source/Runtime/AudioMixer/Private/AudioMixerSubmix.cpp
@@ -193,6 +193,48 @@ namespace Audio
 		}
 	}
 
+	namespace FMixerSubmixPrivate
+	{
+		bool IsAudioBufferSilent(const FAlignedFloatBuffer& InAudioBuffer)
+		{
+			// Extremely cheap silent buffer detection: as soon as we hit a sample which isn't silent, we flag we're not silent
+			bool bIsSilent = true;
+
+			int32 Index = 0;
+
+#if PLATFORM_ENABLE_VECTORINTRINSICS
+			const int32 SimdNum = InAudioBuffer.Num() & 0xFFFFFFF0;
+			for (; Index < SimdNum; Index += 16)
+			{
+				const VectorRegister4x4Float Samples = VectorLoad16(&InAudioBuffer[Index]);
+
+				if (VectorAnyGreaterThan(VectorAbs(Samples.val[0]), GlobalVectorConstants::SmallNumber) ||
+					VectorAnyGreaterThan(VectorAbs(Samples.val[1]), GlobalVectorConstants::SmallNumber) ||
+					VectorAnyGreaterThan(VectorAbs(Samples.val[2]), GlobalVectorConstants::SmallNumber) ||
+					VectorAnyGreaterThan(VectorAbs(Samples.val[3]), GlobalVectorConstants::SmallNumber))
+				{
+					bIsSilent = false;
+					Index = INT_MAX;
+					break;
+				}
+			}
+#endif
+
+			// Finish to the end of the buffer or check each sample if vector intrinsics are disabled
+			for (; Index < InAudioBuffer.Num(); ++Index)
+			{
+				// As soon as we hit a non-silent sample, we're not silent
+				if (FMath::Abs(InAudioBuffer[Index]) > SMALL_NUMBER)
+				{
+					bIsSilent = false;
+					break;
+				}
+			}
+
+			return bIsSilent;
+		}
+	}
+
 	// Unique IDs for mixer submixes
 	static uint32 GSubmixMixerIDs = 0;
 
@@ -1646,38 +1688,15 @@ namespace Audio
 		Audio::ArrayMixIn(InputBuffer, OutAudioBuffer);
 
 		// Once we've finished rendering submix audio, check if the output buffer is silent if we are auto-disabling
+		bool bIsNowSilent = false;
+
+#if UE_AUDIO_PROFILERTRACE_ENABLED
+		const bool bShouldTraceHasActivityMsg = UE_TRACE_CHANNELEXPR_IS_ENABLED(AudioChannel) && OwningSubmixObject.IsValid();
+#endif // UE_AUDIO_PROFILERTRACE_ENABLED
+
 		if (bAutoDisable)
 		{
-			// Extremely cheap silent buffer detection: as soon as we hit a sample which isn't silent, we flag we're not silent
-			bool bIsNowSilent = true;
-			int i = 0;
-#if PLATFORM_ENABLE_VECTORINTRINSICS
-			int SimdNum = OutAudioBuffer.Num() & 0xFFFFFFF0;
-			for (; i < SimdNum; i += 16)
-			{
-				VectorRegister4x4Float Samples = VectorLoad16(&OutAudioBuffer[i]);
-				if (   VectorAnyGreaterThan(VectorAbs(Samples.val[0]), GlobalVectorConstants::SmallNumber)
-					|| VectorAnyGreaterThan(VectorAbs(Samples.val[1]), GlobalVectorConstants::SmallNumber)
-					|| VectorAnyGreaterThan(VectorAbs(Samples.val[2]), GlobalVectorConstants::SmallNumber)
-					|| VectorAnyGreaterThan(VectorAbs(Samples.val[3]), GlobalVectorConstants::SmallNumber))
-				{
-					bIsNowSilent = false;
-					i = INT_MAX;
-					break;
-				}
-			}
-#endif
-
-			// Finish to the end of the buffer or check each sample if vector intrinsics are disabled
-			for (; i < OutAudioBuffer.Num(); ++i)
-			{
-				// As soon as we hit a non-silent sample, we're not silent
-				if (FMath::Abs(OutAudioBuffer[i]) > SMALL_NUMBER)
-				{
-					bIsNowSilent = false;
-					break;
-				}
-			}
+			bIsNowSilent = FMixerSubmixPrivate::IsAudioBufferSilent(OutAudioBuffer);
 
 			// If this is the first time we're silent track when it happens
 			if (bIsNowSilent && !bIsSilent)
@@ -1690,20 +1709,27 @@ namespace Audio
 				bIsSilent = false;
 				SilenceTimeStartSeconds = -1.0;
 			}
-
+		}
+		else
+		{
 #if UE_AUDIO_PROFILERTRACE_ENABLED
-			const bool bChannelEnabled = UE_TRACE_CHANNELEXPR_IS_ENABLED(AudioChannel);
-			if (bChannelEnabled && OwningSubmixObject.IsValid())
+			if (bShouldTraceHasActivityMsg)
 			{
-				UE_TRACE_LOG(Audio, SubmixHasActivity, AudioChannel)
-					<< SubmixHasActivity.DeviceId(MixerDevice->DeviceID)
-					<< SubmixHasActivity.SubmixId(GetTypeHash(OwningSubmixObject->GetPathName()))
-					<< SubmixHasActivity.Timestamp(FPlatformTime::Cycles64())
-					<< SubmixHasActivity.HasActivity(!bIsNowSilent);
+				bIsNowSilent = FMixerSubmixPrivate::IsAudioBufferSilent(OutAudioBuffer);
 			}
 #endif // UE_AUDIO_PROFILERTRACE_ENABLED
 		}
 
+#if UE_AUDIO_PROFILERTRACE_ENABLED
+		if (bShouldTraceHasActivityMsg)
+		{
+			UE_TRACE_LOG(Audio, SubmixHasActivity, AudioChannel)
+				<< SubmixHasActivity.DeviceId(MixerDevice->DeviceID)
+				<< SubmixHasActivity.SubmixId(GetTypeHash(OwningSubmixObject->GetPathName()))
+				<< SubmixHasActivity.Timestamp(FPlatformTime::Cycles64())
+				<< SubmixHasActivity.HasActivity(!bIsNowSilent);
+		}
+#endif // UE_AUDIO_PROFILERTRACE_ENABLED
 	}
 
 	void FMixerSubmix::SendAudioToSubmixBufferListeners(FAlignedFloatBuffer& OutAudioBuffer)
