[Backout] - CL44993360 Return to 37.30 version of XAudio2 driver for 37.50
#rnx
Original CL Desc
-----------------------------------------------------------------
Improvements to xbox resume logic.

#was-backout 44902692
#rnx

diff --git a/Engine/Config/Windows/BaseWindowsEngine.ini b/Engine/Config/Windows/BaseWindowsEngine.ini
--- a/Engine/Config/Windows/BaseWindowsEngine.ini
+++ b/Engine/Config/Windows/BaseWindowsEngine.ini
@@ -1,7 +1,7 @@
 [Audio]
 ; Defining below allows switching to audio mixer using -audiomixer commandline
 AudioCaptureModuleName=AudioCaptureWasapi
-AudioMixerModuleName=AudioMixerWasapi
+AudioMixerModuleName=AudioMixerXAudio2
 PlatformFormat=BINKA
 PlatformStreamingFormat=BINKA
 
diff --git a/Engine/Source/Runtime/AudioDeviceEnumeration/Windows/WindowsMMDeviceEnumeration/Private/WindowsMMDeviceInfoCache.cpp b/Engine/Source/Runtime/AudioDeviceEnumeration/Windows/WindowsMMDeviceEnumeration/Private/WindowsMMDeviceInfoCache.cpp
--- a/Engine/Source/Runtime/AudioDeviceEnumeration/Windows/WindowsMMDeviceEnumeration/Private/WindowsMMDeviceInfoCache.cpp
+++ b/Engine/Source/Runtime/AudioDeviceEnumeration/Windows/WindowsMMDeviceEnumeration/Private/WindowsMMDeviceInfoCache.cpp
@@ -88,7 +88,8 @@ namespace Audio
 	{
 		ensure(SUCCEEDED(CoCreateInstance(__uuidof(MMDeviceEnumerator), NULL, CLSCTX_INPROC_SERVER, IID_PPV_ARGS(&DeviceEnumerator))) && DeviceEnumerator);
 		
-		Refresh();
+		EnumerateEndpoints();
+		EnumerateDefaults();
 	}
 
 	FWindowsMMDeviceCache::FWindowsMMDeviceCache() :
@@ -96,11 +97,6 @@ namespace Audio
 	{
 	}
 
-	void FWindowsMMDeviceCache::Refresh()
-	{
-		EnumerateEndpoints();
-		EnumerateDefaults();
-	}
 
 	bool FWindowsMMDeviceCache::EnumerateChannelMask(uint32 InMask, FCacheEntry& OutInfo)
 	{
@@ -897,17 +893,6 @@ namespace Audio
 		return 0;
 	}
 
-	FName FWindowsMMDeviceCache::GetDefaultOutputDevice(EAudioDeviceRole InRole) const
-	{
-		FReadScopeLock MapReadLock(CacheMutationLock);
-
-		if (ensure(InRole >= EAudioDeviceRole::Console && InRole < EAudioDeviceRole::COUNT))
-		{
-			return DefaultRenderId[(int32)InRole];
-		}
-		return NAME_None;
-	}
-
 	FName FWindowsMMDeviceCache::GetDefaultOutputDevice_NoLock() const
 	{
 		if (!DefaultRenderId[(int32)EAudioDeviceRole::Console].IsNone())
diff --git a/Engine/Source/Runtime/AudioDeviceEnumeration/Windows/WindowsMMDeviceEnumeration/Public/WindowsMMDeviceInfoCache.h b/Engine/Source/Runtime/AudioDeviceEnumeration/Windows/WindowsMMDeviceEnumeration/Public/WindowsMMDeviceInfoCache.h
--- a/Engine/Source/Runtime/AudioDeviceEnumeration/Windows/WindowsMMDeviceEnumeration/Public/WindowsMMDeviceInfoCache.h
+++ b/Engine/Source/Runtime/AudioDeviceEnumeration/Windows/WindowsMMDeviceEnumeration/Public/WindowsMMDeviceInfoCache.h
@@ -107,11 +107,7 @@ struct FWindowsMMDeviceCache : IAudioMixerDeviceChangedListener, IAudioPlatformD
 	UE_API virtual TArray<FAudioPlatformDeviceInfo> GetLogicalAggregateDevices(const FName InHardwareId, const EDeviceEndpointType InEndpointType) const override;
 
 	UE_API virtual TArray<FAudioPlatformDeviceInfo> SynthesizeAggregateDeviceList(const EDeviceEndpointType InType) const;
-
-	UE_API virtual void Refresh() override;
 	
-	UE_API virtual FName GetDefaultOutputDevice(EAudioDeviceRole InRole) const override;
-
 	UE_API FName GetDefaultOutputDevice_NoLock() const;
 
 	UE_API TOptional<FAudioPlatformDeviceInfo> FindDefaultOutputDevice() const override;
diff --git a/Engine/Source/Runtime/AudioMixer/Private/AudioMixerDevice.cpp b/Engine/Source/Runtime/AudioMixer/Private/AudioMixerDevice.cpp
--- a/Engine/Source/Runtime/AudioMixer/Private/AudioMixerDevice.cpp
+++ b/Engine/Source/Runtime/AudioMixer/Private/AudioMixerDevice.cpp
@@ -973,51 +973,44 @@ namespace Audio
 
 		AudioMixerPlatform->OnHardwareUpdate();
 
-		// Device swaps span multiple calls beginning with CheckAudioDeviceChange() and ending with
-		// ResumePlaybackOnNewDevice(). We protect the associated resources by locking and unlocking around
-		// the entire call chain.
+		if (AudioMixerPlatform->CheckAudioDeviceChange())
 		{
-			FScopedDeviceSwapLock Lock(AudioMixerPlatform);
+			// Get the platform device info we're using
+			PlatformInfo = AudioMixerPlatform->GetPlatformDeviceInfo();
 
-			if (AudioMixerPlatform->CheckAudioDeviceChange())
-			{
-				// Get the platform device info we're using
-				PlatformInfo = AudioMixerPlatform->GetPlatformDeviceInfo();
-
-				// Initialize some data that depends on speaker configuration, etc.
-				InitializeChannelAzimuthMap(PlatformInfo.NumChannels);
-
-				// Update the channel device count in case it changed
-				SourceManager->UpdateDeviceChannelCount(PlatformInfo.NumChannels);
+			// Initialize some data that depends on speaker configuration, etc.
+			InitializeChannelAzimuthMap(PlatformInfo.NumChannels);
 
-				// Reset rendering thread ID to this thread ID so that commands can
-				// be flushed. Audio Rendering Thread ID will be reset again in call
-				// to FMixerDevice::OnProcessAudio
-				ResetAudioRenderingThreadId();
+			// Update the channel device count in case it changed
+			SourceManager->UpdateDeviceChannelCount(PlatformInfo.NumChannels);
 
-				// Cache the audio platform thread id for debugging purposes. This is
-				// an attempt to narrow down the causes of UE-209237. The theory is
-				// that the there are multiple threads attempting to run audio rendering
-				// commands. To catch the issue we attempt to cache the thread id before
-				// running audio rendering commands. We then check that the thread id
-				// has not changed for the duration that render thread commands have 
-				// run.
-				std::atomic<int32> CurrentAudioPlatformThreadId = AudioPlatformThreadId.load();
+			// Reset rendering thread ID to this thread ID so that commands can
+			// be flushed. Audio Rendering Thread ID will be reset again in call
+			// to FMixerDevice::OnProcessAudio
+			ResetAudioRenderingThreadId();
+			
+			// Cache the audio platform thread id for debugging purposes. This is
+			// an attempt to narrow down the causes of UE-209237. The theory is
+			// that the there are multiple threads attempting to run audio rendering
+			// commands. To catch the issue we attempt to cache the thread id before
+			// running audio rendering commands. We then check that the thread id
+			// has not changed for the duration that render thread commands have 
+			// run.
+			std::atomic<int32> CurrentAudioPlatformThreadId = AudioPlatformThreadId.load();
 
-				// Force source manager to incorporate device channel count change.
-				FlushAudioRenderingCommands(true /* bPumpSynchronously */);
+			// Force source manager to incorporate device channel count change.
+			FlushAudioRenderingCommands(true /* bPumpSynchronously */);
 
-				if (UAudioDeviceNotificationSubsystem* AudioDeviceNotifSubsystem = UAudioDeviceNotificationSubsystem::Get())
-				{
-					AudioDeviceNotifSubsystem->OnDeviceSwitched(PlatformInfo.DeviceId);
-				}
+			if (UAudioDeviceNotificationSubsystem* AudioDeviceNotifSubsystem = UAudioDeviceNotificationSubsystem::Get())
+			{
+				AudioDeviceNotifSubsystem->OnDeviceSwitched(PlatformInfo.DeviceId);
+			}
 
-				// Related to earlier mention of UE-209237
-				UE_CLOG(CurrentAudioPlatformThreadId != AudioPlatformThreadId, LogAudioMixer, Error, TEXT("Platform audio thread id changed while flushing render commands. Expected %d, found %d. May result in corrupt internal audio source state."), CurrentAudioPlatformThreadId.load(), AudioPlatformThreadId.load());
+			// Related to earlier mention of UE-209237
+			UE_CLOG(CurrentAudioPlatformThreadId != AudioPlatformThreadId, LogAudioMixer, Error, TEXT("Platform audio thread id changed while flushing render commands. Expected %d, found %d. May result in corrupt internal audio source state."), CurrentAudioPlatformThreadId.load(), AudioPlatformThreadId.load());
 
-				// Audio rendering was suspended in CheckAudioDeviceChange if it changed.
-				AudioMixerPlatform->ResumePlaybackOnNewDevice();
-			}
+			// Audio rendering was suspended in CheckAudioDeviceChange if it changed.
+			AudioMixerPlatform->ResumePlaybackOnNewDevice();
 		}
 
 		// Device must be initialized prior to call as submix graph may not be ready yet otherwise.
diff --git a/Engine/Source/Runtime/AudioMixerCore/Private/AudioMixer.cpp b/Engine/Source/Runtime/AudioMixerCore/Private/AudioMixer.cpp
--- a/Engine/Source/Runtime/AudioMixerCore/Private/AudioMixer.cpp
+++ b/Engine/Source/Runtime/AudioMixerCore/Private/AudioMixer.cpp
@@ -723,11 +723,13 @@ namespace Audio
 			{
 				check(AudioStreamInfo.StreamState == EAudioOutputStreamState::Stopped);
 			}
+			else
+			{
+				AudioStreamInfo.StreamState = EAudioOutputStreamState::Stopped;
+			}
 
 			AudioRenderThread.Reset();
 		}
-		
-		AudioStreamInfo.StreamState = EAudioOutputStreamState::Stopped;
 
 		if (AudioRenderEvent != nullptr)
 		{
@@ -952,10 +954,18 @@ namespace Audio
 
 	bool FAudioMixerPlatformSwappable::RequestDeviceSwap(const FString& DeviceID, const bool bInForce, const TCHAR* InReason)
 	{
-		// This critical section protects that device swap context below.
-		FScopeLock Lock(&DeviceSwapCriticalSection);
+		if (!AllowDeviceSwap(bInForce))
+		{
+			UE_LOG(LogAudioMixer, Display, TEXT("NOT-ALLOWING attempt to swap audio render device to new device: '%s', because: '%s', force=%d"),
+				!DeviceID.IsEmpty() ? *DeviceID : TEXT("[System Default]"),
+				InReason ? InReason : TEXT("None specified"),
+				(int32)bInForce
+			);
+			return false;
+		}
 
-		if (AllowDeviceSwap(bInForce))
+		// This critical section protects that device swap context below.
+		if (DeviceSwapCriticalSection.TryLock())
 		{
 			UE_LOG(LogAudioMixer, Display, TEXT("Attempt to swap audio render device to new device: '%s', because: '%s', force=%d"),
 				!DeviceID.IsEmpty() ? *DeviceID : TEXT("[System Default]"),
@@ -968,18 +978,18 @@ namespace Audio
 			{
 				// Set the flag indicating we wish to begin a device swap
 				AudioStreamInfo.StreamState = EAudioOutputStreamState::SwappingDevice;
-				return true;
 			}
-		}
-		else
-		{
-			UE_LOG(LogAudioMixer, Display, TEXT("NOT-ALLOWING attempt to swap audio render device to new device: '%s', because: '%s', force=%d [normal if occurring while suspended on consoles or if force=false and another swap is in-flight]"),
-				!DeviceID.IsEmpty() ? *DeviceID : TEXT("[System Default]"),
-				InReason ? InReason : TEXT("None specified"),
-				(int32)bInForce
-			);
+		
+			DeviceSwapCriticalSection.Unlock();
+
+			return true;
 		}
 		
+		UE_LOG(LogAudioMixer, Display, TEXT("Dropping device swap request; device: '%s', because: '%s', force=%d"),
+		       !DeviceID.IsEmpty() ? *DeviceID : TEXT("[System Default]"),
+		       InReason ? InReason : TEXT("None specified"),
+		       (int32)bInForce);
+
 		return false;
 	}
 
@@ -1014,18 +1024,15 @@ namespace Audio
 			return false;
 		}
 		
-		// If the device is not running or if a device swap is already in progress, don't "double-trigger" a swap unless 
-		// bInForceSwap is true, in which case we'll replace the existing request.
-		if (ActiveDeviceSwap.IsValid() && !bInForceSwap)
+		// If the device is not running or if a device swap is already in progress, don't "double-trigger" a swap.
+		if (ActiveDeviceSwap.IsValid())
 		{
-			UE_LOG(LogAudioMixer, Display, TEXT("Ignoring device swap request, swap already requested. bInForceSwap: %d"), bInForceSwap);
+			UE_LOG(LogAudioMixer, Display, TEXT("Ignoring device swap request, swap already requested."));
 			return false;
 		}
 
-		// Only device swap if the audio device is running, stopped or mid-swap
-		if (AudioStreamInfo.StreamState != EAudioOutputStreamState::Running && 
-			AudioStreamInfo.StreamState != EAudioOutputStreamState::Stopped &&
-			AudioStreamInfo.StreamState != EAudioOutputStreamState::SwappingDevice)
+		// Only device swap if the audio device is running
+		if (AudioStreamInfo.StreamState != EAudioOutputStreamState::Running)
 		{
 			UE_LOG(LogAudioMixer, Display, TEXT("Ignoring device swap request, AudioStreamInfo.StreamState: %d"), AudioStreamInfo.StreamState.load());
 			return false;
@@ -1139,8 +1146,7 @@ namespace Audio
 	void FAudioMixerPlatformSwappable::ResumePlaybackOnNewDevice()
 	{
 		SCOPED_NAMED_EVENT(FAudioMixerPlatformSwappable_ResumePlaybackOnNewDevice, FColor::Blue);
-		UE_LOG(LogAudioMixer, Display, TEXT("FAudioMixerPlatformSwappable::ResumePlaybackOnNewDevice - resuming audio on new device: '%s'"), *AudioStreamInfo.DeviceInfo.Name);
-
+		
 		int32 NumSamplesPopped = 0;
 		TArrayView<const uint8> PoppedAudio = OutputBuffer.PopBufferData(NumSamplesPopped);
 		SubmitBuffer(PoppedAudio.GetData());
diff --git a/Engine/Source/Runtime/AudioMixerCore/Public/AudioMixer.h b/Engine/Source/Runtime/AudioMixerCore/Public/AudioMixer.h
--- a/Engine/Source/Runtime/AudioMixerCore/Public/AudioMixer.h
+++ b/Engine/Source/Runtime/AudioMixerCore/Public/AudioMixer.h
@@ -227,25 +227,6 @@ namespace Audio
 
 	enum class EDeviceEndpointType { Unknown, Render, Capture };
 
-	enum class EAudioDeviceRole
-	{
-		Console,
-		Multimedia,
-		Communications,
-
-		COUNT,
-	};
-
-	enum class EAudioDeviceState
-	{
-		Active,
-		Disabled,
-		NotPresent,
-		Unplugged,
-
-		COUNT,
-	};
-
 	// Interface for Caching Device Info.
 	class IAudioPlatformDeviceInfoCache
 	{
@@ -253,7 +234,6 @@ namespace Audio
 		// Pure Interface. 
 		virtual ~IAudioPlatformDeviceInfoCache() = default;
 			
-		virtual FName GetDefaultOutputDevice(EAudioDeviceRole InRole) const = 0;
 		virtual TOptional<FAudioPlatformDeviceInfo> FindActiveOutputDevice(FName InDeviceID) const = 0;
 		virtual TArray<FAudioPlatformDeviceInfo> GetAllActiveOutputDevices() const = 0;
 		virtual TOptional<FAudioPlatformDeviceInfo> FindDefaultOutputDevice() const = 0;
@@ -263,9 +243,6 @@ namespace Audio
 
 		/** Returns an array of logical audio devices which, when combined, form an aggregate device. */
 		virtual TArray<FAudioPlatformDeviceInfo> GetLogicalAggregateDevices(const FName InHardwareId, const EDeviceEndpointType InEndpointType) const = 0;
-
-		/** Refresh the data in the cache */
-		virtual void Refresh() = 0;
 	};
 
 	/** Defines parameters needed for opening a new audio stream to device. */
@@ -350,6 +327,25 @@ namespace Audio
 		}
 	};
 
+	enum class EAudioDeviceRole
+	{
+		Console,
+		Multimedia,
+		Communications,
+
+		COUNT,
+	};
+
+	enum class EAudioDeviceState
+	{
+		Active,
+		Disabled,
+		NotPresent,
+		Unplugged,
+
+		COUNT,
+	};
+
 	/** Struct used to store render time analysis data. */
 	struct FAudioRenderTimeAnalysis
 	{
@@ -549,12 +545,6 @@ namespace Audio
 		/** Sends a command to swap which output device is being used */
 		virtual bool RequestDeviceSwap(const FString& DeviceID, bool bInForce, const TCHAR* InReason = nullptr) { return false; }
 
-		/** Protects device swap related resources during device swaps */
-		virtual void LockDeviceSwap() {}
-
-		/** Unprotects device swap resources. Balance with LockDeviceSwap() */
-		virtual void UnlockDeviceSwap() {}
-
 		/** Returns the platform device info of the currently open audio stream. */
 		virtual FAudioPlatformDeviceInfo GetPlatformDeviceInfo() const = 0;
 
@@ -816,8 +806,6 @@ namespace Audio
 
 		//~ Begin IAudioMixerPlatformInterface
 		UE_API virtual bool RequestDeviceSwap(const FString& DeviceID, const bool bInForce, const TCHAR* InReason) override;
-		virtual void LockDeviceSwap() override { DeviceSwapCriticalSection.Lock(); }
-		virtual void UnlockDeviceSwap() override { DeviceSwapCriticalSection.Unlock(); }
 		UE_API virtual bool CheckAudioDeviceChange() override;
 		UE_API virtual bool MoveAudioStreamToNewAudioDevice() override;
 		UE_API virtual void ResumePlaybackOnNewDevice() override;
@@ -870,30 +858,6 @@ namespace Audio
 		double LastDeviceSwapTime = 0.0;
 	};
 
-	class FScopedDeviceSwapLock final
-	{
-	public:
-		FScopedDeviceSwapLock() = delete;
-		explicit FScopedDeviceSwapLock(IAudioMixerPlatformInterface* InPlatform)
-			: AudioMixerPlatform(InPlatform)
-		{
-			if (AudioMixerPlatform)
-			{
-				AudioMixerPlatform->LockDeviceSwap();
-			}
-		}
-
-		~FScopedDeviceSwapLock()
-		{
-			if (AudioMixerPlatform)
-			{
-				AudioMixerPlatform->UnlockDeviceSwap();
-			}
-		}
-
-	private:
-		IAudioMixerPlatformInterface* AudioMixerPlatform = nullptr;
-	};
 }
 
 /**
diff --git a/Engine/Source/Runtime/Windows/AudioMixerWasapi/Private/AudioMixerWasapi.cpp b/Engine/Source/Runtime/Windows/AudioMixerWasapi/Private/AudioMixerWasapi.cpp
--- a/Engine/Source/Runtime/Windows/AudioMixerWasapi/Private/AudioMixerWasapi.cpp
+++ b/Engine/Source/Runtime/Windows/AudioMixerWasapi/Private/AudioMixerWasapi.cpp
@@ -688,6 +688,9 @@ namespace Audio
 	
 	bool FAudioMixerWasapi::PostDeviceSwap()
 	{
+		// Once the context is handed off to the device swap routine (either async or synchronous),
+		// it should no longer be valid.
+		check(!DeviceSwapContext.IsValid());
 		bool bDidSucceed = false;
 
 		FWasapiDeviceSwapResult* DeviceSwapResult = StaticCast<FWasapiDeviceSwapResult*>(GetDeviceSwapResult());
@@ -745,17 +748,4 @@ namespace Audio
 		return bDidSucceed;
 	}
 	
-	void FAudioMixerWasapi::ResumePlaybackOnNewDevice()
-	{
-		if (DeviceSwapContext.IsValid())
-		{
-			// If another device swap was requested, don't bother resuming playback
-			AudioStreamInfo.StreamState = EAudioOutputStreamState::SwappingDevice;
-			UE_LOG(LogAudioMixer, Display, TEXT("FAudioMixerWasapi::ResumePlaybackOnNewDevice - device swap requested during previous device swap...not resuming audio"));
-		}
-		else
-		{
-			FAudioMixerPlatformSwappable::ResumePlaybackOnNewDevice();
-		}
-	}
 }
diff --git a/Engine/Source/Runtime/Windows/AudioMixerWasapi/Private/AudioMixerWasapi.h b/Engine/Source/Runtime/Windows/AudioMixerWasapi/Private/AudioMixerWasapi.h
--- a/Engine/Source/Runtime/Windows/AudioMixerWasapi/Private/AudioMixerWasapi.h
+++ b/Engine/Source/Runtime/Windows/AudioMixerWasapi/Private/AudioMixerWasapi.h
@@ -81,7 +81,6 @@ namespace Audio
 		virtual IAudioPlatformDeviceInfoCache* GetDeviceInfoCache() const override;
 		virtual bool IsDeviceInfoValid(const FAudioPlatformDeviceInfo& InDeviceInfo) const override;
 		virtual bool ShouldUseDeviceInfoCache() const override { return true; }
-		virtual void ResumePlaybackOnNewDevice() override;
 		//~ End IAudioMixerPlatformInterface
 
 		//~ Begin IAudioMixerDeviceChangedListener
@@ -111,9 +110,6 @@ namespace Audio
 		 *  FAudioPlatformDeviceInfo rather than looking it up via the requested device Id.
 		 */
 		bool InitDeviceSwapContextInternal(const FString& InRequestedDeviceID, const TCHAR* InReason, const TOptional<FAudioPlatformDeviceInfo>& InDeviceInfo);
-
-		/** Returns the device role to be used when swapping to the default device. */
-		virtual EAudioDeviceRole GetDefaultDeviceRole() const { return EAudioDeviceRole::Console; }
 		
 	private:
 
diff --git a/Engine/Source/Runtime/Windows/AudioMixerWasapi/Private/AudioMixerWasapiDeviceListener.cpp b/Engine/Source/Runtime/Windows/AudioMixerWasapi/Private/AudioMixerWasapiDeviceListener.cpp
--- a/Engine/Source/Runtime/Windows/AudioMixerWasapi/Private/AudioMixerWasapiDeviceListener.cpp
+++ b/Engine/Source/Runtime/Windows/AudioMixerWasapi/Private/AudioMixerWasapiDeviceListener.cpp
@@ -96,10 +96,7 @@ namespace Audio
 		// However, the Windows UX only allows console+multimedia to be toggle as a pair. This means you get two notifications
 		// for default device changing typically. To prevent a trouble trigger we only listen to "Console" here. For more information on 
 		// device roles: https://docs.microsoft.com/en-us/windows/win32/coreaudio/device-roles
-
-		// Note: On XBox, headphones plugged into an XBox controller give a device notification on the Communications role. However, you 
-		// should not swap to a device for a role other than the default Console, so we re-swap to the Console default here and let
-		// the Xbox OS take care of the routing details. 
+		// Note: On XBox, headphones plugged into an XBox controller will use the Communications role.
 		
 		if (InAudioDeviceRole == EAudioDeviceRole::Console || InAudioDeviceRole == EAudioDeviceRole::Communications)
 		{
@@ -108,20 +105,9 @@ namespace Audio
 
 			// Ignore if not listening for events. This means the user has specified a specific device
 			// to use and is not interested in following the current system default device.
-
-			FString DeviceIdForSwap;
-#if PLATFORM_WINDOWS
-			DeviceIdForSwap = DeviceId;
-#else
-			if (IAudioPlatformDeviceInfoCache* Cache = GetDeviceInfoCache())
-			{
-				DeviceIdForSwap = Cache->GetDefaultOutputDevice(GetDefaultDeviceRole()).ToString();
-			}
-#endif // PLATFORM_WINDOWS
-			
 			if (GetIsListeningForDeviceEvents())
 			{
-				RequestDeviceSwap(DeviceIdForSwap, /* force */true, TEXT("FAudioMixerWasapi::OnDefaultRenderDeviceChanged"));
+				RequestDeviceSwap(DeviceId, /* force */true, TEXT("FAudioMixerWasapi::OnDefaultRenderDeviceChanged"));
 			}
 		}
 
@@ -209,17 +195,17 @@ namespace Audio
 			// to use and is not interested in following the current system default device.
 			if (GetIsListeningForDeviceEvents())
 			{
-				FString DefaultDeviceId;
-				if (IAudioPlatformDeviceInfoCache* Cache = GetDeviceInfoCache())
-				{
-					const FName DeviceIdName = Cache->GetDefaultOutputDevice(GetDefaultDeviceRole());
-					if (DeviceIdName != NAME_None)
-					{
-						DefaultDeviceId = DeviceIdName.ToString();
-					}
-				}
-
-				RequestDeviceSwap(DefaultDeviceId, /* force */ true, TEXT("FAudioMixerWasapi::OnDeviceStateChanged"));
+				//FString DefaultDeviceId;
+				//if (IAudioPlatformDeviceInfoCache* Cache = GetDeviceInfoCache())
+				//{
+				//	const FName DeviceIdName = Cache->GetDefaultOutputDevice(GetDefaultDeviceRole());
+				//	if (DeviceIdName != NAME_None)
+				//	{
+				//		DefaultDeviceId = DeviceIdName.ToString();
+				//	}
+				//}
+
+				//RequestDeviceSwap(DefaultDeviceId, /* force */ true, TEXT("FAudioMixerWasapi::OnDeviceStateChanged"));
 			}
 		}
 		
diff --git a/Engine/Source/Runtime/Windows/AudioMixerWasapi/Private/AudioMixerWasapiRenderStream.cpp b/Engine/Source/Runtime/Windows/AudioMixerWasapi/Private/AudioMixerWasapiRenderStream.cpp
--- a/Engine/Source/Runtime/Windows/AudioMixerWasapi/Private/AudioMixerWasapiRenderStream.cpp
+++ b/Engine/Source/Runtime/Windows/AudioMixerWasapi/Private/AudioMixerWasapiRenderStream.cpp
@@ -50,7 +50,6 @@ namespace Audio
 			return false;
 		}
 
-#if PLATFORM_WINDOWS
 		AudioClientProperties ClientProperties = {};
 		ClientProperties.cbSize = sizeof(AudioClientProperties);
 		ClientProperties.bIsOffload = false;
@@ -64,7 +63,6 @@ namespace Audio
 			// Not all clients support raw mode. Log warning and continue.
 			UE_LOG(LogAudioMixerWasapi, Warning, TEXT("InitializeHardware failed SetClientProperties %s"), *AudioClientErrorToFString(Result));
 		}
-#endif //PLATFORM_WINDOWS
 		
 		WAVEFORMATEX* MixFormat = nullptr;
 		Result = TempAudioClient->GetMixFormat(&MixFormat);
diff --git a/Engine/Source/Runtime/Windows/AudioMixerXAudio2/Private/AudioMixerPlatformXAudio2.cpp b/Engine/Source/Runtime/Windows/AudioMixerXAudio2/Private/AudioMixerPlatformXAudio2.cpp
--- a/Engine/Source/Runtime/Windows/AudioMixerXAudio2/Private/AudioMixerPlatformXAudio2.cpp
+++ b/Engine/Source/Runtime/Windows/AudioMixerXAudio2/Private/AudioMixerPlatformXAudio2.cpp
@@ -226,6 +226,9 @@ namespace Audio
 	
 	bool FMixerPlatformXAudio2::PostDeviceSwap()
 	{
+		// Once the context is handed off to the device swap routine (either async or synchronous),
+		// it should no longer be valid.
+		check(!DeviceSwapContext.IsValid());
 		bool bDidSucceed = false;
 		const FXAudio2DeviceSwapResult* DeviceSwapResult = static_cast<const FXAudio2DeviceSwapResult*>(GetDeviceSwapResult());
 		
@@ -280,19 +283,6 @@ namespace Audio
 		return bDidSucceed;
 	}
 	
-	void FMixerPlatformXAudio2::ResumePlaybackOnNewDevice()
-	{
-		if (DeviceSwapContext.IsValid())
-		{
-			// If another device swap was requested, don't bother resuming playback
-			AudioStreamInfo.StreamState = EAudioOutputStreamState::SwappingDevice;
-			UE_LOG(LogAudioMixer, Display, TEXT("FMixerPlatformXAudio2::ResumePlaybackOnNewDevice - device swap requested during previous device swap...not resuming audio"));
-		}
-		else
-		{
-			FAudioMixerPlatformSwappable::ResumePlaybackOnNewDevice();
-		}
-	}
 
 	void FMixerPlatformXAudio2::SynchronousDeviceSwap()
 	{
diff --git a/Engine/Source/Runtime/Windows/AudioMixerXAudio2/Private/AudioMixerPlatformXAudio2.h b/Engine/Source/Runtime/Windows/AudioMixerXAudio2/Private/AudioMixerPlatformXAudio2.h
--- a/Engine/Source/Runtime/Windows/AudioMixerXAudio2/Private/AudioMixerPlatformXAudio2.h
+++ b/Engine/Source/Runtime/Windows/AudioMixerXAudio2/Private/AudioMixerPlatformXAudio2.h
@@ -122,7 +122,6 @@ namespace Audio
 		virtual IAudioPlatformDeviceInfoCache* GetDeviceInfoCache() const override;
 		virtual bool IsDeviceInfoValid(const FAudioPlatformDeviceInfo& InDeviceInfo) const override;
 		virtual bool ShouldUseDeviceInfoCache() const override { return true; }
-		virtual void ResumePlaybackOnNewDevice() override;
 		//~ End IAudioMixerPlatformInterface
 
 		//~ Begin IAudioMixerDeviceChangedListener
