Make Submix Buffer Listeners wait for Submixes to be registered instead of forcing them on

diff --git a/Engine/Source/Runtime/AudioMixer/Private/AudioMixerDevice.cpp b/Engine/Source/Runtime/AudioMixer/Private/AudioMixerDevice.cpp
--- a/Engine/Source/Runtime/AudioMixer/Private/AudioMixerDevice.cpp
+++ b/Engine/Source/Runtime/AudioMixer/Private/AudioMixerDevice.cpp
@@ -126,6 +126,13 @@ FAutoConsoleVariableRef CVarMuteAudio(
 	TEXT("1: Mutes the final submix. 0: Does not mute the final submix."),
 	ECVF_Cheat);
 
+static int32 SubmixBufferListenerForceRegistrationCVar = 1;
+FAutoConsoleVariableRef CVarSubmixBufferListenerForceRegistration(
+	TEXT("au.SubmixBufferListenerForceRegistration"),
+	SubmixBufferListenerForceRegistrationCVar,
+	TEXT("If true, will force register the referenced submix when Registering a Submix Buffer Listener."),
+	ECVF_Default);
+
 #if UE_AUDIO_PROFILERTRACE_ENABLED
 
 UE_TRACE_EVENT_BEGIN(Audio, SubmixAlivePing)
@@ -2153,6 +2160,23 @@ namespace Audio
 		{
 			RebuildSubmixLinks(*InSoundSubmix, SubmixPtr);
 		}
+
+		// Now that the submix is loaded, we can register the submix buffer listeners that were waiting for it
+		if (PendingSubmixBufferListeners.Contains(InSoundSubmix->GetPathName()))
+		{
+			UE_LOG(LogAudioMixer, Display, TEXT("Submix %s had pending Submix Buffer Listeners, registering them now..."), *InSoundSubmix->GetName());
+			
+			TMap<FString, TSharedRef<ISubmixBufferListener, ESPMode::ThreadSafe>>& PendingListeners = PendingSubmixBufferListeners[InSoundSubmix->GetPathName()];
+
+			if (USoundSubmix* Submix = Cast<USoundSubmix>(InSoundSubmix))
+			{
+				for (TPair<FString, TSharedRef<ISubmixBufferListener, ESPMode::ThreadSafe>>& Listener : PendingListeners)
+				{
+					RegisterSubmixBufferListener(Listener.Value, *Submix);
+				}
+				PendingSubmixBufferListeners.Remove(InSoundSubmix->GetPathName());
+			}
+		}
 	}
 
 	void FMixerDevice::LoadSoundSubmix(USoundSubmixBase& InSoundSubmix)
@@ -2985,11 +3009,30 @@ namespace Audio
 				? GetMasterSubmix().Pin()
 				: GetSubmixInstance(SubmixPtr.Get()).Pin();
 
-			// Attempt to register submix if instance not found and is not master (i.e. default) submix
+			// Submix is not currently registered, so add this listener to a pending list for this submix
 			if (!bUseMaster && !FoundSubmix.IsValid() && SubmixPtr.IsValid())
 			{
-				RegisterSoundSubmix(SubmixPtr.Get(), true /* bInit */);
-				FoundSubmix = GetSubmixInstance(SubmixPtr.Get()).Pin();
+
+				FString SubmixName = SubmixPtr->GetPathName();
+				if (SubmixBufferListenerForceRegistrationCVar)
+				{
+					RegisterSoundSubmix(SubmixPtr.Get(), true /* bInit */);
+					FoundSubmix = GetSubmixInstance(SubmixPtr.Get()).Pin();
+					
+					UE_LOG(LogAudioMixer, Warning, TEXT("Submix '%s' force-enabled by Submix Buffer Listener '%s' due to CVar 'au.SubmixBufferListenerForceRegistration'."
+						"Please disable the CVar and manually ensure the submix is registered before registering the submix buffer listener."), *SubmixPtr->GetName(), *InSubmixBufferListener->GetListenerName());
+				}
+				else
+				{
+					UE_LOG(LogAudioMixer, Display, TEXT("Submix '%s' not currently loaded, Submix buffer listener '%s' added to pending list."), *SubmixName, *InSubmixBufferListener->GetListenerName());
+					if (!PendingSubmixBufferListeners.Contains(SubmixName))
+					{
+						PendingSubmixBufferListeners.Add(SubmixName, TMap<FString, TSharedRef<ISubmixBufferListener, ESPMode::ThreadSafe>>());
+					}
+					TMap<FString, TSharedRef<ISubmixBufferListener, ESPMode::ThreadSafe>>& PendingMap = PendingSubmixBufferListeners[SubmixName];
+					PendingMap.Add(InSubmixBufferListener->GetListenerName(), InSubmixBufferListener->AsShared());
+				}
+
 			}
 
 			if (FoundSubmix.IsValid())
@@ -2998,10 +3041,6 @@ namespace Audio
 				FoundSubmix->RegisterBufferListener(InSubmixBufferListener);
 				PRAGMA_ENABLE_DEPRECATION_WARNINGS
 			}
-			else
-			{
-				UE_LOG(LogAudioMixer, Warning, TEXT("Submix buffer listener not registered. Submix not loaded."));
-			}
 		};
 
 		FAudioThread::RunCommandOnAudioThread(MoveTemp(RegisterLambda));
@@ -3026,11 +3065,29 @@ namespace Audio
 
 			FMixerSubmixPtr FoundSubmix = GetSubmixInstance(SubmixPtr.Get()).Pin();
 
-			// Attempt to register submix if instance not found and is not master (i.e. default) submix
+			// Submix is not currently registered
 			if (!FoundSubmix.IsValid() && SubmixPtr.IsValid())
 			{
-				RegisterSoundSubmix(SubmixPtr.Get(), true /* bInit */);
-				FoundSubmix = GetSubmixInstance(SubmixPtr.Get()).Pin();
+				FString SubmixName = SubmixPtr->GetPathName();
+				if (SubmixBufferListenerForceRegistrationCVar)
+				{
+					RegisterSoundSubmix(SubmixPtr.Get(), true /* bInit */);
+					FoundSubmix = GetSubmixInstance(SubmixPtr.Get()).Pin();
+					
+					UE_LOG(LogAudioMixer, Warning, TEXT("Submix '%s' force-enabled by Submix Buffer Listener '%s' due to CVar 'au.SubmixBufferListenerForceRegistration'."
+						"Please disable the CVar and manually ensure the submix is registered before registering the submix buffer listener."), *SubmixPtr->GetName(), *ListenerName);
+				}
+				else
+				{
+					UE_LOG(LogAudioMixer, Display, TEXT("Submix '%s' not currently loaded, Submix buffer listener '%s' added to pending list."), *SubmixName, *ListenerName);
+					if (!PendingSubmixBufferListeners.Contains(SubmixName))
+					{
+						PendingSubmixBufferListeners.Add(SubmixName, TMap<FString, TSharedRef<ISubmixBufferListener, ESPMode::ThreadSafe>>());
+					}
+
+					TMap<FString, TSharedRef<ISubmixBufferListener, ESPMode::ThreadSafe>>& PendingMap = PendingSubmixBufferListeners[SubmixName];
+					PendingMap.Add(ListenerName, InSubmixBufferListener);
+				}
 			}
 
 			if (FoundSubmix.IsValid())
@@ -3038,10 +3095,6 @@ namespace Audio
 				FoundSubmix->RegisterBufferListener(InSubmixBufferListener);
 				UE_LOG(LogAudioMixer, Display, TEXT("Submix buffer listener '%s' registered with submix '%s'"), *ListenerName, *FoundSubmix->SubmixName);
 			}
-			else
-			{
-				UE_LOG(LogAudioMixer, Warning, TEXT("Submix buffer listener '%s' not registered. Submix not loaded."), *ListenerName);
-			}
 		};
 
 		UE_LOG(LogAudioMixer, Display, TEXT("Sending SubmixBufferListener '%s' register command..."), *ListenerName);
@@ -3063,6 +3116,14 @@ namespace Audio
 				? GetMasterSubmix().Pin()
 				: GetSubmixInstance(SubmixPtr.Get()).Pin();
 
+			if (SubmixPtr.IsValid())
+			{
+				if (PendingSubmixBufferListeners.Contains(SubmixPtr->GetPathName()))
+				{
+					PendingSubmixBufferListeners[SubmixPtr->GetPathName()].Remove(InSubmixBufferListener->GetListenerName());
+				}
+			}
+				
 			if (FoundSubmix.IsValid())
 			{
 				PRAGMA_DISABLE_DEPRECATION_WARNINGS
@@ -3083,17 +3144,26 @@ namespace Audio
 		DECLARE_CYCLE_STAT(TEXT("FAudioThreadTask.UnregisterSubmixBufferListener"), STAT_UnregisterSubmixBufferListener, STATGROUP_AudioThreadCommands);
 
 		const TWeakObjectPtr<USoundSubmix> SubmixPtr(&InSubmix);
-		UPTRINT ListenerPtr = reinterpret_cast<UPTRINT>(&InSubmixBufferListener.Get());
+		
 		FString ListenerName = InSubmixBufferListener->GetListenerName();
 		
-		auto UnregisterLambda = [this, SubmixPtr, ListenerPtr, ListenerName]()
+		auto UnregisterLambda = [this, SubmixPtr, InSubmixBufferListener, ListenerName]()
 		{
 			CSV_SCOPED_TIMING_STAT(Audio, UnregisterSubmixBufferListener);
 
+			if (SubmixPtr.IsValid())
+			{
+				if (PendingSubmixBufferListeners.Contains(SubmixPtr->GetPathName()))
+				{
+					PendingSubmixBufferListeners[SubmixPtr->GetPathName()].Remove(InSubmixBufferListener->GetListenerName());
+				}
+			}
+				
 			FMixerSubmixPtr FoundSubmix = GetSubmixInstance(SubmixPtr.Get()).Pin();
 			if (FoundSubmix.IsValid())
 			{
 				UE_LOG(LogAudioMixer, Display, TEXT("Unregistering submix buffer listener '%s' from submix '%s'"), *ListenerName, *FoundSubmix->SubmixName);
+				UPTRINT ListenerPtr = reinterpret_cast<UPTRINT>(&InSubmixBufferListener.Get());
 				FoundSubmix->UnregisterBufferListenerInternal(ListenerPtr);
 			}
 			else
diff --git a/Engine/Source/Runtime/AudioMixer/Public/AudioMixerDevice.h b/Engine/Source/Runtime/AudioMixer/Public/AudioMixerDevice.h
--- a/Engine/Source/Runtime/AudioMixer/Public/AudioMixerDevice.h
+++ b/Engine/Source/Runtime/AudioMixer/Public/AudioMixerDevice.h
@@ -473,6 +473,9 @@ namespace Audio
 		// Contended between RegisterSoundSubmix/UnregisterSoundSubmix on the audio thread and OnProcessAudioStream on the audio mixer thread.
 		FCriticalSection EndpointSubmixesMutationLock;
 
+		/** Submix Buffer Listeners which are trying to listen to submixes which have not been registered yet */
+		TMap<FString, TMap<FString, TSharedRef<ISubmixBufferListener, ESPMode::ThreadSafe>>> PendingSubmixBufferListeners;
+		
 		/** Which submixes have been told to envelope follow with this audio device. */
 		TArray<USoundSubmix*> DelegateBoundSubmixes;
 
