[AudioInsights]

* Auto-enable output metering
* Fixes: https://jira.it.epicgames.com/browse/UE-349717

#Synced-CL 48281426

diff --git a/Engine/Source/Runtime/AudioMixer/Private/AudioMixerDevice.cpp b/Engine/Source/Runtime/AudioMixer/Private/AudioMixerDevice.cpp
--- a/Engine/Source/Runtime/AudioMixer/Private/AudioMixerDevice.cpp
+++ b/Engine/Source/Runtime/AudioMixer/Private/AudioMixerDevice.cpp
@@ -1418,6 +1418,21 @@ namespace Audio
 		}
 	}
 
+	static inline FString GetRequiredSubmixName(const ERequiredSubmixes InRequiredSubmix)
+	{
+		switch (InRequiredSubmix) 
+		{
+			case ERequiredSubmixes::Main: { return TEXT("MasterSubmixDefault"); }
+			case ERequiredSubmixes::BaseDefault: { return TEXT("BaseDefault"); }
+			case ERequiredSubmixes::Reverb: { return TEXT("MasterReverbSubmixDefault"); }
+			case ERequiredSubmixes::EQ: { return TEXT("MasterEQSubmixDefault"); }
+			case ERequiredSubmixes::Count: break;
+		}
+
+		UE_LOG(LogAudioMixer, Warning, TEXT("Unsupported submix name for provided ERequiredSubmix"));
+		return "";
+	}
+
 	void FMixerDevice::InitSoundSubmixes()
 	{
 		if (IsInGameThread())
@@ -1437,19 +1452,19 @@ namespace Audio
 			}
 
 			// 1. Load or reload all sound submixes/instances
-			LoadRequiredSubmix(ERequiredSubmixes::Main, TEXT("MasterSubmixDefault"), false /* DefaultMuteWhenBackgrounded */, AudioSettings->MasterSubmix);
+			LoadRequiredSubmix(ERequiredSubmixes::Main, GetRequiredSubmixName(ERequiredSubmixes::Main), false /* DefaultMuteWhenBackgrounded */, AudioSettings->MasterSubmix);
 
 			// BaseDefaultSubmix is an optional master submix type set by project settings
 			if (AudioSettings->BaseDefaultSubmix.IsValid())
 			{
-				LoadRequiredSubmix(ERequiredSubmixes::BaseDefault, TEXT("BaseDefault"), false /* DefaultMuteWhenBackgrounded */, AudioSettings->BaseDefaultSubmix);
+				LoadRequiredSubmix(ERequiredSubmixes::BaseDefault, GetRequiredSubmixName(ERequiredSubmixes::BaseDefault), false /* DefaultMuteWhenBackgrounded */, AudioSettings->BaseDefaultSubmix);
 			}
 
-			LoadRequiredSubmix(ERequiredSubmixes::Reverb, TEXT("MasterReverbSubmixDefault"), true /* DefaultMuteWhenBackgrounded */, AudioSettings->ReverbSubmix);
+			LoadRequiredSubmix(ERequiredSubmixes::Reverb, GetRequiredSubmixName(ERequiredSubmixes::Reverb), true /* DefaultMuteWhenBackgrounded */, AudioSettings->ReverbSubmix);
 
 			if (!DisableSubmixEffectEQCvar)
 			{
-				LoadRequiredSubmix(ERequiredSubmixes::EQ, TEXT("MasterEQSubmixDefault"), false /* DefaultMuteWhenBackgrounded */, AudioSettings->EQSubmix);
+				LoadRequiredSubmix(ERequiredSubmixes::EQ, GetRequiredSubmixName(ERequiredSubmixes::EQ), false /* DefaultMuteWhenBackgrounded */, AudioSettings->EQSubmix);
 			}
 
 			LoadPluginSoundSubmixes();
@@ -1512,6 +1527,15 @@ namespace Audio
 				}
 			}
 		}
+
+#if UE_AUDIO_PROFILERTRACE_ENABLED
+		// Auto-Enable MainSubmix output metering
+		FAudioThread::RunCommandOnAudioThread([this]()
+		{
+			StartSubmixEnvelopeFollowerFromCVar(GetRequiredSubmixName(ERequiredSubmixes::Main));
+		});
+#endif // UE_AUDIO_PROFILERTRACE_ENABLED
+
 	}
 
 	void FMixerDevice::RebuildSubmixLinks(const USoundSubmixBase& SoundSubmix, FMixerSubmixPtr& SubmixInstance)
