[Audio Insights] Signal Flow part 01: Move Wave Instance lifetime trace messages to AudioMixerSource. Add new SoundWaveStop message.

#Synced-CL 48352470

diff --git a/Engine/Plugins/AudioInsights/Source/AudioInsights/Private/Messages/SoundTraceMessages.cpp b/Engine/Plugins/AudioInsights/Source/AudioInsights/Private/Messages/SoundTraceMessages.cpp
--- a/Engine/Plugins/AudioInsights/Source/AudioInsights/Private/Messages/SoundTraceMessages.cpp
+++ b/Engine/Plugins/AudioInsights/Source/AudioInsights/Private/Messages/SoundTraceMessages.cpp
@@ -10,6 +10,7 @@ namespace UE::Audio::Insights
 		const FName SoundWaveStart = "SoundWaveStart";
 		const FName SoundWaveIsAlivePing = "SoundWaveIsAlivePing";
 		const FName SoundStop = "SoundStop";
+		const FName SoundWaveStop = "SoundWaveStop";
 
 		const FName PriorityParam = "Priority";
 		const FName DistanceParam = "Distance";
@@ -91,7 +92,10 @@ namespace UE::Audio::Insights
 
 	uint32 FSoundWaveStartMessage::GetSizeOf() const
 	{
-		return FSoundStartMessage::GetSizeOf() + sizeof(ActiveSoundPlayOrder);
+		uint32 MemorySize = sizeof(FSoundWaveStartMessage);
+
+		MemorySize += ActorIconName.GetAllocatedSize();
+		return MemorySize;
 	}
 
 	uint32 FSoundStopMessage::GetSizeOf() const
@@ -99,6 +103,11 @@ namespace UE::Audio::Insights
 		return sizeof(FSoundStopMessage);
 	}
 
+	uint32 FSoundWaveStopMessage::GetSizeOf() const
+	{
+		return sizeof(FSoundWaveStopMessage);
+	}
+
 	FSoundDashboardEntry::FSoundDashboardEntry()
 	{
 		constexpr uint32 DataPointsCapacity = 256u;
diff --git a/Engine/Plugins/AudioInsights/Source/AudioInsights/Public/Messages/SoundTraceMessages.h b/Engine/Plugins/AudioInsights/Source/AudioInsights/Public/Messages/SoundTraceMessages.h
--- a/Engine/Plugins/AudioInsights/Source/AudioInsights/Public/Messages/SoundTraceMessages.h
+++ b/Engine/Plugins/AudioInsights/Source/AudioInsights/Public/Messages/SoundTraceMessages.h
@@ -45,6 +45,7 @@ namespace UE::Audio::Insights
 		extern const FName SoundWaveStart;
 		extern const FName SoundWaveIsAlivePing;
 		extern const FName SoundStop;
+		extern const FName SoundWaveStop;
 		
 		extern const FName PriorityParam;
 		extern const FName DistanceParam;
@@ -123,6 +124,8 @@ namespace UE::Audio::Insights
 		{
 			WaveInstancePlayOrder = InContext.EventData.GetValue<uint32>("PlayOrder");
 			ActiveSoundPlayOrder = InContext.EventData.GetValue<uint32>("ActiveSoundPlayOrder");
+
+			InContext.EventData.GetString("ActorIconName", ActorIconName);
 		}
 
 		virtual uint32 GetSizeOf() const override;
@@ -130,6 +133,7 @@ namespace UE::Audio::Insights
 		virtual uint64 GetID() const override { return SoundMessageUtils::GeneratePlayOrderUniqueID(ActiveSoundPlayOrder, WaveInstancePlayOrder); }
 
 		uint32 WaveInstancePlayOrder = INDEX_NONE;
+		FString ActorIconName;
 	};
 
 	struct FSoundWaveIsAlivePingMessage : public FSoundWaveStartMessage
@@ -159,6 +163,25 @@ namespace UE::Audio::Insights
 		uint32 ActiveSoundPlayOrder = INDEX_NONE;
 	};
 
+	struct FSoundWaveStopMessage : public FSoundMessageBase
+	{
+		FSoundWaveStopMessage() = default;
+
+		FSoundWaveStopMessage(const Trace::IAnalyzer::FOnEventContext& InContext)
+			: FSoundMessageBase(InContext)
+		{
+			WaveInstancePlayOrder = InContext.EventData.GetValue<uint32>("PlayOrder");
+			ActiveSoundPlayOrder = InContext.EventData.GetValue<uint32>("ActiveSoundPlayOrder");
+		}
+
+		virtual uint32 GetSizeOf() const override;
+		virtual const FName GetMessageName() const override { return SoundMessageNames::SoundWaveStop; }
+		virtual uint64 GetID() const override { return SoundMessageUtils::GeneratePlayOrderUniqueID(ActiveSoundPlayOrder, WaveInstancePlayOrder); }
+
+		uint32 WaveInstancePlayOrder = INDEX_NONE;
+		uint32 ActiveSoundPlayOrder = INDEX_NONE;
+	};
+
 #define DEFINE_SOUND_PARAM_MESSAGE(ClassName, MessageName, ParamName, Type, Default)			\
 	struct ClassName : public FSoundParameterMessage											\
 	{																							\
diff --git a/Engine/Source/Runtime/AudioMixer/Private/AudioMixerSource.cpp b/Engine/Source/Runtime/AudioMixer/Private/AudioMixerSource.cpp
--- a/Engine/Source/Runtime/AudioMixer/Private/AudioMixerSource.cpp
+++ b/Engine/Source/Runtime/AudioMixer/Private/AudioMixerSource.cpp
@@ -2,6 +2,7 @@
 #include "AudioMixerSource.h"
 
 #include "Audio/AudioTimingLog.h"
+#include "Audio/AudioTraceUtil.h"
 #include "AudioDefines.h"
 #include "AudioMixerSourceBuffer.h"
 #include "ActiveSound.h"
@@ -24,6 +25,35 @@
 CSV_DECLARE_CATEGORY_MODULE_EXTERN(AUDIOMIXERCORE_API, Audio);
 
 #if UE_AUDIO_PROFILERTRACE_ENABLED
+UE_TRACE_EVENT_BEGIN(Audio, SoundWaveStart)
+	UE_TRACE_EVENT_FIELD(uint32, DeviceId)
+	UE_TRACE_EVENT_FIELD(uint64, Timestamp)
+	UE_TRACE_EVENT_FIELD(uint32, PlayOrder)
+	UE_TRACE_EVENT_FIELD(uint32, ActiveSoundPlayOrder)
+	UE_TRACE_EVENT_FIELD(UE::Trace::WideString, AssetPath)
+	UE_TRACE_EVENT_FIELD(UE::Trace::WideString, SoundClassName)
+	UE_TRACE_EVENT_FIELD(UE::Trace::WideString, ActorLabel)
+	UE_TRACE_EVENT_FIELD(UE::Trace::WideString, ActorIconName)
+UE_TRACE_EVENT_END()
+
+UE_TRACE_EVENT_BEGIN(Audio, SoundWaveIsAlivePing)
+	UE_TRACE_EVENT_FIELD(uint32, DeviceId)
+	UE_TRACE_EVENT_FIELD(uint64, Timestamp)
+	UE_TRACE_EVENT_FIELD(uint32, PlayOrder)
+	UE_TRACE_EVENT_FIELD(uint32, ActiveSoundPlayOrder)
+	UE_TRACE_EVENT_FIELD(UE::Trace::WideString, AssetPath)
+	UE_TRACE_EVENT_FIELD(UE::Trace::WideString, SoundClassName)
+	UE_TRACE_EVENT_FIELD(UE::Trace::WideString, ActorLabel)
+	UE_TRACE_EVENT_FIELD(UE::Trace::WideString, ActorIconName)
+UE_TRACE_EVENT_END()
+
+UE_TRACE_EVENT_BEGIN(Audio, SoundWaveStop)
+	UE_TRACE_EVENT_FIELD(uint32, DeviceId)
+	UE_TRACE_EVENT_FIELD(uint64, Timestamp)
+	UE_TRACE_EVENT_FIELD(uint32, PlayOrder)
+	UE_TRACE_EVENT_FIELD(uint32, ActiveSoundPlayOrder)
+UE_TRACE_EVENT_END()
+
 UE_TRACE_EVENT_BEGIN(Audio, SoundRelativeRenderCost)
 	UE_TRACE_EVENT_FIELD(uint32, DeviceId)
 	UE_TRACE_EVENT_FIELD(uint64, Timestamp)
@@ -811,6 +841,10 @@ namespace Audio
 
 		Audio::FAudioDebugger::DrawDebugInfo(*this);
 #endif // ENABLE_AUDIO_DEBUG
+
+#if UE_AUDIO_PROFILERTRACE_ENABLED
+		UpdateIsAlivePingTraceMessage();
+#endif // UE_AUDIO_PROFILERTRACE_ENABLED
 	}
 
 	bool FMixerSource::PrepareForInitialization(FWaveInstance* InWaveInstance)
@@ -1050,6 +1084,10 @@ namespace Audio
 			MixerSourceVoice->Play();
 		}
 
+#if UE_AUDIO_PROFILERTRACE_ENABLED
+		SendSoundWaveStartTraceMessage();
+#endif // UE_AUDIO_PROFILERTRACE_ENABLED
+
 		bIsStopping = false;
 		Paused = false;
 		Playing = true;
@@ -1130,6 +1168,10 @@ namespace Audio
 		{
 			if (MixerSourceVoice && Playing)
 			{
+#if UE_AUDIO_PROFILERTRACE_ENABLED
+				SendSoundWaveStopTraceMessage();
+#endif // UE_AUDIO_PROFILERTRACE_ENABLED
+
 				MixerSourceVoice->Stop();
 			}
 
@@ -1837,6 +1879,75 @@ namespace Audio
 	}
 #endif // if ENABLE_AUDIO_DEBUG
 
+#if UE_AUDIO_PROFILERTRACE_ENABLED
+	void FMixerSource::SendSoundWaveStartTraceMessage() const
+	{
+		const bool bChannelEnabled = UE_TRACE_CHANNELEXPR_IS_ENABLED(AudioMixerChannel);
+		if (bChannelEnabled && WaveInstance && WaveInstance->ActiveSound && WaveInstance->SoundClass)
+		{
+			const TObjectPtr<UObject> SoundWaveObjectPtr = FSoftObjectPath(WaveInstance->WaveData->GetPathName()).ResolveObject();
+			if (SoundWaveObjectPtr)
+			{
+				const FString SoundWaveBaseAssetName = Audio::Trace::Util::GetSoundBaseAssetName(SoundWaveObjectPtr->GetClass());
+				const FString OwnerActorLabel = Audio::Trace::Util::GetOwnerActorLabel(*WaveInstance->ActiveSound);
+
+				UE_TRACE_LOG(Audio, SoundWaveStart, AudioMixerChannel)
+					<< SoundWaveStart.DeviceId(AudioDevice->DeviceID)
+					<< SoundWaveStart.Timestamp(FPlatformTime::Cycles64())
+					<< SoundWaveStart.PlayOrder(WaveInstance->GetPlayOrder())
+					<< SoundWaveStart.ActiveSoundPlayOrder(WaveInstance->ActiveSound->GetPlayOrder())
+					<< SoundWaveStart.AssetPath(*SoundWaveObjectPtr->GetPathName())
+					<< SoundWaveStart.SoundClassName(*SoundWaveBaseAssetName)
+					<< SoundWaveStart.ActorLabel(*OwnerActorLabel)
+					<< SoundWaveStart.ActorIconName(*Audio::Trace::Util::GetOwnerActorIconName(*WaveInstance->ActiveSound).ToString());
+			}
+		}
+	}
+
+	void FMixerSource::SendSoundWaveStopTraceMessage() const
+	{
+		const bool bChannelEnabled = UE_TRACE_CHANNELEXPR_IS_ENABLED(AudioMixerChannel);
+		if (bChannelEnabled && MixerDevice && WaveInstance && WaveInstance->ActiveSound)
+		{
+			UE_TRACE_LOG(Audio, SoundWaveStop, AudioMixerChannel)
+				<< SoundWaveStop.DeviceId(MixerDevice->DeviceID)
+				<< SoundWaveStop.Timestamp(FPlatformTime::Cycles64())
+				<< SoundWaveStop.PlayOrder(WaveInstance->GetPlayOrder())
+				<< SoundWaveStop.ActiveSoundPlayOrder(WaveInstance->ActiveSound->GetPlayOrder());
+		}
+	}
+
+	void FMixerSource::UpdateIsAlivePingTraceMessage() const
+	{
+		const bool bChannelEnabled = UE_TRACE_CHANNELEXPR_IS_ENABLED(AudioMixerChannel);
+
+		if (bChannelEnabled && WaveInstance && WaveInstance->IsPlaying() && WaveInstance->ActiveSound && WaveInstance->SoundClass)
+		{
+			// Occasionally send a ping to Audio Insights to confirm the sound is still playing.
+			// This is to ensure long duration sounds (such as ambiences) will still appear in snapshot files if their start action was before the snapshot start
+			if (FMath::Floor(WaveInstance->ActiveSound->PlaybackTimeUnscaled + AudioDevice->GetDeviceDeltaTime()) - FMath::Floor(WaveInstance->ActiveSound->PlaybackTimeUnscaled) > 0.0f)
+			{
+				const TObjectPtr<UObject> SoundWaveObjectPtr = FSoftObjectPath(WaveInstance->WaveData->GetPathName()).ResolveObject();
+				if (SoundWaveObjectPtr)
+				{
+					const FString SoundWaveBaseAssetName = Audio::Trace::Util::GetSoundBaseAssetName(SoundWaveObjectPtr->GetClass());
+					const FString OwnerActorLabel = Audio::Trace::Util::GetOwnerActorLabel(*WaveInstance->ActiveSound);
+
+					UE_TRACE_LOG(Audio, SoundWaveIsAlivePing, AudioMixerChannel)
+						<< SoundWaveIsAlivePing.DeviceId(AudioDevice->DeviceID)
+						<< SoundWaveIsAlivePing.Timestamp(FPlatformTime::Cycles64())
+						<< SoundWaveIsAlivePing.PlayOrder(WaveInstance->GetPlayOrder())
+						<< SoundWaveIsAlivePing.ActiveSoundPlayOrder(WaveInstance->ActiveSound->GetPlayOrder())
+						<< SoundWaveIsAlivePing.AssetPath(*SoundWaveObjectPtr->GetPathName())
+						<< SoundWaveIsAlivePing.SoundClassName(*SoundWaveBaseAssetName)
+						<< SoundWaveIsAlivePing.ActorLabel(*OwnerActorLabel)
+						<< SoundWaveIsAlivePing.ActorIconName(*Audio::Trace::Util::GetOwnerActorIconName(*WaveInstance->ActiveSound).ToString());
+				}
+			}
+		}
+	}
+#endif // UE_AUDIO_PROFILERTRACE_ENABLED
+
 	bool FMixerSource::ComputeMonoChannelMap(Audio::FAlignedFloatBuffer& OutChannelMap)
 	{
 		if (IsUsingObjectBasedSpatialization())
diff --git a/Engine/Source/Runtime/AudioMixer/Private/AudioMixerSource.h b/Engine/Source/Runtime/AudioMixer/Private/AudioMixerSource.h
--- a/Engine/Source/Runtime/AudioMixer/Private/AudioMixerSource.h
+++ b/Engine/Source/Runtime/AudioMixer/Private/AudioMixerSource.h
@@ -98,6 +98,12 @@ namespace Audio
 		void UpdateCPUCoreUtilization();
 #endif // ENABLE_AUDIO_DEBUG
 
+#if UE_AUDIO_PROFILERTRACE_ENABLED
+		void SendSoundWaveStartTraceMessage() const;
+		void SendSoundWaveStopTraceMessage() const;
+		void UpdateIsAlivePingTraceMessage() const;
+#endif // UE_AUDIO_PROFILERTRACE_ENABLED
+
 		/** Computes the mono-channel map. */
 		bool ComputeMonoChannelMap(Audio::FAlignedFloatBuffer& OutChannelMap);
 
diff --git a/Engine/Source/Runtime/Engine/Private/SoundWave.cpp b/Engine/Source/Runtime/Engine/Private/SoundWave.cpp
--- a/Engine/Source/Runtime/Engine/Private/SoundWave.cpp
+++ b/Engine/Source/Runtime/Engine/Private/SoundWave.cpp
@@ -112,26 +112,6 @@ FAutoConsoleVariableRef CvarAllowReverbForMultichannelSources(
 	ECVF_Default);
 
 #if UE_AUDIO_PROFILERTRACE_ENABLED
-UE_TRACE_EVENT_BEGIN(Audio, SoundWaveStart)
-	UE_TRACE_EVENT_FIELD(uint32, DeviceId)
-	UE_TRACE_EVENT_FIELD(uint64, Timestamp)
-	UE_TRACE_EVENT_FIELD(uint32, PlayOrder)
-	UE_TRACE_EVENT_FIELD(uint32, ActiveSoundPlayOrder)
-	UE_TRACE_EVENT_FIELD(UE::Trace::WideString, AssetPath)
-	UE_TRACE_EVENT_FIELD(UE::Trace::WideString, SoundClassName)
-	UE_TRACE_EVENT_FIELD(UE::Trace::WideString, ActorLabel)
-UE_TRACE_EVENT_END()
-
-UE_TRACE_EVENT_BEGIN(Audio, SoundWaveIsAlivePing)
-	UE_TRACE_EVENT_FIELD(uint32, DeviceId)
-	UE_TRACE_EVENT_FIELD(uint64, Timestamp)
-	UE_TRACE_EVENT_FIELD(uint32, PlayOrder)
-	UE_TRACE_EVENT_FIELD(uint32, ActiveSoundPlayOrder)
-	UE_TRACE_EVENT_FIELD(UE::Trace::WideString, AssetPath)
-	UE_TRACE_EVENT_FIELD(UE::Trace::WideString, SoundClassName)
-	UE_TRACE_EVENT_FIELD(UE::Trace::WideString, ActorLabel)
-UE_TRACE_EVENT_END()
-
 UE_TRACE_EVENT_BEGIN(Audio, SoundPriority)
 	UE_TRACE_EVENT_FIELD(uint32, DeviceId)
 	UE_TRACE_EVENT_FIELD(uint64, Timestamp)
@@ -4001,45 +3981,6 @@ void USoundWave::Parse(FAudioDevice* AudioDevice, const UPTRINT NodeWaveInstance
 
 	if (bChannelEnabled && bWaveIsActive && WaveInstance->SoundClass != nullptr)
 	{
-		if (bIsNewWave)
-		{
-			const TObjectPtr<UObject> SoundWaveObjectPtr = FSoftObjectPath(WaveInstance->WaveData->GetPathName()).ResolveObject();
-			if (SoundWaveObjectPtr)
-			{
-				const FString SoundWaveBaseAssetName = Audio::Trace::Util::GetSoundBaseAssetName(SoundWaveObjectPtr->GetClass());
-				const FString OwnerActorLabel = Audio::Trace::Util::GetOwnerActorLabel(ActiveSound);
-
-				UE_TRACE_LOG(Audio, SoundWaveStart, AudioMixerChannel)
-					<< SoundWaveStart.DeviceId(AudioDevice->DeviceID)
-					<< SoundWaveStart.Timestamp(FPlatformTime::Cycles64())
-					<< SoundWaveStart.PlayOrder(WaveInstance->GetPlayOrder())
-					<< SoundWaveStart.ActiveSoundPlayOrder(ActiveSound.GetPlayOrder())
-					<< SoundWaveStart.AssetPath(*SoundWaveObjectPtr->GetPathName())
-					<< SoundWaveStart.SoundClassName(*SoundWaveBaseAssetName)
-					<< SoundWaveStart.ActorLabel(*OwnerActorLabel);
-			}
-		}
-		// Occasionally send a ping to Audio Insights to confirm the sound is still playing.
-		// This is to ensure long duration sounds (such as ambiences) will still appear in snapshot files if their start action was before the snapshot start
-		else if (FMath::Floor(ActiveSound.PlaybackTimeUnscaled + AudioDevice->GetDeviceDeltaTime()) - FMath::Floor(ActiveSound.PlaybackTimeUnscaled) > 0.0f)
-		{
-			const TObjectPtr<UObject> SoundWaveObjectPtr = FSoftObjectPath(WaveInstance->WaveData->GetPathName()).ResolveObject();
-			if (SoundWaveObjectPtr)
-			{
-				const FString SoundWaveBaseAssetName = Audio::Trace::Util::GetSoundBaseAssetName(SoundWaveObjectPtr->GetClass());
-				const FString OwnerActorLabel = Audio::Trace::Util::GetOwnerActorLabel(ActiveSound);
-
-				UE_TRACE_LOG(Audio, SoundWaveIsAlivePing, AudioMixerChannel)
-					<< SoundWaveIsAlivePing.DeviceId(AudioDevice->DeviceID)
-					<< SoundWaveIsAlivePing.Timestamp(FPlatformTime::Cycles64())
-					<< SoundWaveIsAlivePing.PlayOrder(WaveInstance->GetPlayOrder())
-					<< SoundWaveIsAlivePing.ActiveSoundPlayOrder(ActiveSound.GetPlayOrder())
-					<< SoundWaveIsAlivePing.AssetPath(*SoundWaveObjectPtr->GetPathName())
-					<< SoundWaveIsAlivePing.SoundClassName(*SoundWaveBaseAssetName)
-					<< SoundWaveIsAlivePing.ActorLabel(*OwnerActorLabel);
-			}
-		}
-
 		UE_TRACE_LOG(Audio, SoundPriority, AudioMixerChannel)
 			<< SoundPriority.DeviceId(AudioDevice->DeviceID)
 			<< SoundPriority.Timestamp(FPlatformTime::Cycles64())
