[Audio Insights] Signal Flow part 02: Send new trace messages for asset send values, as well as BusIsAlive pings

#Synced-CL 48352470

diff --git a/Engine/Source/Runtime/AudioMixer/Private/AudioMixerDevice.cpp b/Engine/Source/Runtime/AudioMixer/Private/AudioMixerDevice.cpp
--- a/Engine/Source/Runtime/AudioMixer/Private/AudioMixerDevice.cpp
+++ b/Engine/Source/Runtime/AudioMixer/Private/AudioMixerDevice.cpp
@@ -144,6 +144,15 @@ UE_TRACE_EVENT_BEGIN(Audio, SubmixAlivePing)
 	UE_TRACE_EVENT_FIELD(bool, IsMainSubmix)
 UE_TRACE_EVENT_END()
 
+UE_TRACE_EVENT_BEGIN(Audio, SubmixToSubmixSend)
+	UE_TRACE_EVENT_FIELD(uint32, DeviceId)
+	UE_TRACE_EVENT_FIELD(uint64, Timestamp)
+	UE_TRACE_EVENT_FIELD(uint32, SenderID)
+	UE_TRACE_EVENT_FIELD(uint32, ReceiverID)
+	UE_TRACE_EVENT_FIELD(float, SendLevel)
+	UE_TRACE_EVENT_FIELD(bool, IsFinalStage)
+UE_TRACE_EVENT_END()
+
 #endif // UE_AUDIO_PROFILERTRACE_ENABLED
 
 // Link to "Audio" profiling category
@@ -443,6 +452,36 @@ namespace Audio
 		}
 	}
 
+	void FMixerDevice::SendSubmixesSendLevelTraceMessages()
+	{
+		const bool bChannelEnabled = UE_TRACE_CHANNELEXPR_IS_ENABLED(AudioChannel);
+		if (bChannelEnabled)
+		{
+			Submixes.Iterate([this](const FSubmixMap::FPair& SubmixPair)
+			{
+				const FMixerSubmixPtr& Submix = SubmixPair.Value;
+				if (!Submix.IsValid())
+				{
+					return;
+				}
+
+				const Audio::FMixerSubmixPtr OutputSubmix = Submix->GetParentSubmix().Pin();
+				
+				// If this submix has no parent, send the Audio Device as the receiver ID and mark as a final stage of the chain
+				const bool bIsFinalStage = !OutputSubmix.IsValid();
+				const uint32 ReceiverID = bIsFinalStage ? DeviceID : GetTypeHash(OutputSubmix->GetPathName());
+				
+				UE_TRACE_LOG(Audio, SubmixToSubmixSend, AudioMixerChannel)
+					<< SubmixToSubmixSend.DeviceId(DeviceID)
+					<< SubmixToSubmixSend.Timestamp(FPlatformTime::Cycles64())
+					<< SubmixToSubmixSend.SenderID(GetTypeHash(Submix->GetPathName()))
+					<< SubmixToSubmixSend.ReceiverID(ReceiverID)
+					<< SubmixToSubmixSend.SendLevel(Submix->CurrentOutputVolume)
+					<< SubmixToSubmixSend.IsFinalStage(bIsFinalStage);
+			});
+		}
+	}
+
 	void FMixerDevice::RegisterSubmixEnvelopeFollowerCVars()
 	{
 		// StartSubmixEnvelopeFollowerCVar
@@ -1246,6 +1285,7 @@ namespace Audio
 
 #if UE_AUDIO_PROFILERTRACE_ENABLED
 		UpdateSubmixesAlivePingTime();
+		SendSubmixesSendLevelTraceMessages();
 #endif //UE_AUDIO_PROFILERTRACE_ENABLED
 
 		// notify interested parties
diff --git a/Engine/Source/Runtime/AudioMixer/Private/AudioMixerSourceManager.cpp b/Engine/Source/Runtime/AudioMixer/Private/AudioMixerSourceManager.cpp
--- a/Engine/Source/Runtime/AudioMixer/Private/AudioMixerSourceManager.cpp
+++ b/Engine/Source/Runtime/AudioMixer/Private/AudioMixerSourceManager.cpp
@@ -277,6 +277,22 @@ UE_TRACE_EVENT_BEGIN(Audio, SoundSourceEnvelope)
 	UE_TRACE_EVENT_FIELD(float, Envelope)
 UE_TRACE_EVENT_END()
 
+UE_TRACE_EVENT_BEGIN(Audio, SoundToSubmixSend)
+	UE_TRACE_EVENT_FIELD(uint32, DeviceId)
+	UE_TRACE_EVENT_FIELD(uint64, Timestamp)
+	UE_TRACE_EVENT_FIELD(uint32, SenderID)
+	UE_TRACE_EVENT_FIELD(uint32, ReceiverID)
+	UE_TRACE_EVENT_FIELD(float, SendLevel)
+UE_TRACE_EVENT_END()
+
+UE_TRACE_EVENT_BEGIN(Audio, SoundToBusSend)
+	UE_TRACE_EVENT_FIELD(uint32, DeviceId)
+	UE_TRACE_EVENT_FIELD(uint64, Timestamp)
+	UE_TRACE_EVENT_FIELD(uint32, SenderID)
+	UE_TRACE_EVENT_FIELD(uint32, ReceiverID)
+	UE_TRACE_EVENT_FIELD(float, SendLevel)
+UE_TRACE_EVENT_END()
+
 // Audio Bus messages
 UE_TRACE_EVENT_BEGIN(Audio, AudioBusStart)
 	UE_TRACE_EVENT_FIELD(uint32, DeviceId)
@@ -286,6 +302,14 @@ UE_TRACE_EVENT_BEGIN(Audio, AudioBusStart)
 	UE_TRACE_EVENT_FIELD(int32, NumChannels)
 UE_TRACE_EVENT_END()
 
+UE_TRACE_EVENT_BEGIN(Audio, AudioBusIsAlivePing)
+	UE_TRACE_EVENT_FIELD(uint32, DeviceId)
+	UE_TRACE_EVENT_FIELD(uint32, AudioBusId)
+	UE_TRACE_EVENT_FIELD(uint64, Timestamp)
+	UE_TRACE_EVENT_FIELD(UE::Trace::WideString, Name)
+	UE_TRACE_EVENT_FIELD(int32, NumChannels)
+UE_TRACE_EVENT_END()
+
 UE_TRACE_EVENT_BEGIN(Audio, AudioBusStop)
 	UE_TRACE_EVENT_FIELD(uint32, DeviceId)
 	UE_TRACE_EVENT_FIELD(uint32, AudioBusId)
@@ -388,6 +412,20 @@ namespace Audio
 		}
 	}
 
+	void TraceAudioBusIsAlivePingMessage(const uint32 InDeviceId, const uint32 InObjectId, const FString& InAudioBusName, const int32 InNumChannels)
+	{
+		const bool bChannelEnabled = UE_TRACE_CHANNELEXPR_IS_ENABLED(AudioMixerChannel);
+		if (bChannelEnabled)
+		{
+			UE_TRACE_LOG(Audio, AudioBusIsAlivePing, AudioChannel)
+				<< AudioBusIsAlivePing.DeviceId(InDeviceId)
+				<< AudioBusIsAlivePing.AudioBusId(InObjectId)
+				<< AudioBusIsAlivePing.Timestamp(FPlatformTime::Cycles64())
+				<< AudioBusIsAlivePing.Name(InAudioBusName.IsEmpty() ? *FString::Printf(TEXT("%d"), InObjectId) : *InAudioBusName)
+				<< AudioBusIsAlivePing.NumChannels(InNumChannels);
+		}
+	}
+
 	void TraceAudioBusStopMessage(const uint32 InDeviceId, const uint32 InObjectId)
 	{
 		const bool bChannelEnabled = UE_TRACE_CHANNELEXPR_IS_ENABLED(AudioMixerChannel);
@@ -399,6 +437,20 @@ namespace Audio
 				<< AudioBusStop.Timestamp(FPlatformTime::Cycles64());
 		}
 	}
+
+	void TraceSourceToBusSendMessage(const uint32 InDeviceId, const uint32 InSourcePlayOrder, const uint32 InBusID, const float InSendLevel)
+	{
+		const bool bChannelEnabled = UE_TRACE_CHANNELEXPR_IS_ENABLED(AudioMixerChannel);
+		if (bChannelEnabled)
+		{
+			UE_TRACE_LOG(Audio, SoundToBusSend, AudioMixerChannel)
+				<< SoundToBusSend.DeviceId(InDeviceId)
+				<< SoundToBusSend.Timestamp(FPlatformTime::Cycles64())
+				<< SoundToBusSend.SenderID(InSourcePlayOrder)
+				<< SoundToBusSend.ReceiverID(InBusID)
+				<< SoundToBusSend.SendLevel(InSendLevel);
+		}
+	}
 #endif // UE_AUDIO_PROFILERTRACE_ENABLED
 
 	/*************************************************************************
@@ -1386,6 +1438,8 @@ namespace Audio
 						AudioBusPtr->BusName = AudioBusSend.AudioBusName;
 						TraceAudioBusStartMessage(MixerDevice->DeviceID, AudioBusSend.AudioBusId, AudioBusSend.AudioBusName, AudioBusSend.BusChannels);
 					}
+
+					TraceSourceToBusSendMessage(MixerDevice->DeviceID, SourceInfo.PlayOrder, SourceInfo.AudioBusId, AudioBusSend.SendLevel);
 #endif // UE_AUDIO_PROFILERTRACE_ENABLED
 					
 					// Store on this source, which buses its sending its audio to
@@ -2213,6 +2267,8 @@ namespace Audio
 				AudioBusPtr->BusName = InBusName;
 				TraceAudioBusStartMessage(MixerDevice->DeviceID, AudioBusId, InBusName, SourceInfo.NumInputChannels);
 			}
+
+			TraceSourceToBusSendMessage(MixerDevice->DeviceID, SourceInfo.PlayOrder, SourceInfo.AudioBusId, BusSendLevel);
 #endif // UE_AUDIO_PROFILERTRACE_ENABLED
 
 			// Check to see if we need to create new bus data. If we are not playing a bus with this id, then we
@@ -2838,9 +2894,28 @@ namespace Audio
 						<< AudioBusEnvelopeValues.Timestamp(FPlatformTime::Cycles64())
 						<< AudioBusEnvelopeValues.EnvelopeValues(AudioBus->EnvelopeValues, AudioBus->GetNumChannels());
 				}
+
+				if (BusIsAlivePingTimeSec >= BusIsAlivePingWaitTimeSec)
+				{
+					TraceAudioBusIsAlivePingMessage(MixerDevice->DeviceID, AudioBusKey.ObjectId, AudioBus->BusName, AudioBus->NumChannels);
+				}
 			}
 #endif // UE_AUDIO_PROFILERTRACE_ENABLED
 		}
+
+#if UE_AUDIO_PROFILERTRACE_ENABLED
+		if (bChannelEnabled)
+		{
+			if (BusIsAlivePingTimeSec >= BusIsAlivePingWaitTimeSec)
+			{
+				BusIsAlivePingTimeSec = 0.0;
+			}
+			else
+			{
+				BusIsAlivePingTimeSec += MixerDevice->GetAudioClockDelta();
+			}
+		}
+#endif // UE_AUDIO_PROFILERTRACE_ENABLED
 	}
 
 	void FMixerSourceManager::AddBusPatchDependency(const uint64 InstanceID, const FAudioRenderStepId FirstStep, const FAudioRenderStepId SecondStep)
@@ -3258,6 +3333,19 @@ namespace Audio
 					<< SoundDistanceAttenuation.PlayOrder(SourceInfo.PlayOrder)
 					<< SoundDistanceAttenuation.ActiveSoundPlayOrder(SourceInfo.ActiveSoundPlayOrder)
 					<< SoundDistanceAttenuation.DistanceAttenuation(SourceInfo.DistanceAttenuationSourceDestination);
+
+				for (const Audio::FMixerSourceSubmixSend& SubmixSend : SourceInfo.SubmixSends)
+				{
+					if (SubmixSend.Submix.Pin().IsValid())
+					{
+						UE_TRACE_LOG(Audio, SoundToSubmixSend, AudioMixerChannel)
+							<< SoundToSubmixSend.DeviceId(MixerDevice->DeviceID)
+							<< SoundToSubmixSend.Timestamp(FPlatformTime::Cycles64())
+							<< SoundToSubmixSend.SenderID(SourceInfo.PlayOrder)
+							<< SoundToSubmixSend.ReceiverID(GetTypeHash(SubmixSend.Submix.Pin()->GetPathName()))
+							<< SoundToSubmixSend.SendLevel(SubmixSend.SendLevel);
+					}
+				}
 			}
 #endif // UE_AUDIO_PROFILERTRACE_ENABLED
 		}
diff --git a/Engine/Source/Runtime/AudioMixer/Private/AudioMixerSourceManager.h b/Engine/Source/Runtime/AudioMixer/Private/AudioMixerSourceManager.h
--- a/Engine/Source/Runtime/AudioMixer/Private/AudioMixerSourceManager.h
+++ b/Engine/Source/Runtime/AudioMixer/Private/AudioMixerSourceManager.h
@@ -399,6 +399,9 @@ namespace Audio
 
 		IConsoleVariable* StartAudioBusEnvelopeFollowerCVar = nullptr;
 		IConsoleVariable* StopAudioBusEnvelopeFollowerCVar  = nullptr;
+
+		double BusIsAlivePingTimeSec = 0.0;
+		static constexpr double BusIsAlivePingWaitTimeSec = 1.0;
 #endif // UE_AUDIO_PROFILERTRACE_ENABLED
 
 		// Critical section to ensure mutating effect chains is thread-safe
diff --git a/Engine/Source/Runtime/AudioMixer/Public/AudioMixerDevice.h b/Engine/Source/Runtime/AudioMixer/Public/AudioMixerDevice.h
--- a/Engine/Source/Runtime/AudioMixer/Public/AudioMixerDevice.h
+++ b/Engine/Source/Runtime/AudioMixer/Public/AudioMixerDevice.h
@@ -409,6 +409,7 @@ namespace Audio
 		void OnTraceStarted(FTraceAuxiliary::EConnectionType TraceType, const FString& TraceDestination);
 		void SendSubmixesAlivePing();
 		void UpdateSubmixesAlivePingTime();
+		void SendSubmixesSendLevelTraceMessages();
 
 		void RegisterSubmixEnvelopeFollowerCVars();
 		void UnregisterSubmixEnvelopeFollowerCVars();
