[Audio render scheduler] Improve logging, add au.RenderScheudler.dump command to dump scheduler state.

#rnx
#Synced-CL 48092843

diff --git a/Engine/Source/Runtime/AudioMixer/Private/AudioMixerBus.cpp b/Engine/Source/Runtime/AudioMixer/Private/AudioMixerBus.cpp
--- a/Engine/Source/Runtime/AudioMixer/Private/AudioMixerBus.cpp
+++ b/Engine/Source/Runtime/AudioMixer/Private/AudioMixerBus.cpp
@@ -254,7 +254,11 @@ namespace Audio
 
 	const TCHAR* FMixerAudioBus::GetRenderStepName()
 	{
+#if ENABLE_AUDIO_DEBUG
+		return *BusName;
+#else
 		return TEXT("FMixerAudioBus mixing");
+#endif
 	}
 
 #if UE_AUDIO_PROFILERTRACE_ENABLED
diff --git a/Engine/Source/Runtime/AudioMixer/Private/AudioMixerCVars.cpp b/Engine/Source/Runtime/AudioMixer/Private/AudioMixerCVars.cpp
--- a/Engine/Source/Runtime/AudioMixer/Private/AudioMixerCVars.cpp
+++ b/Engine/Source/Runtime/AudioMixer/Private/AudioMixerCVars.cpp
@@ -2,6 +2,9 @@
 
 #include "AudioMixerCVars.h"
 
+#include "AudioMixerDevice.h"
+#include "AudioRenderScheduler.h"
+#include "Engine/Engine.h"
 #include "HAL/IConsoleManager.h"
 
 namespace AudioMixerCVars
@@ -13,4 +16,21 @@ namespace AudioMixerCVars
 		TEXT("Use audio rendering scheduler to handle render dependencies.\n")
 		TEXT("0: Disabled, 1: Enabled"),
 		ECVF_Default);
+
+#if ENABLE_AUDIO_DEBUG
+	FAutoConsoleCommand AudioRenderSchedulerDumpCmd(
+		TEXT("au.RenderScheduler.dump"),
+		TEXT("Dump the internal state of audio render scheduler to the log."),
+		FConsoleCommandDelegate::CreateStatic(
+			[]()
+			{
+				FAudioDevice* AudioDevice = GEngine != nullptr ? GEngine->GetMainAudioDeviceRaw() : nullptr;
+				if (AudioDevice != nullptr)
+				{
+					static_cast<Audio::FMixerDevice*>(AudioDevice)->GetRenderScheduler().DumpSteps();
+				}
+			}
+		)
+	);
+#endif
 }
diff --git a/Engine/Source/Runtime/AudioMixer/Private/AudioMixerSourceManager.cpp b/Engine/Source/Runtime/AudioMixer/Private/AudioMixerSourceManager.cpp
--- a/Engine/Source/Runtime/AudioMixer/Private/AudioMixerSourceManager.cpp
+++ b/Engine/Source/Runtime/AudioMixer/Private/AudioMixerSourceManager.cpp
@@ -1499,17 +1499,17 @@ namespace Audio
 
 			SourceSubmixOutputBuffer.Reset(SourceSubmixOutputResetSettings);
 
-			if (AudioMixerCVars::UseRenderScheduler)
-			{
-				MixerDevice->GetRenderScheduler().AddStep(FAudioRenderStepId::FromTransmitterID(SourceInfo.MixerSourceBuffer->GetInstanceID()), &SourceInfo);
-			}
-
 #if AUDIO_MIXER_ENABLE_DEBUG_MODE
 			AUDIO_MIXER_CHECK(!SourceInfo.bIsDebugMode);
 			SourceInfo.bIsDebugMode = InitParams.bIsDebugMode;
 			SourceInfo.DebugName = InitParams.DebugName;
 #endif 
 
+			if (AudioMixerCVars::UseRenderScheduler)
+			{
+				MixerDevice->GetRenderScheduler().AddStep(FAudioRenderStepId::FromTransmitterID(SourceInfo.MixerSourceBuffer->GetInstanceID()), &SourceInfo);
+			}
+
 			AUDIO_MIXER_DEBUG_LOG(SourceId, TEXT("Is initializing"));
 		});
 	}
@@ -2870,7 +2870,10 @@ namespace Audio
 				continue;
 			}
 
-			AudioBus->Update();
+			if (!AudioMixerCVars::UseRenderScheduler)
+			{
+				AudioBus->Update();
+			}
 			
 #if UE_AUDIO_PROFILERTRACE_ENABLED
 			if (bChannelEnabled)
@@ -3915,11 +3918,11 @@ namespace Audio
 
 			// Now generate bus audio (bGenerateBuses = true)
 			GenerateSourceAudio(true);
-
-			// Update the buses now
-			UpdateBuses();
 		}
 
+		// Update the buses now
+		UpdateBuses();
+
 		// Let the plugin know we finished processing all sources
 		if (bUsingSpatializationPlugin)
 		{
diff --git a/Engine/Source/Runtime/AudioMixer/Private/AudioRenderScheduler.cpp b/Engine/Source/Runtime/AudioMixer/Private/AudioRenderScheduler.cpp
--- a/Engine/Source/Runtime/AudioMixer/Private/AudioRenderScheduler.cpp
+++ b/Engine/Source/Runtime/AudioMixer/Private/AudioRenderScheduler.cpp
@@ -86,13 +86,21 @@ namespace Audio
 		check(Id.IsValid() && Step != nullptr);
 		FInternalIndex Index = GetInternalIndex(Id);
 		FStepEntry& Entry = Steps[Index];
-		if (ensureMsgf(Entry.StepInterface == nullptr, TEXT("Duplicate render step added to audio render scheduler")))
+		if (Entry.StepInterface == nullptr)
 		{
-			UE_LOGFMT(LogAudioRenderScheduler, Verbose, "Adding render step '{Name}' with id {Id}", Step->GetRenderStepName(), Id.GetRawValue());
+			UE_LOGFMT(LogAudioRenderScheduler, Verbose, "Adding render step '{Name}' with id {Id}",
+				("Name", Step->GetRenderStepName()),
+				("Id", Id.GetRawValue()));
 			Entry.StepInterface = Step;
 			Entry.ReferenceCount++;
 			bDirty = true;
 		}
+		else
+		{
+			UE_LOGFMT(LogAudioRenderScheduler, Error, "Tried to add duplicate render step '{Name}' with id {Id}",
+				("Name", Step->GetRenderStepName()),
+				("Id", Id.GetRawValue()));
+		}
 	}
 
 	void FAudioRenderScheduler::RemoveStep(const FAudioRenderStepId Id)
@@ -102,13 +110,20 @@ namespace Audio
 		check(Id.IsValid());
 		FInternalIndex Index = GetInternalIndex(Id);
 		FStepEntry& Entry = Steps[Index];
-		if (ensureMsgf(Entry.StepInterface != nullptr, TEXT("Unknown render step removed from audio render scheduler")))
+		if (Entry.StepInterface != nullptr)
 		{
-			UE_LOGFMT(LogAudioRenderScheduler, Verbose, "Removing render step '{Name}' with id {Id}", Entry.StepInterface->GetRenderStepName(), Id.GetRawValue());
+			UE_LOGFMT(LogAudioRenderScheduler, Verbose, "Removing render step '{Name}' with id {Id}",
+				("Name", Entry.StepInterface->GetRenderStepName()),
+				("Id", Id.GetRawValue()));
 			Entry.StepInterface = nullptr;
 			RemoveReference(Index);
 			bDirty = true;
 		}
+		else
+		{
+			UE_LOGFMT(LogAudioRenderScheduler, Error, "Tried to remove unknown render step with id {Id}",
+				("Id", Id.GetRawValue()));
+		}
 	}
 
 	void FAudioRenderScheduler::AddDependency(const FAudioRenderStepId FirstStep, const FAudioRenderStepId SecondStep)
@@ -145,11 +160,17 @@ namespace Audio
 		const FInternalIndex FirstIndex = GetInternalIndex(FirstStep);
 		const FInternalIndex SecondIndex = GetInternalIndex(SecondStep);
 
-		if (ensureMsgf(Steps[SecondIndex].Prerequisites.RemoveSingleSwap(FirstIndex), TEXT("Unknown dependency removed from audio render scheduler")))
+		if (Steps[SecondIndex].Prerequisites.RemoveSingleSwap(FirstIndex))
 		{
 			RemoveReference(FirstIndex);
 			RemoveReference(SecondIndex);
 		}
+		else
+		{
+			UE_LOGFMT(LogAudioRenderScheduler, Error, "Unknown dependency removed from audio render scheduler: {Id1} -> {Id2}",
+				("Id1", FirstStep.GetRawValue()),
+				("Id2", SecondStep.GetRawValue()));
+		}
 
 		bDirty = true;
 	}
@@ -172,6 +193,14 @@ namespace Audio
 
 			Task.Wait();
 		}
+
+#if ENABLE_AUDIO_DEBUG
+		if (bShouldDumpSteps)
+		{
+			DumpStepsInternal();
+			bShouldDumpSteps = false;
+		}
+#endif
 	}
 
 	void FAudioRenderScheduler::ScheduleAll(const bool bSingleThreaded)
@@ -338,7 +367,9 @@ namespace Audio
 			if (VisitedIndices.Contains(PrerequisiteStep))
 			{
 				// The link between CurrentStep and PrerequisiteStep is part of a cycle, so we'll break it.
-				UE_LOGFMT(LogAudioRenderScheduler, Verbose, "Breaking cycle between step {Id1} and step {Id2}", Steps[CurrentStep].StepId.GetRawValue(), Steps[PrerequisiteStep].StepId.GetRawValue());
+				UE_LOGFMT(LogAudioRenderScheduler, Display, "Breaking cycle between step {Id1} and step {Id2}",
+					("Id1", Steps[CurrentStep].StepId.GetRawValue()),
+					("Id2", Steps[PrerequisiteStep].StepId.GetRawValue()));
 				BreakLink(PrerequisiteStep, CurrentStep);
 				BrokenLinks.Add(TPair<FInternalIndex, FInternalIndex>{PrerequisiteStep, CurrentStep});
 				return;
@@ -407,4 +438,30 @@ namespace Audio
 
 		return false;
 	}
+
+#if ENABLE_AUDIO_DEBUG
+	void FAudioRenderScheduler::DumpSteps()
+	{
+		bShouldDumpSteps = true;
+	}
+
+	void FAudioRenderScheduler::DumpStepsInternal()
+	{
+		for (const FStepEntry& Step : Steps)
+		{
+			if (Step.ReferenceCount == 0)
+			{
+				continue;
+			}
+
+			UE_LOGFMT(LogAudioRenderScheduler, Display, "Index: {Index} Id: {Id} Name: '{Name}' Prerequisites: {Prerequisites}",
+				("Index", int(&Step - &Steps[0])),
+				("Id", Step.StepId.GetRawValue()),
+				("Name", Step.StepInterface ? Step.StepInterface->GetRenderStepName() : TEXT("(Empty)")),
+				("Prerequisites", Step.Prerequisites));
+		}
+
+		UE_LOGFMT(LogAudioRenderScheduler, Display, "Step ordering: {Steps}", ("Steps", ActiveSteps));
+	}
+#endif
 }
diff --git a/Engine/Source/Runtime/AudioMixer/Private/AudioRenderScheduler.h b/Engine/Source/Runtime/AudioMixer/Private/AudioRenderScheduler.h
--- a/Engine/Source/Runtime/AudioMixer/Private/AudioRenderScheduler.h
+++ b/Engine/Source/Runtime/AudioMixer/Private/AudioRenderScheduler.h
@@ -93,6 +93,10 @@ namespace Audio
 		// Execute all steps and block until they're done.
 		void RenderBlock(const bool bSingleThreaded);
 
+#if ENABLE_AUDIO_DEBUG
+		// Dump the state of the scheduler after next run
+		void DumpSteps();
+#endif
 	private:
 		// Internally the step id's are mapped to smaller internal indices to reduce
 		//  memory footprint.
@@ -111,6 +115,12 @@ namespace Audio
 		bool FindIndirectLink(const FInternalIndex FirstStep, const FInternalIndex SecondStep);
 		bool FindIndirectLinkSkipping(const FInternalIndex FirstStep, const FInternalIndex SecondStep, TArray<FInternalIndex>& StepsToSkip);
 
+#if ENABLE_AUDIO_DEBUG
+		void DumpStepsInternal();
+
+		bool bShouldDumpSteps = false;
+#endif
+
 		struct FStepEntry
 		{
 			FAudioRenderStepId StepId;
