Add additional logging for source initialization
#rnx

diff --git a/Engine/Source/Runtime/AudioMixer/Private/AudioMixerSource.cpp b/Engine/Source/Runtime/AudioMixer/Private/AudioMixerSource.cpp
--- a/Engine/Source/Runtime/AudioMixer/Private/AudioMixerSource.cpp
+++ b/Engine/Source/Runtime/AudioMixer/Private/AudioMixerSource.cpp
@@ -627,6 +627,7 @@ namespace Audio
 			InitParams.MixerSourceBuffer = MixerSourceBuffer;
 			MixerSourceBuffer = nullptr;
 
+			UE_LOG(LogAudioMixer, Verbose, TEXT("Initializing mixer source voice: %s"), *InWaveInstance->GetName());
 			if (MixerSourceVoice->Init(InitParams))
 			{
 				// Initialize the propagation interface as soon as we have a valid source id
