[Audio Insights] Bughawk fix - make sure we're sending the right audio bus ID when sending bus send info trace messages

#Synced-CL 48646572

diff --git a/Engine/Source/Runtime/AudioMixer/Private/AudioMixerSourceManager.cpp b/Engine/Source/Runtime/AudioMixer/Private/AudioMixerSourceManager.cpp
--- a/Engine/Source/Runtime/AudioMixer/Private/AudioMixerSourceManager.cpp
+++ b/Engine/Source/Runtime/AudioMixer/Private/AudioMixerSourceManager.cpp
@@ -1439,7 +1439,7 @@ namespace Audio
 						TraceAudioBusStartMessage(MixerDevice->DeviceID, AudioBusSend.AudioBusId, AudioBusSend.AudioBusName, AudioBusSend.BusChannels);
 					}
 
-					TraceSourceToBusSendMessage(MixerDevice->DeviceID, SourceInfo.PlayOrder, SourceInfo.AudioBusId, AudioBusSend.SendLevel);
+					TraceSourceToBusSendMessage(MixerDevice->DeviceID, SourceInfo.PlayOrder, AudioBusSend.AudioBusId, AudioBusSend.SendLevel);
 #endif // UE_AUDIO_PROFILERTRACE_ENABLED
 					
 					// Store on this source, which buses its sending its audio to
