Bughawk fix for using array Reset and thus not setting array size where AddZeroed should be utilized
#rnx
#Synced-CL 48675545

diff --git a/Engine/Source/Runtime/AudioMixer/Private/AudioMixerSourceManager.cpp b/Engine/Source/Runtime/AudioMixer/Private/AudioMixerSourceManager.cpp
--- a/Engine/Source/Runtime/AudioMixer/Private/AudioMixerSourceManager.cpp
+++ b/Engine/Source/Runtime/AudioMixer/Private/AudioMixerSourceManager.cpp
@@ -3231,7 +3231,7 @@ namespace Audio
 		if (SourceInfo.AudioBusSends[(int32)EBusSendType::PreEffect].Num() > 0)
 		{
 			SourceInfo.PreEffectBuffer.Reset();
-			SourceInfo.PreEffectBuffer.Reserve(SourceInfo.PreDistanceAttenuationBuffer.Num());
+			SourceInfo.PreEffectBuffer.AddZeroed(SourceInfo.PreDistanceAttenuationBuffer.Num());
 
 			FMemory::Memcpy(SourceInfo.PreEffectBuffer.GetData(), SourceInfo.PreDistanceAttenuationBuffer.GetData(), sizeof(float) * SourceInfo.PreDistanceAttenuationBuffer.Num());
 		}
