[AudioInsights] Auto-Enable all submix envelope following

* Move existing auto-enable code into FMixerDevice::RebuildSubmixLinks()
* Intentionally leveraging the existing cvar code, instead of ripping it out
* As in future we still may want to re-use this approach if/when consoles support cvars.
* Provide overloaded functions so we can target single Submixes here and in future
* Passes Bughawk
* Fixes: https://jira.it.epicgames.com/browse/UE-353705

* Tested with AudioQAModulation level simulating console builds by disabling dashboard cvar code
* Submixes can still be monitored

#Synced-CL 49030512

diff --git a/Engine/Source/Runtime/AudioMixer/Private/AudioMixerDevice.cpp b/Engine/Source/Runtime/AudioMixer/Private/AudioMixerDevice.cpp
--- a/Engine/Source/Runtime/AudioMixer/Private/AudioMixerDevice.cpp
+++ b/Engine/Source/Runtime/AudioMixer/Private/AudioMixerDevice.cpp
@@ -547,17 +547,13 @@ namespace Audio
 	{
 		ensure(IsInAudioThread());
 
-		// Attack/Release default times used in USoundSubmix constructor
-		constexpr int32 AttackTime  = 10;
-		constexpr int32 ReleaseTime = 500;
-
-		Submixes.Iterate([this, &InSubmixName, AttackTime, ReleaseTime](const FSubmixMap::FPair& SubmixPair)
+		Submixes.Iterate([this, &InSubmixName](const FSubmixMap::FPair& SubmixPair)
 		{
 			const FMixerSubmixPtr& Submix = SubmixPair.Value;
 
 			if (Submix.IsValid() && Submix->GetName() == InSubmixName)
 			{
-				Submix->StartEnvelopeFollowing(AttackTime, ReleaseTime);
+				StartSubmixEnvelopeFollowing(Submix);
 			}
 		});
 
@@ -565,7 +561,7 @@ namespace Audio
 		{
 			if (RequiredSubmix.IsValid() && RequiredSubmix->GetName() == InSubmixName)
 			{
-				RequiredSubmix->StartEnvelopeFollowing(AttackTime, ReleaseTime);
+				StartSubmixEnvelopeFollowing(RequiredSubmix);
 			}
 		}
 	}
@@ -580,7 +576,7 @@ namespace Audio
 
 			if (Submix.IsValid() && Submix->GetName() == InSubmixName)
 			{
-				Submix->StopEnvelopeFollowing();
+				StopSubmixEnvelopeFollowing(Submix);
 			}
 		});
 
@@ -588,10 +584,39 @@ namespace Audio
 		{
 			if (RequiredSubmix.IsValid() && RequiredSubmix->GetName() == InSubmixName)
 			{
-				RequiredSubmix->StopEnvelopeFollowing();
+				StopSubmixEnvelopeFollowing(RequiredSubmix);
 			}
 		}
 	}
+
+	void FMixerDevice::StartSubmixEnvelopeFollowing(const FMixerSubmixPtr& InSubmix)
+	{
+		ensure(IsInAudioThread());
+
+		if (!InSubmix || !InSubmix->IsValid())
+		{
+			return;
+		}
+
+		// Attack/Release default times used in USoundSubmix constructor
+		constexpr int32 AttackTime  = 10;
+		constexpr int32 ReleaseTime = 500;
+
+		InSubmix->StartEnvelopeFollowing(AttackTime, ReleaseTime);
+	}
+
+	void FMixerDevice::StopSubmixEnvelopeFollowing(const FMixerSubmixPtr& InSubmix)
+	{
+		ensure(IsInAudioThread());
+
+		if (!InSubmix || !InSubmix->IsValid())
+		{
+			return;
+		}
+
+		InSubmix->StopEnvelopeFollowing();
+	}
+
 #endif // UE_AUDIO_PROFILERTRACE_ENABLED
 
 	void FMixerDevice::AddReferencedObjects(FReferenceCollector& Collector)
@@ -1567,15 +1592,6 @@ namespace Audio
 				}
 			}
 		}
-
-#if UE_AUDIO_PROFILERTRACE_ENABLED
-		// Auto-Enable MainSubmix output metering
-		FAudioThread::RunCommandOnAudioThread([this]()
-		{
-			StartSubmixEnvelopeFollowerFromCVar(GetRequiredSubmixName(ERequiredSubmixes::Main));
-		});
-#endif // UE_AUDIO_PROFILERTRACE_ENABLED
-
 	}
 
 	void FMixerDevice::RebuildSubmixLinks(const USoundSubmixBase& SoundSubmix, FMixerSubmixPtr& SubmixInstance)
@@ -1607,6 +1623,15 @@ namespace Audio
 			SubmixInstance->SetParentSubmix(ParentSubmixInstance);
 			ParentSubmixInstance->AddChildSubmix(SubmixInstance);
 		}
+
+#if UE_AUDIO_PROFILERTRACE_ENABLED
+		// Auto-Enable Submix output metering
+		FAudioThread::RunCommandOnAudioThread([this, ParentSubmixInstance, SubmixInstance]()
+		{
+			StartSubmixEnvelopeFollowing(ParentSubmixInstance);
+			StartSubmixEnvelopeFollowing(SubmixInstance);
+		});
+#endif // UE_AUDIO_PROFILERTRACE_ENABLED
 	}
 
  	FAudioPlatformSettings FMixerDevice::GetPlatformSettings() const
diff --git a/Engine/Source/Runtime/AudioMixer/Public/AudioMixerDevice.h b/Engine/Source/Runtime/AudioMixer/Public/AudioMixerDevice.h
--- a/Engine/Source/Runtime/AudioMixer/Public/AudioMixerDevice.h
+++ b/Engine/Source/Runtime/AudioMixer/Public/AudioMixerDevice.h
@@ -419,6 +419,9 @@ namespace Audio
 
 		void StartSubmixEnvelopeFollowerFromCVar(const FString& InSubmixName);
 		void StopSubmixEnvelopeFollowerFromCVar(const FString& InSubmixName);
+
+		void StartSubmixEnvelopeFollowing(const FMixerSubmixPtr& InSubmix);
+		void StopSubmixEnvelopeFollowing(const FMixerSubmixPtr& InSubmix);
 #endif // UE_AUDIO_PROFILERTRACE_ENABLED
 
 		TArray<USoundSubmix*> RequiredSubmixes;
