soundwave scrubber crash fixes

- Added less fatal checks and more graceful ensures and fallbacks
- Slight rounding errors can ocurr that rarely hit some of the aggressive edges

[RN] Fixed edge cases in sound file scrubbing sequencer scrubbing
#Synced-CL 48995942

diff --git a/Engine/Source/Runtime/AudioMixer/Private/Generators/SoundWaveScrubber.cpp b/Engine/Source/Runtime/AudioMixer/Private/Generators/SoundWaveScrubber.cpp
--- a/Engine/Source/Runtime/AudioMixer/Private/Generators/SoundWaveScrubber.cpp
+++ b/Engine/Source/Runtime/AudioMixer/Private/Generators/SoundWaveScrubber.cpp
@@ -159,8 +159,11 @@ void FSoundWaveScrubber::DecodeToDataChunk(FDecodedDataChunk& InOutDataChunk, fl
 
 FSoundWaveScrubber::FGrain FSoundWaveScrubber::SpawnGrain()
 {
-	// Make sure we don't have a decoder error when we're trying to spawn grains
-	check(!bHasErrorWithDecoder);
+	// If we have a decoder error, spawn null grains vs asserting
+	if (bHasErrorWithDecoder)
+	{
+		return { 0 };
+	}
 	
 	// Try to retrieve a decoded data chunk for the current read frame based on the curret, interpolated playhead time seconds value
 	int32 CurrentReadFrame = CurrentPlayheadTimeSeconds.GetValue() * AudioMixerSampleRate;
@@ -169,7 +172,11 @@ FSoundWaveScrubber::FGrain FSoundWaveScrubber::SpawnGrain()
 	{
 		DecodedDataChunkIndex = DecodeDataChunkIndexForCurrentReadIndex(CurrentReadFrame);
 	}
-	check(DecodedDataChunkIndex != INDEX_NONE);
+
+	if (!ensureMsgf(DecodedDataChunkIndex != INDEX_NONE, TEXT("Failed to retreived decoded chunk for read frame index: %d"), CurrentReadFrame))
+	{
+		return { 0 };
+	}
 
 	// Get the grain runtime data for the new grain spawn
 	FGrain NewGrain;
@@ -291,7 +298,11 @@ int32 FSoundWaveScrubber::RenderAudio(TArrayView<float>& OutAudio)
 	{
 		if (!ActiveGrains.Num())
 		{
-			ActiveGrains.Add(SpawnGrain());
+			FSoundWaveScrubber::FGrain NewGrain = SpawnGrain();
+			if (NewGrain.GrainDurationFrames > 0)
+			{
+				ActiveGrains.Add(NewGrain);
+			}
 
 			NumFramesTillNextGrainSpawn = CurrentHalfGrainDurationFrames;
 		}
@@ -303,9 +314,12 @@ int32 FSoundWaveScrubber::RenderAudio(TArrayView<float>& OutAudio)
 
 		while (RemainingFrames > 0)
 		{
-			// This renders the currently active grains (which should only ever max 2 grains)
-			// starting from the given start render frame and for the  numbe of indicated frames
-			RenderActiveGrains(OutAudio, StartRenderFrame, NumFramesToRender);
+			if (ActiveGrains.Num() > 0)
+			{
+				// This renders the currently active grains (which should only ever max 2 grains)
+				// starting from the given start render frame and for the  numbe of indicated frames
+				RenderActiveGrains(OutAudio, StartRenderFrame, NumFramesToRender);
+			}
 
 			// Update the number of frames rendered this render block
 			NumFramesRendered += NumFramesToRender;
@@ -321,7 +335,11 @@ int32 FSoundWaveScrubber::RenderAudio(TArrayView<float>& OutAudio)
 			{
 				// Spawn a new grain with a frame start 
 				StartRenderFrame = NumFrames - RemainingFrames;
-				ActiveGrains.Add(SpawnGrain());
+				FSoundWaveScrubber::FGrain NewGrain = SpawnGrain(); 
+				if (NewGrain.GrainDurationFrames > 0)
+				{
+					ActiveGrains.Add(NewGrain);
+				}
 
 				// If we still have frames remaining in the buffer, then 
 				NumFramesTillNextGrainSpawn = CurrentHalfGrainDurationFrames;
@@ -336,6 +354,15 @@ int32 FSoundWaveScrubber::RenderAudio(TArrayView<float>& OutAudio)
 
 void FSoundWaveScrubber::UpdateGrainDecodeData(FGrain& InGrain)
 {
+	// Check for invalid chunk index 
+	if (!ensureMsgf(DecodedChunks.IsValidIndex(InGrain.DecodedDataChunkIndex), TEXT("Decoded data chunk index was invalid: %d"), InGrain.DecodedDataChunkIndex))
+	{
+		// If the index is invalid, finish the grain by setting it's rendered frame count to the grain duration. This kills the grain.
+		InGrain.DecodedDataChunkIndex = INDEX_NONE;
+		InGrain.CurrentRenderedFramesCount = InGrain.GrainDurationFrames;
+		return;
+	}
+
 	FDecodedDataChunk& DecodedData = DecodedChunks[InGrain.DecodedDataChunkIndex];
 
 	// Total number of frames of decoded data in the chunk
@@ -344,25 +371,63 @@ void FSoundWaveScrubber::UpdateGrainDecodeData(FGrain& InGrain)
 	// The number of frames that this grain is offset from the decoded data
 	int32 NumFramesOffsetInDecodedData = InGrain.CurrentReadFrame - DecodedData.FrameStart;
 	int32 NumFramesPossibleToRenderInChunk = NumReadFramesInDecodedData - NumFramesOffsetInDecodedData;
-	check(NumFramesPossibleToRenderInChunk >= 0);
 
+	// If our current grain chunk is out of range of the current chunk, lets re-acquire the chunk
+	// note this can happen to slight SRC rounding errors, etc.
+	if (NumFramesOffsetInDecodedData < 0 || NumFramesPossibleToRenderInChunk < 0)
+	{
+		// first decrement the current decoded data chunk
+		DecodedData.NumGrainsUsingChunk = FMath::Max(DecodedData.NumGrainsUsingChunk - 1, 0);
+		int32 NewIndex = GetDecodedDataChunkIndexForCurrentReadIndex(InGrain.CurrentReadFrame);
+		if (NewIndex == INDEX_NONE)
+		{
+			NewIndex = DecodeDataChunkIndexForCurrentReadIndex(InGrain.CurrentReadFrame);
+		}
+
+		// Failed to reacquire the decode chunk
+		if (!ensureMsgf(NewIndex != INDEX_NONE, TEXT("Failed to reacquire decode chunk in soundwave scrubber at current read frame: %d"), InGrain.CurrentReadFrame))
+		{
+			InGrain.DecodedDataChunkIndex = INDEX_NONE;
+			InGrain.CurrentRenderedFramesCount = InGrain.GrainDurationFrames;
+			return;
+		}
+
+		// Update the grain to the refreshed decode chunk
+		InGrain.DecodedDataChunkIndex = NewIndex;
+		DecodedChunks[InGrain.DecodedDataChunkIndex].NumGrainsUsingChunk++;
+	}
 	// If we've totally consumed this decoded audio chunk, we need to get a new decoded audio chunk
-	if (!NumFramesPossibleToRenderInChunk)
+	else if (NumFramesPossibleToRenderInChunk == 0)
 	{
 		// We're no longer using this decoded audio chunk
-		DecodedData.NumGrainsUsingChunk--;
-		check(DecodedData.NumGrainsUsingChunk >= 0);
+		DecodedData.NumGrainsUsingChunk  = FMath::Max(DecodedData.NumGrainsUsingChunk - 1, 0);
 
-		InGrain.DecodedDataChunkIndex = GetDecodedDataChunkIndexForCurrentReadIndex(InGrain.CurrentReadFrame);
-		if (InGrain.DecodedDataChunkIndex == INDEX_NONE)
+		int32 NewIndex = GetDecodedDataChunkIndexForCurrentReadIndex(InGrain.CurrentReadFrame);
+		if (NewIndex == INDEX_NONE)
 		{
-			InGrain.DecodedDataChunkIndex = DecodeDataChunkIndexForCurrentReadIndex(InGrain.CurrentReadFrame);
+			NewIndex = DecodeDataChunkIndexForCurrentReadIndex(InGrain.CurrentReadFrame);
+		}
+		if (!ensureMsgf(NewIndex != INDEX_NONE, TEXT("Couldn't find a decoded data chunk for current read frame: %d"), InGrain.CurrentReadFrame))
+		{
+			// Kill the grain
+			InGrain.DecodedDataChunkIndex = INDEX_NONE;
+			InGrain.CurrentRenderedFramesCount = InGrain.GrainDurationFrames;
+			return;
 		}
-		check(InGrain.DecodedDataChunkIndex != INDEX_NONE);
 
-		check(InGrain.CurrentReadFrame >= DecodedChunks[InGrain.DecodedDataChunkIndex].FrameStart);
-		check(InGrain.CurrentReadFrame < DecodedChunks[InGrain.DecodedDataChunkIndex].FrameStart + DecodedChunks[InGrain.DecodedDataChunkIndex].PCMAudio.Num() / NumChannels);
+		// Defensive, non-fatal range validation
+		// We very rarely get out-of-range due to slight rounding errors.
+		const int32 NewFrameStart = DecodedChunks[NewIndex].FrameStart;
+		const int32 NewFrameEnd = NewFrameStart + DecodedChunks[NewIndex].PCMAudio.Num() / NumChannels;
+		if (!ensureMsgf(InGrain.CurrentReadFrame >= NewFrameStart && InGrain.CurrentReadFrame < NewFrameEnd, TEXT("Current read frame (%d) is out of range [%d, %d)."), InGrain.CurrentReadFrame, NewFrameStart, NewFrameEnd))
+		{
+			// We're out of range, kill the grain
+			InGrain.CurrentRenderedFramesCount = InGrain.GrainDurationFrames;
+			return;
+		}
 
+		// this chunk is good and we're in the right range
+		InGrain.DecodedDataChunkIndex = NewIndex;
 		DecodedChunks[InGrain.DecodedDataChunkIndex].NumGrainsUsingChunk++;
 	}
 }
@@ -379,16 +444,23 @@ void FSoundWaveScrubber::RenderActiveGrains(TArrayView<float>& OutAudio, int32 I
 		// This is the number of frames we have left to render for this grain
 		int32 NumFramesLeftInGrain = Grain.GrainDurationFrames - Grain.CurrentRenderedFramesCount;
 		int32 NumFramesLeftToRender = FMath::Min(NumFramesToRender, NumFramesLeftInGrain);
-		check(NumFramesLeftToRender > 0);
+		NumFramesLeftToRender = FMath::Max(NumFramesLeftToRender, 0);
 
 		int32 GrainWriteIndex = InStartFrame;
 
-		bool bGrainFinished = false;
+		bool bGrainFinished = (NumFramesLeftInGrain == 0);
 		while (NumFramesLeftToRender > 0 && !bGrainFinished)
 		{
 			// Make sure we have a valid decode data chunk ready for rendering
 			UpdateGrainDecodeData(Grain);
 
+			// It's possible we still don't have a valid decoded data chunk index
+			if (!DecodedChunks.IsValidIndex(Grain.DecodedDataChunkIndex))
+			{
+				bGrainFinished = true;
+				break;
+			}
+
 			// Retrieve the decoded data.
 			const FDecodedDataChunk& DecodedData = DecodedChunks[Grain.DecodedDataChunkIndex];
 
@@ -396,16 +468,13 @@ void FSoundWaveScrubber::RenderActiveGrains(TArrayView<float>& OutAudio, int32 I
 			int32 NumReadFramesInDecodedData = DecodedData.PCMAudio.Num() / NumChannels;
 
 			// The number of frames that this grain is offset from the decoded data
-			int32 NumFramesOffsetInDecodedData = Grain.CurrentReadFrame - DecodedData.FrameStart;
-			check(NumFramesOffsetInDecodedData >= 0);
-
-			int32 NumFramesPossibleToRenderInChunk = NumReadFramesInDecodedData - NumFramesOffsetInDecodedData;
-			check(NumFramesPossibleToRenderInChunk >= 0);
+			int32 NumFramesOffsetInDecodedData = FMath::Max(Grain.CurrentReadFrame - DecodedData.FrameStart, 0);
+			int32 NumFramesPossibleToRenderInChunk = FMath::Max(NumReadFramesInDecodedData - NumFramesOffsetInDecodedData, 0);
 
 			// This may completely consume the decoded chunk here, so only render the maximum number of frames in the 
 			// chunk. if NumFramesToRender is smaller than that, then we won't completely consume the decoded data
 			int32 NumFramesToRenderInThisChunk = FMath::Min(NumFramesPossibleToRenderInChunk, NumFramesLeftToRender);
-			check(NumFramesToRenderInThisChunk >= 0);
+			NumFramesToRenderInThisChunk = FMath::Max(NumFramesToRenderInThisChunk, 0);
 
 			int32 SampleWriteIndex = GrainWriteIndex * NumChannels;
 			int32 SampleReadIndex = NumFramesOffsetInDecodedData * NumChannels;
@@ -446,10 +515,13 @@ void FSoundWaveScrubber::RenderActiveGrains(TArrayView<float>& OutAudio, int32 I
 		// If the grain has finished, then remove it from the active grain list
 		if (bGrainFinished)
 		{
-			DecodedChunks[Grain.DecodedDataChunkIndex].NumGrainsUsingChunk--;
-			check(DecodedChunks[Grain.DecodedDataChunkIndex].NumGrainsUsingChunk >= 0);
+			if (DecodedChunks.IsValidIndex(Grain.DecodedDataChunkIndex))
+			{
+				FDecodedDataChunk& DataChunk = DecodedChunks[Grain.DecodedDataChunkIndex];
+				DataChunk.NumGrainsUsingChunk = FMath::Max(DataChunk.NumGrainsUsingChunk - 1, 0);
+			}
 			ActiveGrains.RemoveAtSwap(GrainIndex, EAllowShrinking::No);
-			NumActiveGrains--;
+			NumActiveGrains = FMath::Max(NumActiveGrains - 1, 0);
 		}
 	}
 }
