FMixerDevice - Remove Submixes from root when unloading

#rnx
#tests Cooked build
#Synced-CL 49267577

diff --git a/Engine/Source/Runtime/AudioMixer/Private/AudioMixerDevice.cpp b/Engine/Source/Runtime/AudioMixer/Private/AudioMixerDevice.cpp
--- a/Engine/Source/Runtime/AudioMixer/Private/AudioMixerDevice.cpp
+++ b/Engine/Source/Runtime/AudioMixer/Private/AudioMixerDevice.cpp
@@ -2394,6 +2394,12 @@ namespace Audio
 		}
 
 		Submixes.Remove(InSoundSubmix.GetUniqueID());
+		
+		if (InSoundSubmix.IsRooted())
+		{
+			USoundSubmixBase& MutableSubmix = const_cast<USoundSubmixBase&>(InSoundSubmix);
+			MutableSubmix.RemoveFromRoot();
+		}
 	}
 
 	FMixerSubmixPtr FMixerDevice::FindSubmixInstanceByObjectId(uint32 InObjectId)
