FMixerDevice - Fix misspelling of hierarchy

#rnx

#Synced-CL 49371692

diff --git a/Engine/Source/Runtime/AudioMixer/Private/AudioMixerDevice.cpp b/Engine/Source/Runtime/AudioMixer/Private/AudioMixerDevice.cpp
--- a/Engine/Source/Runtime/AudioMixer/Private/AudioMixerDevice.cpp
+++ b/Engine/Source/Runtime/AudioMixer/Private/AudioMixerDevice.cpp
@@ -242,7 +242,7 @@ namespace Audio
 
 	static FAutoConsoleCommand DumpSubmixCmd(
 		TEXT("au.submix.drawgraph"),
-		TEXT("Draws the submix heirarchy for this world to the debug output"),
+		TEXT("Draws the submix hierarchy for this world to the debug output"),
 		FConsoleCommandWithWorldArgsAndOutputDeviceDelegate::CreateLambda([](const TArray<FString>& InArgs, UWorld* InWorld, FOutputDevice& OutLog)
 		{
 			if (InWorld)
@@ -257,7 +257,7 @@ namespace Audio
 
 
 	
-	static void DrawSubmixHeirarchy(USoundSubmixBase* InSubmix, const TSharedPtr<Audio::FMixerSubmix, ESPMode::ThreadSafe> InInstance, const FMixerDevice* InDevice, int32 InIdent, FOutputDevice& Ar, const TCHAR* GroupingText)
+	static void DrawSubmixHierarchy(USoundSubmixBase* InSubmix, const TSharedPtr<Audio::FMixerSubmix, ESPMode::ThreadSafe> InInstance, const FMixerDevice* InDevice, int32 InIdent, FOutputDevice& Ar, const TCHAR* GroupingText)
 	{
 		if (!InSubmix)
 		{
@@ -277,12 +277,12 @@ namespace Audio
 		Ar.Logf(TEXT("%sName=%s,Instance=0x%p,Id=%u,Fx=%s,[%s]"), *Indet, *InSubmix->GetName(), InInstance.Get(), InInstance ? InInstance->GetId() : 0, *FxChain, GroupingText);
 		for (const auto& i : InSubmix->ChildSubmixes)
 		{
-			DrawSubmixHeirarchy(i, InDevice->GetSubmixInstance(i).Pin(), InDevice, InIdent+1, Ar, TEXT("Static"));
+			DrawSubmixHierarchy(i, InDevice->GetSubmixInstance(i).Pin(), InDevice, InIdent+1, Ar, TEXT("Static"));
 		}
 		const auto& DynamicSubmixes = InSubmix->DynamicChildSubmixes.FindOrAdd(InDevice->DeviceID);
 		for (const auto& i : DynamicSubmixes.ChildSubmixes)
 		{
-			DrawSubmixHeirarchy(i, InDevice->GetSubmixInstance(i).Pin(), InDevice, InIdent+1, Ar, TEXT("Dynamic"));
+			DrawSubmixHierarchy(i, InDevice->GetSubmixInstance(i).Pin(), InDevice, InIdent+1, Ar, TEXT("Dynamic"));
 		}
 	}
 
@@ -344,7 +344,7 @@ namespace Audio
 		for (int32 i = 0; i < RequiredSubmixes.Num(); i++)
 		{
 			InOutput.Logf(TEXT("SlotName=[%s]"), ToCStr(LexToString(static_cast<ERequiredSubmixes>(i))));
-			DrawSubmixHeirarchy(RequiredSubmixes[i], RequiredSubmixInstances[i], this, 1, InOutput, TEXT("In Slot"));
+			DrawSubmixHierarchy(RequiredSubmixes[i], RequiredSubmixInstances[i], this, 1, InOutput, TEXT("In Slot"));
 		}
 	}
 
