[Audio Insights] Fix master reverb submix not being connected inside the signal flow graph to the master output submix
* Make sure we also send the submix to submix send levels for the required submixes inside AudioMixerDevice - this fixes the rever submix having registered inputs but no registered outputs.
* Include the output of a dummy node when generating the unique key - using just the input and tree depth is not enough to avoid duplicate keys being generated.

#Synced-CL 49648429

diff --git a/Engine/Plugins/AudioInsights/Source/AudioInsights/Private/Views/SignalFlowDashboardViewFactory.cpp b/Engine/Plugins/AudioInsights/Source/AudioInsights/Private/Views/SignalFlowDashboardViewFactory.cpp
--- a/Engine/Plugins/AudioInsights/Source/AudioInsights/Private/Views/SignalFlowDashboardViewFactory.cpp
+++ b/Engine/Plugins/AudioInsights/Source/AudioInsights/Private/Views/SignalFlowDashboardViewFactory.cpp
@@ -979,7 +979,7 @@ namespace UE::Audio::Insights
 
 		// Create a dummy node in between the input and output nodes
 		// We can use these dummy connection nodes to help order and position real nodes in the graph and reduce connections crossing
-		const FSignalFlowEntryKey DummyConnectionEntryKey = FSignalFlowEntryKey(InputKey.DeviceID, TreeDepthPair.LayerID, HashCombine(GetTypeHash(InputKey), TreeDepthPair.TreeDepth));
+		const FSignalFlowEntryKey DummyConnectionEntryKey = FSignalFlowEntryKey(InputKey.DeviceID, TreeDepthPair.LayerID, HashCombine(GetTypeHash(InputKey), GetTypeHash(OutputKey), TreeDepthPair.TreeDepth));
 		TSharedPtr<FDummyConnectionNode> DummyConnectionNode = MakeShared<FDummyConnectionNode>(DummyConnectionEntryKey, TreeDepthPair.TreeDepth, TreeDepthPair.LayerID, InputNode->Entry->Timestamp);
 		
 		DummyConnectionNode->ConnectionOutputKey = ConnectionOutputNodeKey;
diff --git a/Engine/Source/Runtime/AudioMixer/Private/AudioMixerDevice.cpp b/Engine/Source/Runtime/AudioMixer/Private/AudioMixerDevice.cpp
--- a/Engine/Source/Runtime/AudioMixer/Private/AudioMixerDevice.cpp
+++ b/Engine/Source/Runtime/AudioMixer/Private/AudioMixerDevice.cpp
@@ -457,15 +457,14 @@ namespace Audio
 		const bool bChannelEnabled = UE_TRACE_CHANNELEXPR_IS_ENABLED(AudioChannel);
 		if (bChannelEnabled)
 		{
-			Submixes.Iterate([this](const FSubmixMap::FPair& SubmixPair)
+			auto TraceSubmixToSubmixSendMessage = [this](const FMixerSubmixPtr& InSubmix)
 			{
-				const FMixerSubmixPtr& Submix = SubmixPair.Value;
-				if (!Submix.IsValid())
+				if (!InSubmix.IsValid())
 				{
 					return;
 				}
 
-				const Audio::FMixerSubmixPtr OutputSubmix = Submix->GetParentSubmix().Pin();
+				const Audio::FMixerSubmixPtr OutputSubmix = InSubmix->GetParentSubmix().Pin();
 				
 				// If this submix has no parent, send the Audio Device as the receiver ID and mark as a final stage of the chain
 				const bool bIsFinalStage = !OutputSubmix.IsValid();
@@ -474,11 +473,21 @@ namespace Audio
 				UE_TRACE_LOG(Audio, SubmixToSubmixSend, AudioMixerChannel)
 					<< SubmixToSubmixSend.DeviceId(DeviceID)
 					<< SubmixToSubmixSend.Timestamp(FPlatformTime::Cycles64())
-					<< SubmixToSubmixSend.SenderID(GetTypeHash(Submix->GetPathName()))
+					<< SubmixToSubmixSend.SenderID(GetTypeHash(InSubmix->GetPathName()))
 					<< SubmixToSubmixSend.ReceiverID(ReceiverID)
-					<< SubmixToSubmixSend.SendLevel(Submix->CurrentOutputVolume)
+					<< SubmixToSubmixSend.SendLevel(InSubmix->CurrentOutputVolume)
 					<< SubmixToSubmixSend.IsFinalStage(bIsFinalStage);
+			};
+
+			Submixes.Iterate([&TraceSubmixToSubmixSendMessage](const FSubmixMap::FPair& SubmixPair)
+			{
+				TraceSubmixToSubmixSendMessage(SubmixPair.Value);
 			});
+
+			for (const FMixerSubmixPtr& RequiredSubmixInstance : RequiredSubmixInstances)
+			{
+				TraceSubmixToSubmixSendMessage(RequiredSubmixInstance);
+			}
 		}
 	}
 
