AudioMixerDevice - Add ref count to submixes using the register/unregister API. When no longer referenced, submixes will be removed from root.

#rnx

diff --git a/Engine/Plugins/Runtime/ResonanceAudio/Source/ResonanceAudio/Private/ResonanceAudioReverb.cpp b/Engine/Plugins/Runtime/ResonanceAudio/Source/ResonanceAudio/Private/ResonanceAudioReverb.cpp
--- a/Engine/Plugins/Runtime/ResonanceAudio/Source/ResonanceAudio/Private/ResonanceAudioReverb.cpp
+++ b/Engine/Plugins/Runtime/ResonanceAudio/Source/ResonanceAudio/Private/ResonanceAudioReverb.cpp
@@ -5,6 +5,7 @@
 #include "ResonanceAudioReverb.h"
 
 #include "AudioDeviceManager.h"
+#include "AudioMixerDevice.h"
 #include "ResonanceAudioCommon.h"
 #include "ResonanceAudioModule.h"
 #include "ResonanceAudioSettings.h"
@@ -38,8 +39,7 @@ namespace ResonanceAudio
 	FResonanceAudioReverb::~FResonanceAudioReverb()
 	{
 		ReverbPluginPreset = nullptr;
-
-		if (ReverbSubmixPtr)
+		if (!Audio::GRemoveSubmixesFromRoot && ReverbSubmixPtr)
 		{
 			ReverbSubmixPtr->SetParentSubmix(nullptr, false /* bModifyAssets */);
 			ReverbSubmixPtr->RemoveFromRoot();
@@ -215,7 +215,12 @@ namespace ResonanceAudio
 			}
 
 			ReverbSubmixPtr = NewObject<USoundSubmix>();
-			ReverbSubmixPtr->AddToRoot();
+			// This AddToRoot is redundant since we AddToRoot in the LoadSoundSubmix stack in the AudioMixerDevice.
+			// LoadSoundSubmix gets called immediately after this function in FMixerDevice::LoadPluginSoundSubmixes.
+			if (!Audio::GRemoveSubmixesFromRoot)
+			{
+				ReverbSubmixPtr->AddToRoot();
+			}
 			ReverbSubmixPtr->bMuteWhenBackgrounded = true;
 
 			if (const FAudioDeviceManager* DeviceManager = FAudioDeviceManager::Get())
diff --git a/Engine/Source/Runtime/AudioMixer/Private/AudioMixerDevice.cpp b/Engine/Source/Runtime/AudioMixer/Private/AudioMixerDevice.cpp
--- a/Engine/Source/Runtime/AudioMixer/Private/AudioMixerDevice.cpp
+++ b/Engine/Source/Runtime/AudioMixer/Private/AudioMixerDevice.cpp
@@ -160,6 +160,14 @@ CSV_DECLARE_CATEGORY_MODULE_EXTERN(AUDIOMIXERCORE_API, Audio);
 
 namespace Audio
 {
+	int32 GRemoveSubmixesFromRoot = 1;
+	FAutoConsoleVariableRef CVarRemoveSubmixesFromRoot(
+		TEXT("au.RemoveSubmixesFromRoot"),
+		GRemoveSubmixesFromRoot,
+		TEXT("Submixes that are registered will use a ref counter. When unregistering, the ref counter is decremented. When no longer referenced, the submixes are removed from root.\n")
+		TEXT("1: Submixes are removed from root when no longer referenced. 0: Submixes remain rooted."),
+		ECVF_Default);
+	
 	void FSubmixMap::Add(const FSubmixMap::FObjectId InObjectId, FMixerSubmixPtr InMixerSubmix)
 	{
 		if (DisableSubmixMutationLockCVar)
@@ -239,7 +247,6 @@ namespace Audio
 		return TSet<FObjectId>(Keys);
 	}
 
-
 	static FAutoConsoleCommand DumpSubmixCmd(
 		TEXT("au.submix.drawgraph"),
 		TEXT("Draws the submix hierarchy for this world to the debug output"),
@@ -1432,7 +1439,12 @@ namespace Audio
 			LLM_SCOPE(ELLMTag::AudioMixerPlugins);
 			USoundSubmix* ReverbPluginSubmix = ReverbPluginInterface->LoadSubmix();
 			check(ReverbPluginSubmix);
-			ReverbPluginSubmix->AddToRoot();
+			
+			// This AddToRoot is redundant since we AddToRoot in the LoadSoundSubmix stack
+			if (!GRemoveSubmixesFromRoot)
+			{
+				ReverbPluginSubmix->AddToRoot();
+			}
 
 			LoadSoundSubmix(*ReverbPluginSubmix);
 			InitSoundfieldAndEndpointDataForSubmix(*ReverbPluginSubmix, GetSubmixInstance(ReverbPluginSubmix).Pin(), false);
@@ -2258,7 +2270,14 @@ namespace Audio
 		FMixerSubmixPtr MixerSubmix = GetSubmixInstance(&InSoundSubmix).Pin();
 		if (!MixerSubmix.IsValid())
 		{
-			InSoundSubmix.AddToRoot();
+			if (GRemoveSubmixesFromRoot)
+			{
+				InSoundSubmix.IncrementLoadRefCount();
+			}
+			else
+			{
+				InSoundSubmix.AddToRoot();
+			}
 
 			MixerSubmix = MakeShared<FMixerSubmix, ESPMode::ThreadSafe>(this);
 			Submixes.Add(InSoundSubmix.GetUniqueID(), MixerSubmix);
@@ -2377,7 +2396,13 @@ namespace Audio
 			ExternalEndpointSubmixes.Remove(MixerSubmix);
 		}
 
-		Submixes.Remove(InSoundSubmix.GetUniqueID());
+		const int32 RemoveCount = Submixes.Remove(InSoundSubmix.GetUniqueID());
+		
+		if (GRemoveSubmixesFromRoot && RemoveCount > 0)
+		{
+			USoundSubmixBase& MutableSubmix = const_cast<USoundSubmixBase&>(InSoundSubmix);
+			MutableSubmix.DecrementLoadRefCount();
+		}
 	}
 
 	FMixerSubmixPtr FMixerDevice::FindSubmixInstanceByObjectId(uint32 InObjectId)
diff --git a/Engine/Source/Runtime/AudioMixer/Public/AudioMixerDevice.h b/Engine/Source/Runtime/AudioMixer/Public/AudioMixerDevice.h
--- a/Engine/Source/Runtime/AudioMixer/Public/AudioMixerDevice.h
+++ b/Engine/Source/Runtime/AudioMixer/Public/AudioMixerDevice.h
@@ -35,6 +35,8 @@ enum class ERequiredSubmixes : uint8
 
 namespace Audio
 {
+	AUDIOMIXER_API extern int32 GRemoveSubmixesFromRoot;
+	
 	// Audio Namespace Forward Declarations
 	class FMixerSourceManager;
 	class FMixerSourceVoice;
diff --git a/Engine/Source/Runtime/Engine/Classes/Sound/SoundSubmix.h b/Engine/Source/Runtime/Engine/Classes/Sound/SoundSubmix.h
--- a/Engine/Source/Runtime/Engine/Classes/Sound/SoundSubmix.h
+++ b/Engine/Source/Runtime/Engine/Classes/Sound/SoundSubmix.h
@@ -234,11 +234,17 @@ public:
 		return false; 
 	}
 	
+	ENGINE_API void IncrementLoadRefCount();
+	ENGINE_API void DecrementLoadRefCount();
+	
 protected:
 	//~ Begin UObject Interface.
 	ENGINE_API virtual FString GetDesc() override;
 	ENGINE_API virtual void BeginDestroy() override;
 	ENGINE_API virtual void PostLoad() override;
+	
+private:
+	std::atomic<int32> RefCount{ 0 };
 
 public:
 
diff --git a/Engine/Source/Runtime/Engine/Private/SoundSubmix.cpp b/Engine/Source/Runtime/Engine/Private/SoundSubmix.cpp
--- a/Engine/Source/Runtime/Engine/Private/SoundSubmix.cpp
+++ b/Engine/Source/Runtime/Engine/Private/SoundSubmix.cpp
@@ -422,6 +422,24 @@ FAudioDevice* USoundSubmix::GetAudioDeviceFrom(const UObject* WorldContextObject
 	return nullptr;
 }
 
+void USoundSubmixBase::IncrementLoadRefCount()
+{
+	const int32 OldRefCount = RefCount.fetch_add(1, std::memory_order_relaxed);
+	if (OldRefCount == 0)
+	{
+		AddToRoot();
+	}
+}
+
+void USoundSubmixBase::DecrementLoadRefCount()
+{
+	const int32 OldRefCount = RefCount.fetch_sub(1, std::memory_order_relaxed);
+	if (OldRefCount == 1 && IsRooted())
+	{
+		RemoveFromRoot();
+	}
+}
+
 FString USoundSubmixBase::GetDesc()
 {
 	return FString("Sound Submix");
