[WASAPI] Reswap audio device hardware to try to restart the audio command flow if the command queue overgrows.

#Synced-CL 50066651

diff --git a/Engine/Source/Runtime/AudioMixer/Private/AudioMixerDevice.cpp b/Engine/Source/Runtime/AudioMixer/Private/AudioMixerDevice.cpp
--- a/Engine/Source/Runtime/AudioMixer/Private/AudioMixerDevice.cpp
+++ b/Engine/Source/Runtime/AudioMixer/Private/AudioMixerDevice.cpp
@@ -1033,6 +1033,12 @@ namespace Audio
 
 		SourceManager->Update();
 
+		if (SourceManager->CheckAndResetHardwareCallbackStallFlag())
+		{
+			// This is a foreground thread check if we are not receiving hardware callbacks. 
+			AudioMixerPlatform->RequestDeviceSwap(TEXT(""), /* force */true, TEXT("FMixerDevice::UpdateHardware(). Force device swap because hardware has stalled."));
+		}
+
 		AudioMixerPlatform->OnHardwareUpdate();
 
 		// Device swaps span multiple calls beginning with CheckAudioDeviceChange() and ending with
diff --git a/Engine/Source/Runtime/AudioMixer/Private/AudioMixerSourceManager.cpp b/Engine/Source/Runtime/AudioMixer/Private/AudioMixerSourceManager.cpp
--- a/Engine/Source/Runtime/AudioMixer/Private/AudioMixerSourceManager.cpp
+++ b/Engine/Source/Runtime/AudioMixer/Private/AudioMixerSourceManager.cpp
@@ -4115,12 +4115,15 @@ namespace Audio
 			}
 			
 			// check that we haven't gone over the max size
+			// if so, set a flag that the AudioMixerDevice can check and try swapping to safe hardware
 			const SIZE_T MaxBufferSizeInBytes = ((SIZE_T)CommandBufferMaxSizeInMbCvar) << 20;
 			if (CurrentBufferSizeInBytes >= MaxBufferSizeInBytes)
 			{
 				int32 NumTimesOvergrown = CommandBuffers[AudioThreadCommandIndex].NumTimesOvergrown.Increment();
 				UE_LOG(LogAudioMixer, Error, TEXT("%d: Command buffer %d allocated size has grown to %" SIZE_T_FMT "mb! Likely cause the AudioRenderer has hung"),
 					NumTimesOvergrown, AudioThreadCommandIndex, CurrentBufferSizeInBytes >> 20);
+
+				bShouldResetHardwareCallback = true;
 			}
 		}
 
@@ -4135,7 +4138,7 @@ namespace Audio
 	{
 		TRACE_CPUPROFILER_EVENT_SCOPE(AudioMixerThreadCommands::PumpCommandQueue)
 		AudioRenderThreadId = FPlatformTLS::GetCurrentThreadId();
-		
+
 		// If we're already triggered, we need to wait for the audio thread to reset it before pumping
 		if (FPlatformProcess::SupportsMultithreading())
 		{
diff --git a/Engine/Source/Runtime/AudioMixer/Private/AudioMixerSourceManager.h b/Engine/Source/Runtime/AudioMixer/Private/AudioMixerSourceManager.h
--- a/Engine/Source/Runtime/AudioMixer/Private/AudioMixerSourceManager.h
+++ b/Engine/Source/Runtime/AudioMixer/Private/AudioMixerSourceManager.h
@@ -333,6 +333,13 @@ namespace Audio
 			return MixerDevice->GetRenderScheduler();
 		}
 
+		const bool CheckAndResetHardwareCallbackStallFlag()
+		{
+			bool bShouldResetTemp = bShouldResetHardwareCallback;
+			bShouldResetHardwareCallback = false;
+			return bShouldResetTemp; 
+		}
+
 	private:
 #define INVALID_AUDIO_RENDER_THREAD_ID static_cast<uint32>(-1)
 		uint32 AudioRenderThreadId = INVALID_AUDIO_RENDER_THREAD_ID;
@@ -410,6 +417,7 @@ namespace Audio
 		
 		FCommands CommandBuffers[2];
 		FThreadSafeCounter RenderThreadCommandBufferIndex;
+		bool bShouldResetHardwareCallback = false;
 
 		FEvent* CommandsProcessedEvent;
 		FCriticalSection CommandBufferIndexCriticalSection;
