[Audio Insights] Better distinction between which trace messages are sent to which audio trace channels. Audio is being reserved for more light weight audio logic (start/stop etc), whereas AudioMixer is more reserved for parameter data (which is more expensive to generate).

#Synced-CL 50156409

diff --git a/Engine/Source/Runtime/AudioMixer/Private/AudioMixerDevice.cpp b/Engine/Source/Runtime/AudioMixer/Private/AudioMixerDevice.cpp
--- a/Engine/Source/Runtime/AudioMixer/Private/AudioMixerDevice.cpp
+++ b/Engine/Source/Runtime/AudioMixer/Private/AudioMixerDevice.cpp
@@ -461,7 +461,7 @@ namespace Audio
 
 	void FMixerDevice::SendSubmixesSendLevelTraceMessages()
 	{
-		const bool bChannelEnabled = UE_TRACE_CHANNELEXPR_IS_ENABLED(AudioChannel);
+		const bool bChannelEnabled = UE_TRACE_CHANNELEXPR_IS_ENABLED(AudioMixerChannel);
 		if (bChannelEnabled)
 		{
 			auto TraceSubmixToSubmixSendMessage = [this](const FMixerSubmixPtr& InSubmix)
diff --git a/Engine/Source/Runtime/AudioMixer/Private/AudioMixerSource.cpp b/Engine/Source/Runtime/AudioMixer/Private/AudioMixerSource.cpp
--- a/Engine/Source/Runtime/AudioMixer/Private/AudioMixerSource.cpp
+++ b/Engine/Source/Runtime/AudioMixer/Private/AudioMixerSource.cpp
@@ -1883,8 +1883,8 @@ namespace Audio
 #if UE_AUDIO_PROFILERTRACE_ENABLED
 	void FMixerSource::SendSoundWaveStartTraceMessage() const
 	{
-		const bool bChannelEnabled = UE_TRACE_CHANNELEXPR_IS_ENABLED(AudioMixerChannel);
-		if (bChannelEnabled && WaveInstance && WaveInstance->ActiveSound && WaveInstance->SoundClass)
+		const bool bChannelEnabled = UE_TRACE_CHANNELEXPR_IS_ENABLED(AudioChannel);
+		if (bChannelEnabled && WaveInstance && WaveInstance->WaveData && WaveInstance->ActiveSound && WaveInstance->SoundClass)
 		{
 			const TObjectPtr<UObject> SoundWaveObjectPtr = FSoftObjectPath(WaveInstance->WaveData->GetPathName()).ResolveObject();
 			if (SoundWaveObjectPtr)
@@ -1892,7 +1892,7 @@ namespace Audio
 				const FString SoundWaveBaseAssetName = Audio::Trace::Util::GetSoundBaseAssetName(SoundWaveObjectPtr->GetClass());
 				const FString OwnerActorLabel = Audio::Trace::Util::GetOwnerActorLabel(*WaveInstance->ActiveSound);
 
-				UE_TRACE_LOG(Audio, SoundWaveStart, AudioMixerChannel)
+				UE_TRACE_LOG(Audio, SoundWaveStart, AudioChannel)
 					<< SoundWaveStart.DeviceId(AudioDevice->DeviceID)
 					<< SoundWaveStart.Timestamp(FPlatformTime::Cycles64())
 					<< SoundWaveStart.PlayOrder(WaveInstance->GetPlayOrder())
@@ -1907,10 +1907,10 @@ namespace Audio
 
 	void FMixerSource::SendSoundWaveStopTraceMessage() const
 	{
-		const bool bChannelEnabled = UE_TRACE_CHANNELEXPR_IS_ENABLED(AudioMixerChannel);
+		const bool bChannelEnabled = UE_TRACE_CHANNELEXPR_IS_ENABLED(AudioChannel);
 		if (bChannelEnabled && MixerDevice && WaveInstance && WaveInstance->ActiveSound)
 		{
-			UE_TRACE_LOG(Audio, SoundWaveStop, AudioMixerChannel)
+			UE_TRACE_LOG(Audio, SoundWaveStop, AudioChannel)
 				<< SoundWaveStop.DeviceId(MixerDevice->DeviceID)
 				<< SoundWaveStop.Timestamp(FPlatformTime::Cycles64())
 				<< SoundWaveStop.PlayOrder(WaveInstance->GetPlayOrder())
@@ -1920,9 +1920,9 @@ namespace Audio
 
 	void FMixerSource::UpdateIsAlivePingTraceMessage() const
 	{
-		const bool bChannelEnabled = UE_TRACE_CHANNELEXPR_IS_ENABLED(AudioMixerChannel);
+		const bool bChannelEnabled = UE_TRACE_CHANNELEXPR_IS_ENABLED(AudioChannel);
 
-		if (bChannelEnabled && WaveInstance && WaveInstance->IsPlaying() && WaveInstance->ActiveSound && WaveInstance->SoundClass)
+		if (bChannelEnabled && WaveInstance && WaveInstance->IsPlaying() && WaveInstance->WaveData && WaveInstance->ActiveSound && WaveInstance->SoundClass)
 		{
 			// Occasionally send a ping to Audio Insights to confirm the sound is still playing.
 			// This is to ensure long duration sounds (such as ambiences) will still appear in snapshot files if their start action was before the snapshot start
@@ -1934,7 +1934,7 @@ namespace Audio
 					const FString SoundWaveBaseAssetName = Audio::Trace::Util::GetSoundBaseAssetName(SoundWaveObjectPtr->GetClass());
 					const FString OwnerActorLabel = Audio::Trace::Util::GetOwnerActorLabel(*WaveInstance->ActiveSound);
 
-					UE_TRACE_LOG(Audio, SoundWaveIsAlivePing, AudioMixerChannel)
+					UE_TRACE_LOG(Audio, SoundWaveIsAlivePing, AudioChannel)
 						<< SoundWaveIsAlivePing.DeviceId(AudioDevice->DeviceID)
 						<< SoundWaveIsAlivePing.Timestamp(FPlatformTime::Cycles64())
 						<< SoundWaveIsAlivePing.PlayOrder(WaveInstance->GetPlayOrder())
diff --git a/Engine/Source/Runtime/AudioMixer/Private/AudioMixerSourceManager.cpp b/Engine/Source/Runtime/AudioMixer/Private/AudioMixerSourceManager.cpp
--- a/Engine/Source/Runtime/AudioMixer/Private/AudioMixerSourceManager.cpp
+++ b/Engine/Source/Runtime/AudioMixer/Private/AudioMixerSourceManager.cpp
@@ -400,7 +400,7 @@ namespace Audio
 #if UE_AUDIO_PROFILERTRACE_ENABLED
 	void TraceAudioBusStartMessage(const uint32 InDeviceId, const uint32 InObjectId, const FString& InAudioBusName, const int32 InNumChannels)
 	{
-		const bool bChannelEnabled = UE_TRACE_CHANNELEXPR_IS_ENABLED(AudioMixerChannel);
+		const bool bChannelEnabled = UE_TRACE_CHANNELEXPR_IS_ENABLED(AudioChannel);
 		if (bChannelEnabled)
 		{
 			UE_TRACE_LOG(Audio, AudioBusStart, AudioChannel)
@@ -414,7 +414,7 @@ namespace Audio
 
 	void TraceAudioBusIsAlivePingMessage(const uint32 InDeviceId, const uint32 InObjectId, const FString& InAudioBusName, const int32 InNumChannels)
 	{
-		const bool bChannelEnabled = UE_TRACE_CHANNELEXPR_IS_ENABLED(AudioMixerChannel);
+		const bool bChannelEnabled = UE_TRACE_CHANNELEXPR_IS_ENABLED(AudioChannel);
 		if (bChannelEnabled)
 		{
 			UE_TRACE_LOG(Audio, AudioBusIsAlivePing, AudioChannel)
@@ -428,7 +428,7 @@ namespace Audio
 
 	void TraceAudioBusStopMessage(const uint32 InDeviceId, const uint32 InObjectId)
 	{
-		const bool bChannelEnabled = UE_TRACE_CHANNELEXPR_IS_ENABLED(AudioMixerChannel);
+		const bool bChannelEnabled = UE_TRACE_CHANNELEXPR_IS_ENABLED(AudioChannel);
 		if (bChannelEnabled)
 		{
 			UE_TRACE_LOG(Audio, AudioBusStop, AudioChannel)
@@ -561,16 +561,17 @@ namespace Audio
 		constexpr float AttackTime  = 10.0f;
 		constexpr float ReleaseTime = 500.0f;
 
+		const bool bChannelEnabled = UE_TRACE_CHANNELEXPR_IS_ENABLED(AudioMixerChannel);
+
 		for (auto& [AudioBusKey, AudioBus] : AudioBuses)
 		{
 			if (AudioBusKey.ObjectId == InAudioBusId && AudioBus.IsValid())
 			{
 				AudioBus->StartEnvelopeFollower(AttackTime, ReleaseTime, MixerDevice->GetSampleRate());
 
-				const bool bChannelEnabled = UE_TRACE_CHANNELEXPR_IS_ENABLED(AudioChannel);
 				if (bChannelEnabled && AudioBus->bIsEnvelopeFollowing)
 				{
-					UE_TRACE_LOG(Audio, AudioBusEnvelopeFollowerEnabled, AudioChannel)
+					UE_TRACE_LOG(Audio, AudioBusEnvelopeFollowerEnabled, AudioMixerChannel)
 						<< AudioBusEnvelopeFollowerEnabled.DeviceId(MixerDevice->DeviceID)
 						<< AudioBusEnvelopeFollowerEnabled.AudioBusId(AudioBusKey.ObjectId)
 						<< AudioBusEnvelopeFollowerEnabled.Timestamp(FPlatformTime::Cycles64())
@@ -586,16 +587,17 @@ namespace Audio
 	{
 		ensure(IsInAudioThread());
 
+		const bool bChannelEnabled = UE_TRACE_CHANNELEXPR_IS_ENABLED(AudioMixerChannel);
+
 		for (auto& [AudioBusKey, AudioBus] : AudioBuses)
 		{
 			if (AudioBusKey.ObjectId == InAudioBusId && AudioBus.IsValid())
 			{
 				AudioBus->StopEnvelopeFollower();
 
-				const bool bChannelEnabled = UE_TRACE_CHANNELEXPR_IS_ENABLED(AudioChannel);
 				if (bChannelEnabled && !AudioBus->bIsEnvelopeFollowing)
 				{
-					UE_TRACE_LOG(Audio, AudioBusEnvelopeFollowerEnabled, AudioChannel)
+					UE_TRACE_LOG(Audio, AudioBusEnvelopeFollowerEnabled, AudioMixerChannel)
 						<< AudioBusEnvelopeFollowerEnabled.DeviceId(MixerDevice->DeviceID)
 						<< AudioBusEnvelopeFollowerEnabled.AudioBusId(AudioBusKey.ObjectId)
 						<< AudioBusEnvelopeFollowerEnabled.Timestamp(FPlatformTime::Cycles64())
@@ -2920,7 +2922,8 @@ namespace Audio
 		RenderThreadPhase = ESourceManagerRenderThreadPhase::UpdateBusses;
 
 #if UE_AUDIO_PROFILERTRACE_ENABLED
-		const bool bChannelEnabled = UE_TRACE_CHANNELEXPR_IS_ENABLED(AudioMixerChannel);
+		const bool bAudioChannelEnabled = UE_TRACE_CHANNELEXPR_IS_ENABLED(AudioChannel);
+		const bool bAudioMixerChannelEnabled = UE_TRACE_CHANNELEXPR_IS_ENABLED(AudioMixerChannel);
 #endif // UE_AUDIO_PROFILERTRACE_ENABLED
 
 		// Update the bus states post mixing. This flips the current/previous buffer indices.
@@ -2937,10 +2940,10 @@ namespace Audio
 			}
 			
 #if UE_AUDIO_PROFILERTRACE_ENABLED
-			if (bChannelEnabled)
+			if (bAudioChannelEnabled)
 			{
 				const float* AudioBuffer = AudioBus->GetCurrentBusBuffer();
-				const int32 NumSamples   = AudioBus->NumChannels * AudioBus->NumFrames;
+				const int32 NumSamples = AudioBus->NumChannels * AudioBus->NumFrames;
 
 				const bool bIsAudioBufferSilent = IsAudioBufferSilent(AudioBuffer, NumSamples);
 
@@ -2950,25 +2953,28 @@ namespace Audio
 					<< AudioBusHasActivity.Timestamp(FPlatformTime::Cycles64())
 					<< AudioBusHasActivity.HasActivity(!bIsAudioBufferSilent);
 
+				if (BusIsAlivePingTimeSec >= BusIsAlivePingWaitTimeSec)
+				{
+					TraceAudioBusIsAlivePingMessage(MixerDevice->DeviceID, AudioBusKey.ObjectId, AudioBus->BusName, AudioBus->NumChannels);
+				}
+			}
+
+			if (bAudioMixerChannelEnabled)
+			{
 				if (AudioBus->bIsEnvelopeFollowing)
 				{
-					UE_TRACE_LOG(Audio, AudioBusEnvelopeValues, AudioChannel)
+					UE_TRACE_LOG(Audio, AudioBusEnvelopeValues, AudioMixerChannel)
 						<< AudioBusEnvelopeValues.DeviceId(MixerDevice->DeviceID)
 						<< AudioBusEnvelopeValues.AudioBusId(AudioBusKey.ObjectId)
 						<< AudioBusEnvelopeValues.Timestamp(FPlatformTime::Cycles64())
 						<< AudioBusEnvelopeValues.EnvelopeValues(AudioBus->EnvelopeValues, AudioBus->GetNumChannels());
 				}
-
-				if (BusIsAlivePingTimeSec >= BusIsAlivePingWaitTimeSec)
-				{
-					TraceAudioBusIsAlivePingMessage(MixerDevice->DeviceID, AudioBusKey.ObjectId, AudioBus->BusName, AudioBus->NumChannels);
-				}
 			}
 #endif // UE_AUDIO_PROFILERTRACE_ENABLED
 		}
 
 #if UE_AUDIO_PROFILERTRACE_ENABLED
-		if (bChannelEnabled)
+		if (bAudioChannelEnabled)
 		{
 			if (BusIsAlivePingTimeSec >= BusIsAlivePingWaitTimeSec)
 			{
diff --git a/Engine/Source/Runtime/AudioMixer/Private/AudioMixerSubmix.cpp b/Engine/Source/Runtime/AudioMixer/Private/AudioMixerSubmix.cpp
--- a/Engine/Source/Runtime/AudioMixer/Private/AudioMixerSubmix.cpp
+++ b/Engine/Source/Runtime/AudioMixer/Private/AudioMixerSubmix.cpp
@@ -1350,13 +1350,13 @@ namespace Audio
 				DryLevelMod.ClearModulators();
 				
 #if UE_AUDIO_PROFILERTRACE_ENABLED
-				const bool bChannelEnabled = UE_TRACE_CHANNELEXPR_IS_ENABLED(AudioChannel);
+				const bool bChannelEnabled = UE_TRACE_CHANNELEXPR_IS_ENABLED(AudioMixerChannel);
 				if (bChannelEnabled && OwningSubmixObject.IsValid())
 				{
 					TArray<float> SilentEnvelopeValues;
 					SilentEnvelopeValues.Init(0.0f, NumChannels);
 
-					UE_TRACE_LOG(Audio, SubmixEnvelopeValues, AudioChannel)
+					UE_TRACE_LOG(Audio, SubmixEnvelopeValues, AudioMixerChannel)
 						<< SubmixEnvelopeValues.DeviceId(MixerDevice->DeviceID)
 						<< SubmixEnvelopeValues.SubmixId(GetTypeHash(OwningSubmixObject->GetPathName()))
 						<< SubmixEnvelopeValues.Timestamp(FPlatformTime::Cycles64())
@@ -1637,10 +1637,10 @@ namespace Audio
 				Audio::ArrayClampInPlace(MakeArrayView(EnvelopeValues, NumChannels), 0.f, 1.f);
 
 #if UE_AUDIO_PROFILERTRACE_ENABLED
-				const bool bChannelEnabled = UE_TRACE_CHANNELEXPR_IS_ENABLED(AudioChannel);
+				const bool bChannelEnabled = UE_TRACE_CHANNELEXPR_IS_ENABLED(AudioMixerChannel);
 				if (bChannelEnabled)
 				{
-					UE_TRACE_LOG(Audio, SubmixEnvelopeValues, AudioChannel)
+					UE_TRACE_LOG(Audio, SubmixEnvelopeValues, AudioMixerChannel)
 						<< SubmixEnvelopeValues.DeviceId(MixerDevice->DeviceID)
 						<< SubmixEnvelopeValues.SubmixId(GetTypeHash(OwningSubmixObject->GetPathName()))
 						<< SubmixEnvelopeValues.Timestamp(FPlatformTime::Cycles64())
@@ -2350,10 +2350,10 @@ namespace Audio
 			bIsEnvelopeFollowing = true;
 
 #if UE_AUDIO_PROFILERTRACE_ENABLED
-			const bool bChannelEnabled = UE_TRACE_CHANNELEXPR_IS_ENABLED(AudioChannel);
+			const bool bChannelEnabled = UE_TRACE_CHANNELEXPR_IS_ENABLED(AudioMixerChannel);
 			if (bChannelEnabled && OwningSubmixObject.IsValid())
 			{
-				UE_TRACE_LOG(Audio, SubmixEnvelopeFollowerEnabled, AudioChannel)
+				UE_TRACE_LOG(Audio, SubmixEnvelopeFollowerEnabled, AudioMixerChannel)
 					<< SubmixEnvelopeFollowerEnabled.DeviceId(MixerDevice->DeviceID)
 					<< SubmixEnvelopeFollowerEnabled.SubmixId(GetTypeHash(OwningSubmixObject->GetPathName()))
 					<< SubmixEnvelopeFollowerEnabled.Timestamp(FPlatformTime::Cycles64())
@@ -2368,10 +2368,10 @@ namespace Audio
 		bIsEnvelopeFollowing = false;
 
 #if UE_AUDIO_PROFILERTRACE_ENABLED
-		const bool bChannelEnabled = UE_TRACE_CHANNELEXPR_IS_ENABLED(AudioChannel);
+		const bool bChannelEnabled = UE_TRACE_CHANNELEXPR_IS_ENABLED(AudioMixerChannel);
 		if (bChannelEnabled && OwningSubmixObject.IsValid())
 		{
-			UE_TRACE_LOG(Audio, SubmixEnvelopeFollowerEnabled, AudioChannel)
+			UE_TRACE_LOG(Audio, SubmixEnvelopeFollowerEnabled, AudioMixerChannel)
 				<< SubmixEnvelopeFollowerEnabled.DeviceId(MixerDevice->DeviceID)
 				<< SubmixEnvelopeFollowerEnabled.SubmixId(GetTypeHash(OwningSubmixObject->GetPathName()))
 				<< SubmixEnvelopeFollowerEnabled.Timestamp(FPlatformTime::Cycles64())
diff --git a/Engine/Source/Runtime/AudioMixerCore/Private/AudioMixer.cpp b/Engine/Source/Runtime/AudioMixerCore/Private/AudioMixer.cpp
--- a/Engine/Source/Runtime/AudioMixerCore/Private/AudioMixer.cpp
+++ b/Engine/Source/Runtime/AudioMixerCore/Private/AudioMixer.cpp
@@ -25,8 +25,8 @@
 CSV_DEFINE_CATEGORY_MODULE(AUDIOMIXERCORE_API, Audio, false);
 
 #if UE_AUDIO_PROFILERTRACE_ENABLED
-UE_TRACE_CHANNEL_DEFINE(AudioChannel);
-UE_TRACE_CHANNEL_DEFINE(AudioMixerChannel);
+UE_TRACE_CHANNEL_DEFINE(AudioChannel, "Audio logic data (such as asset load/unload/start/stop, etc.)");
+UE_TRACE_CHANNEL_DEFINE(AudioMixerChannel, "Audio parameter data (such as amplitude, pitch, etc.)");
 #endif // UE_AUDIO_PROFILERTRACE_ENABLED
 
 // Audio Insights Event Log in-built event types
diff --git a/Engine/Source/Runtime/Engine/Private/Audio/AudioTraceUtil.cpp b/Engine/Source/Runtime/Engine/Private/Audio/AudioTraceUtil.cpp
--- a/Engine/Source/Runtime/Engine/Private/Audio/AudioTraceUtil.cpp
+++ b/Engine/Source/Runtime/Engine/Private/Audio/AudioTraceUtil.cpp
@@ -102,7 +102,7 @@ namespace Audio::Trace
 	{
 		void SendActiveSoundEvent(const FActiveSound& InActiveSound, const FString& Event)
 		{
-			const bool bChannelEnabled = UE_TRACE_CHANNELEXPR_IS_ENABLED(AudioMixerChannel);
+			const bool bChannelEnabled = UE_TRACE_CHANNELEXPR_IS_ENABLED(AudioChannel);
 
 			if (bChannelEnabled && InActiveSound.AudioDevice)
 			{
@@ -124,11 +124,11 @@ namespace Audio::Trace
 
 		void SendEvent(const Audio::FDeviceId AudioDeviceID, const FString& Event, const uint32 AudioObjectID, const FString& AssetPath, const FString& ActorLabel, const FString& ActorIconName, const FString& SoundClassName)
 		{
-			const bool bChannelEnabled = UE_TRACE_CHANNELEXPR_IS_ENABLED(AudioMixerChannel);
+			const bool bChannelEnabled = UE_TRACE_CHANNELEXPR_IS_ENABLED(AudioChannel);
 
 			if (bChannelEnabled)
 			{
-				UE_TRACE_LOG(Audio, EventLog, AudioMixerChannel)
+				UE_TRACE_LOG(Audio, EventLog, AudioChannel)
 					<< EventLog.DeviceId(AudioDeviceID)
 					<< EventLog.Timestamp(FPlatformTime::Cycles64())
 					<< EventLog.PlayOrder(AudioObjectID)
diff --git a/Engine/Source/Runtime/Engine/Private/AudioDevice.cpp b/Engine/Source/Runtime/Engine/Private/AudioDevice.cpp
--- a/Engine/Source/Runtime/Engine/Private/AudioDevice.cpp
+++ b/Engine/Source/Runtime/Engine/Private/AudioDevice.cpp
@@ -5036,7 +5036,7 @@ void FAudioDevice::UpdateAudioEngineSubsystems()
 #if UE_AUDIO_PROFILERTRACE_ENABLED
 void FAudioDevice::SendSoundStartedTraceMessage(const FActiveSound& InActiveSound) const
 {
-	const bool bChannelEnabled = UE_TRACE_CHANNELEXPR_IS_ENABLED(AudioMixerChannel);
+	const bool bChannelEnabled = UE_TRACE_CHANNELEXPR_IS_ENABLED(AudioChannel);
 
 	if (bChannelEnabled)
 	{
@@ -5048,7 +5048,7 @@ void FAudioDevice::SendSoundStartedTraceMessage(const FActiveSound& InActiveSoun
 
 		const FString OwnerActorLabel = Audio::Trace::Util::GetOwnerActorLabel(InActiveSound);
 
-		UE_TRACE_LOG(Audio, SoundStart, AudioMixerChannel)
+		UE_TRACE_LOG(Audio, SoundStart, AudioChannel)
 			<< SoundStart.DeviceId(InActiveSound.AudioDevice->DeviceID)
 			<< SoundStart.Timestamp(FPlatformTime::Cycles64())
 			<< SoundStart.PlayOrder(InActiveSound.PlayOrder)
@@ -5062,7 +5062,7 @@ void FAudioDevice::SendSoundStartedTraceMessage(const FActiveSound& InActiveSoun
 
 void FAudioDevice::SendSoundIsAlivePingTraceMessage(const FActiveSound& InActiveSound) const
 {
-	const bool bChannelEnabled = UE_TRACE_CHANNELEXPR_IS_ENABLED(AudioMixerChannel);
+	const bool bChannelEnabled = UE_TRACE_CHANNELEXPR_IS_ENABLED(AudioChannel);
 
 	if (bChannelEnabled)
 	{
@@ -5075,7 +5075,7 @@ void FAudioDevice::SendSoundIsAlivePingTraceMessage(const FActiveSound& InActive
 		const FString SoundBaseAssetName = Audio::Trace::Util::GetSoundBaseAssetName(SoundObjectPtr->GetClass());
 		const FString OwnerActorLabel = Audio::Trace::Util::GetOwnerActorLabel(InActiveSound);
 
-		UE_TRACE_LOG(Audio, SoundIsAlivePing, AudioMixerChannel)
+		UE_TRACE_LOG(Audio, SoundIsAlivePing, AudioChannel)
 			<< SoundIsAlivePing.DeviceId(InActiveSound.AudioDevice->DeviceID)
 			<< SoundIsAlivePing.Timestamp(FPlatformTime::Cycles64())
 			<< SoundIsAlivePing.PlayOrder(InActiveSound.PlayOrder)
@@ -6036,13 +6036,13 @@ void FAudioDevice::RemoveActiveSound(FActiveSound* ActiveSound)
 	UAudioComponent::PlaybackCompleted(ComponentID, false);
 
 #if UE_AUDIO_PROFILERTRACE_ENABLED
-	const bool bChannelEnabled = UE_TRACE_CHANNELEXPR_IS_ENABLED(AudioMixerChannel);
+	const bool bChannelEnabled = UE_TRACE_CHANNELEXPR_IS_ENABLED(AudioChannel);
 	if (bChannelEnabled)
 	{
 		const TObjectPtr<UObject> SoundObjectPtr = Audio::Trace::Util::GetSoundObjectPointer(*ActiveSound);
 		if (SoundObjectPtr)
 		{
-			UE_TRACE_LOG(Audio, SoundStop, AudioMixerChannel)
+			UE_TRACE_LOG(Audio, SoundStop, AudioChannel)
 				<< SoundStop.DeviceId(ActiveSound->AudioDevice->DeviceID)
 				<< SoundStop.Timestamp(FPlatformTime::Cycles64())
 				<< SoundStop.PlayOrder(ActiveSound->PlayOrder)
