Add cvar for logging a sound's relative render cost exceeding a threshold

#Synced-CL 50263872

diff --git a/Engine/Source/Runtime/AudioMixer/Private/AudioMixerSource.cpp b/Engine/Source/Runtime/AudioMixer/Private/AudioMixerSource.cpp
--- a/Engine/Source/Runtime/AudioMixer/Private/AudioMixerSource.cpp
+++ b/Engine/Source/Runtime/AudioMixer/Private/AudioMixerSource.cpp
@@ -95,6 +95,16 @@ static FAutoConsoleCommand GSetAudioMixerSourceFadeMin(
 	)
 );
 
+#if ENABLE_AUDIO_DEBUG
+static float RelativeRenderCostLogThresholdCVar = -1.f;
+FAutoConsoleVariableRef CVarRelativeRenderCostLogThreshold(
+	TEXT("au.RelativeRenderCostLogThresholdCVar"),
+	RelativeRenderCostLogThresholdCVar,
+	TEXT("Log (verbose) if the Relative Render Cost of a playing wave exceeds this threshold. " \
+		"Default: -1.f (disabled), Enabled for values >= 0"),
+	ECVF_Default);
+#endif // if ENABLE_AUDIO_DEBUG
+
 namespace Audio
 {
 	namespace MixerSourcePrivate
@@ -1848,6 +1858,19 @@ namespace Audio
 			{
 				FScopeLock DebugInfoLock(&DebugInfo->CS);
 				DebugInfo->RelativeRenderCost = RelativeRenderCost;
+
+				// au.RelativeRenderCostLogThresholdCVar logging
+				static constexpr float RelativeRenderCostLogTolerance = 0.1f;
+				if (RelativeRenderCostLogThresholdCVar >= 0 &&
+					RelativeRenderCost > RelativeRenderCostLogThresholdCVar &&
+					RelativeRenderCost - DebugInfo->MaxLoggedRelativeRenderCost > RelativeRenderCostLogTolerance)
+				{
+					UE_LOG(LogAudioMixer, Verbose, TEXT("New peak Relative Render Cost: %f (exceeds log threshold %f) for wave %s."),
+						RelativeRenderCost,
+						RelativeRenderCostLogThresholdCVar,
+						*WaveInstance->GetName());
+					DebugInfo->MaxLoggedRelativeRenderCost = RelativeRenderCost;
+				}
 			}
 #endif // if ENABLE_AUDIO_DEBUG
 
diff --git a/Engine/Source/Runtime/Engine/Public/Audio.h b/Engine/Source/Runtime/Engine/Public/Audio.h
--- a/Engine/Source/Runtime/Engine/Public/Audio.h
+++ b/Engine/Source/Runtime/Engine/Public/Audio.h
@@ -883,6 +883,10 @@ public:
 		/** Relative cost to render wave. */
 		float RelativeRenderCost = 1.f;
 
+		/** Max relative render cost already logged for this sound 
+		(used to limit log spam from RelativeRenderCostLogThresholdCVar) */
+		float MaxLoggedRelativeRenderCost = 0.f;
+
 		/** Basic CS so we can pass this around safely. */
 		FCriticalSection CS;
 	};
