[Audio] Let users know that UAudioSettings::MonoChannelUpmixMethod cannot be changed without an editor restart
+ Mono channel maps are cached upon initialization of the audio device so in hidden static variables so impossible to update those in FMixerDevice::UpdateHardware

#Synced-CL 50218806

diff --git a/Engine/Source/Runtime/AudioMixer/Private/AudioMixerDevice.cpp b/Engine/Source/Runtime/AudioMixer/Private/AudioMixerDevice.cpp
--- a/Engine/Source/Runtime/AudioMixer/Private/AudioMixerDevice.cpp
+++ b/Engine/Source/Runtime/AudioMixer/Private/AudioMixerDevice.cpp
@@ -1024,12 +1024,12 @@ namespace Audio
 		LLM_SCOPE(ELLMTag::AudioMixer);
 
 		// If we're in editor, re-query these in case they changed. 
-		if (GIsEditor)
+#if WITH_EDITOR
 		{
 			const UAudioSettings* AudioSettings = GetDefault<UAudioSettings>();
-			MonoChannelUpmixMethod = AudioSettings->MonoChannelUpmixMethod;
 			PanningMethod = AudioSettings->PanningMethod;
 		}
+#endif
 
 		SourceManager->Update();
 
diff --git a/Engine/Source/Runtime/Engine/Classes/Sound/AudioSettings.h b/Engine/Source/Runtime/Engine/Classes/Sound/AudioSettings.h
--- a/Engine/Source/Runtime/Engine/Classes/Sound/AudioSettings.h
+++ b/Engine/Source/Runtime/Engine/Classes/Sound/AudioSettings.h
@@ -235,7 +235,7 @@ class UAudioSettings : public UDeveloperSettings
 	/**
 	* The upmixing method for mono sound sources. Defines how mono channels are up-mixed to stereo channels.
 	*/
-	UPROPERTY(config, EditAnywhere, Category = "Mix", AdvancedDisplay)
+	UPROPERTY(config, EditAnywhere, Category = "Mix", AdvancedDisplay, meta = (ConfigRestartRequired = true))
 	EMonoChannelUpmixMethod MonoChannelUpmixMethod;
 
 	/**
