[Audio Insights] Re-enable envelope followers on submixes automatically.

#Synced-CL 50345329

diff --git a/Engine/Source/Runtime/AudioMixer/Private/AudioMixerDevice.cpp b/Engine/Source/Runtime/AudioMixer/Private/AudioMixerDevice.cpp
--- a/Engine/Source/Runtime/AudioMixer/Private/AudioMixerDevice.cpp
+++ b/Engine/Source/Runtime/AudioMixer/Private/AudioMixerDevice.cpp
@@ -382,6 +382,7 @@ namespace Audio
 
 #if UE_AUDIO_PROFILERTRACE_ENABLED
 		FTraceAuxiliary::OnTraceStarted.AddRaw(this, &FMixerDevice::OnTraceStarted);
+		FTraceAuxiliary::OnTraceStopped.AddRaw(this, &FMixerDevice::OnTraceStopped);
 #endif // UE_AUDIO_PROFILERTRACE_ENABLED
 	}
 
@@ -391,6 +392,7 @@ namespace Audio
 
 #if UE_AUDIO_PROFILERTRACE_ENABLED
 		FTraceAuxiliary::OnTraceStarted.RemoveAll(this);
+		FTraceAuxiliary::OnTraceStopped.RemoveAll(this);
 #endif // UE_AUDIO_PROFILERTRACE_ENABLED
 
 		// Shutdown all pending clock events, as they may have references to 
@@ -413,6 +415,35 @@ namespace Audio
 		SendSubmixesAlivePing();
 	}
 
+	void FMixerDevice::OnTraceStopped(FTraceAuxiliary::EConnectionType TraceType, const FString& TraceDestination)
+	{
+		// Auto-Disable output metering
+		Submixes.Iterate([this](const FSubmixMap::FPair& SubmixPair)
+		{
+			const FMixerSubmixPtr& Submix = SubmixPair.Value;
+			if (Submix.IsValid() && Submix->bIsEnvelopeFollowing && SubmixesWithAutoEnvelopeFollowers.Contains(Submix->GetId()))
+			{
+				FAudioThread::RunCommandOnAudioThread([this, Submix]()
+				{
+					StopSubmixEnvelopeFollowing(Submix);
+				});
+			}
+		});
+
+		for (const FMixerSubmixPtr& RequiredSubmix : RequiredSubmixInstances)
+		{
+			if (RequiredSubmix.IsValid() && RequiredSubmix->bIsEnvelopeFollowing && SubmixesWithAutoEnvelopeFollowers.Contains(RequiredSubmix->GetId()))
+			{
+				FAudioThread::RunCommandOnAudioThread([this, RequiredSubmix]()
+				{
+					StopSubmixEnvelopeFollowing(RequiredSubmix);
+				});
+			}
+		}
+
+		SubmixesWithAutoEnvelopeFollowers.Reset();
+	}
+
 	void FMixerDevice::SendSubmixesAlivePing()
 	{
 		const bool bChannelEnabled = UE_TRACE_CHANNELEXPR_IS_ENABLED(AudioChannel);
@@ -446,6 +477,8 @@ namespace Audio
 				}
 			}
 		}
+
+		AutomaticallyStartSubmixEnvelopeFollowers();
 	}
 
 	void FMixerDevice::UpdateSubmixesAlivePingTime()
@@ -539,10 +572,40 @@ namespace Audio
 
 		if (InStartSubmixEnvelopeFollowerCVar)
 		{
-			FAudioThread::RunCommandOnAudioThread([this, SubmixName = InStartSubmixEnvelopeFollowerCVar->GetString()]()
+			const FString SubmixName = InStartSubmixEnvelopeFollowerCVar->GetString();
+			if (!SubmixesWithCVarEnvelopeFollowers.Contains(SubmixName))
 			{
-				StartSubmixEnvelopeFollowerFromCVar(SubmixName);
-			});
+				Submixes.Iterate([this, &SubmixName](const FSubmixMap::FPair& SubmixPair)
+				{
+					const FMixerSubmixPtr& Submix = SubmixPair.Value;
+
+					if (Submix.IsValid() && Submix->GetName() == SubmixName)
+					{
+						FAudioThread::RunCommandOnAudioThread([this, Submix]()
+						{
+							StartSubmixEnvelopeFollowing(Submix);
+						});
+
+						SubmixesWithCVarEnvelopeFollowers.Add(SubmixName);
+					}
+				});
+
+				for (const FMixerSubmixPtr& RequiredSubmix : RequiredSubmixInstances)
+				{
+					if (RequiredSubmix.IsValid() && RequiredSubmix->GetName() == SubmixName)
+					{
+						FAudioThread::RunCommandOnAudioThread([this, RequiredSubmix]()
+						{
+							StartSubmixEnvelopeFollowing(RequiredSubmix);
+						});
+
+						SubmixesWithCVarEnvelopeFollowers.Add(SubmixName);
+					}
+				}
+
+				SubmixesWithCVarEnvelopeFollowers.Add(SubmixName);
+			}
+			
 		}
 	}
 
@@ -552,59 +615,106 @@ namespace Audio
 
 		if (InStopSubmixEnvelopeFollowerCVar)
 		{
-			FAudioThread::RunCommandOnAudioThread([this, SubmixName = InStopSubmixEnvelopeFollowerCVar->GetString()]()
+			const FString SubmixName = InStopSubmixEnvelopeFollowerCVar->GetString();
+			if (SubmixesWithCVarEnvelopeFollowers.Contains(SubmixName))
 			{
-				StopSubmixEnvelopeFollowerFromCVar(SubmixName);
-			});
+				Submixes.Iterate([this, &SubmixName](const FSubmixMap::FPair& SubmixPair)
+				{
+					const FMixerSubmixPtr& Submix = SubmixPair.Value;
+
+					if (Submix.IsValid() && Submix->GetName() == SubmixName)
+					{
+						FAudioThread::RunCommandOnAudioThread([this, Submix]()
+						{
+							StopSubmixEnvelopeFollowing(Submix);
+						});
+
+						SubmixesWithCVarEnvelopeFollowers.Remove(SubmixName);
+					}
+				});
+
+				for (const FMixerSubmixPtr& RequiredSubmix : RequiredSubmixInstances)
+				{
+					if (RequiredSubmix.IsValid() && RequiredSubmix->GetName() == SubmixName)
+					{
+						FAudioThread::RunCommandOnAudioThread([this, RequiredSubmix]()
+						{
+							StopSubmixEnvelopeFollowing(RequiredSubmix);
+						});
+
+						SubmixesWithCVarEnvelopeFollowers.Remove(SubmixName);
+					}
+				}
+			}
 		}
 	}
 
-	void FMixerDevice::StartSubmixEnvelopeFollowerFromCVar(const FString& InSubmixName)
+	void FMixerDevice::StartSubmixEnvelopeFollowing(const FMixerSubmixPtr& InSubmix)
 	{
 		ensure(IsInAudioThread());
 
+		if (!InSubmix || !InSubmix->IsValid())
+		{
+			return;
+		}
+
 		// Attack/Release default times used in USoundSubmix constructor
-		constexpr int32 AttackTime  = 10;
+		constexpr int32 AttackTime = 10;
 		constexpr int32 ReleaseTime = 500;
 
-		Submixes.Iterate([this, &InSubmixName, AttackTime, ReleaseTime](const FSubmixMap::FPair& SubmixPair)
-		{
-			const FMixerSubmixPtr& Submix = SubmixPair.Value;
+		InSubmix->StartEnvelopeFollowing(AttackTime, ReleaseTime);
+	}
 
-			if (Submix.IsValid() && Submix->GetName() == InSubmixName)
-			{
-				Submix->StartEnvelopeFollowing(AttackTime, ReleaseTime);
-			}
-		});
+	void FMixerDevice::StopSubmixEnvelopeFollowing(const FMixerSubmixPtr& InSubmix)
+	{
+		ensure(IsInAudioThread());
 
-		for (const FMixerSubmixPtr& RequiredSubmix : RequiredSubmixInstances)
+		if (!InSubmix || !InSubmix->IsValid() || !InSubmix->bIsEnvelopeFollowing)
 		{
-			if (RequiredSubmix.IsValid() && RequiredSubmix->GetName() == InSubmixName)
-			{
-				RequiredSubmix->StartEnvelopeFollowing(AttackTime, ReleaseTime);
-			}
+			return;
 		}
+
+		InSubmix->StopEnvelopeFollowing();
 	}
 
-	void FMixerDevice::StopSubmixEnvelopeFollowerFromCVar(const FString& InSubmixName)
+	void FMixerDevice::AutomaticallyStartSubmixEnvelopeFollowers()
 	{
-		ensure(IsInAudioThread());
-
-		Submixes.Iterate([this, &InSubmixName](const FSubmixMap::FPair& SubmixPair)
+		const bool bAudioMixerChannelEnabled = UE_TRACE_CHANNELEXPR_IS_ENABLED(AudioMixerChannel);
+		if (bAudioMixerChannelEnabled)
 		{
-			const FMixerSubmixPtr& Submix = SubmixPair.Value;
-
-			if (Submix.IsValid() && Submix->GetName() == InSubmixName)
+			// Auto-Enable output metering
+			Submixes.Iterate([this](const FSubmixMap::FPair& SubmixPair)
 			{
-				Submix->StopEnvelopeFollowing();
-			}
-		});
+				const FMixerSubmixPtr& Submix = SubmixPair.Value;
+				if (Submix.IsValid())
+				{
+					const uint32 SubmixID = Submix->GetId();
+					if (!SubmixesWithAutoEnvelopeFollowers.Contains(SubmixID))
+					{
+						SubmixesWithAutoEnvelopeFollowers.Add(SubmixID);
+						
+						FAudioThread::RunCommandOnAudioThread([this, Submix]()
+						{
+							StartSubmixEnvelopeFollowing(Submix);
+						});
+					}
+				}
+			});
 
-		for (const FMixerSubmixPtr& RequiredSubmix : RequiredSubmixInstances)
-		{
-			if (RequiredSubmix.IsValid() && RequiredSubmix->GetName() == InSubmixName)
+			for (const FMixerSubmixPtr& RequiredSubmix : RequiredSubmixInstances)
 			{
-				RequiredSubmix->StopEnvelopeFollowing();
+				if (RequiredSubmix.IsValid())
+				{
+					const uint32 SubmixID = RequiredSubmix->GetId();
+					if (!SubmixesWithAutoEnvelopeFollowers.Contains(SubmixID))
+					{
+						SubmixesWithAutoEnvelopeFollowers.Add(SubmixID);
+						FAudioThread::RunCommandOnAudioThread([this, RequiredSubmix]()
+						{
+							StartSubmixEnvelopeFollowing(RequiredSubmix);
+						});
+					}
+				}
 			}
 		}
 	}
@@ -1601,15 +1711,6 @@ namespace Audio
 				}
 			}
 		}
-
-#if UE_AUDIO_PROFILERTRACE_ENABLED
-		// Auto-Enable MainSubmix output metering
-		FAudioThread::RunCommandOnAudioThread([this]()
-		{
-			StartSubmixEnvelopeFollowerFromCVar(GetRequiredSubmixName(ERequiredSubmixes::Main));
-		});
-#endif // UE_AUDIO_PROFILERTRACE_ENABLED
-
 	}
 
 	void FMixerDevice::RebuildSubmixLinks(const USoundSubmixBase& SoundSubmix, FMixerSubmixPtr& SubmixInstance)
@@ -2409,6 +2510,17 @@ namespace Audio
 			ExternalEndpointSubmixes.Remove(MixerSubmix);
 		}
 
+#if UE_AUDIO_PROFILERTRACE_ENABLED
+		if (MixerSubmix.IsValid())
+		{
+			FAudioThread::RunCommandOnGameThread([this, SubmixID = MixerSubmix->GetId(), SubmixName = MixerSubmix->GetName()]()
+			{
+				SubmixesWithAutoEnvelopeFollowers.Remove(SubmixID);
+				SubmixesWithCVarEnvelopeFollowers.Remove(SubmixName);
+			});
+		}
+#endif // UE_AUDIO_PROFILERTRACE_ENABLED
+
 		const int32 RemoveCount = Submixes.Remove(InSoundSubmix.GetUniqueID());
 		
 		if (GRemoveSubmixesFromRoot && RemoveCount > 0)
diff --git a/Engine/Source/Runtime/AudioMixer/Private/AudioMixerSubmix.cpp b/Engine/Source/Runtime/AudioMixer/Private/AudioMixerSubmix.cpp
--- a/Engine/Source/Runtime/AudioMixer/Private/AudioMixerSubmix.cpp
+++ b/Engine/Source/Runtime/AudioMixer/Private/AudioMixerSubmix.cpp
@@ -2332,6 +2332,8 @@ namespace Audio
 
 	void FMixerSubmix::StartEnvelopeFollowing(int32 AttackTime, int32 ReleaseTime)
 	{
+		EnvelopeFollowerRequestsRefCount.Increment();
+
 		if (!bIsEnvelopeFollowing)
 		{
 			FEnvelopeFollowerInitParams EnvelopeFollowerInitParams;
@@ -2365,19 +2367,29 @@ namespace Audio
 
 	void FMixerSubmix::StopEnvelopeFollowing()
 	{
-		bIsEnvelopeFollowing = false;
+		if (!bIsEnvelopeFollowing)
+		{
+			return;
+		}
 
-#if UE_AUDIO_PROFILERTRACE_ENABLED
-		const bool bChannelEnabled = UE_TRACE_CHANNELEXPR_IS_ENABLED(AudioMixerChannel);
-		if (bChannelEnabled && OwningSubmixObject.IsValid())
+		EnvelopeFollowerRequestsRefCount.Decrement();
+		if (EnvelopeFollowerRequestsRefCount.GetValue() <= 0)
 		{
-			UE_TRACE_LOG(Audio, SubmixEnvelopeFollowerEnabled, AudioMixerChannel)
+			EnvelopeFollowerRequestsRefCount.Reset();
+			bIsEnvelopeFollowing = false;
+
+#if UE_AUDIO_PROFILERTRACE_ENABLED
+			const bool bChannelEnabled = UE_TRACE_CHANNELEXPR_IS_ENABLED(AudioMixerChannel);
+			if (bChannelEnabled && OwningSubmixObject.IsValid())
+			{
+				UE_TRACE_LOG(Audio, SubmixEnvelopeFollowerEnabled, AudioMixerChannel)
 				<< SubmixEnvelopeFollowerEnabled.DeviceId(MixerDevice->DeviceID)
 				<< SubmixEnvelopeFollowerEnabled.SubmixId(GetTypeHash(OwningSubmixObject->GetPathName()))
 				<< SubmixEnvelopeFollowerEnabled.Timestamp(FPlatformTime::Cycles64())
 				<< SubmixEnvelopeFollowerEnabled.EnvelopeFollowerEnabled(false);
-		}
+			}
 #endif // UE_AUDIO_PROFILERTRACE_ENABLED
+		}
 	}
 
 	void FMixerSubmix::AddEnvelopeFollowerDelegate(const FOnSubmixEnvelopeBP& OnSubmixEnvelopeBP)
diff --git a/Engine/Source/Runtime/AudioMixer/Public/AudioMixerDevice.h b/Engine/Source/Runtime/AudioMixer/Public/AudioMixerDevice.h
--- a/Engine/Source/Runtime/AudioMixer/Public/AudioMixerDevice.h
+++ b/Engine/Source/Runtime/AudioMixer/Public/AudioMixerDevice.h
@@ -409,6 +409,7 @@ namespace Audio
 
 #if UE_AUDIO_PROFILERTRACE_ENABLED
 		void OnTraceStarted(FTraceAuxiliary::EConnectionType TraceType, const FString& TraceDestination);
+		void OnTraceStopped(FTraceAuxiliary::EConnectionType TraceType, const FString& TraceDestination);
 		void SendSubmixesAlivePing();
 		void UpdateSubmixesAlivePingTime();
 		void SendSubmixesSendLevelTraceMessages();
@@ -419,8 +420,10 @@ namespace Audio
 		void OnStartSubmixEnvelopeFollowerCVarChanged(IConsoleVariable* InStartSubmixEnvelopeFollowerCVar);
 		void OnStopSubmixEnvelopeFollowerCVarChanged(IConsoleVariable* InStopSubmixEnvelopeFollowerCVar);
 
-		void StartSubmixEnvelopeFollowerFromCVar(const FString& InSubmixName);
-		void StopSubmixEnvelopeFollowerFromCVar(const FString& InSubmixName);
+		void StartSubmixEnvelopeFollowing(const FMixerSubmixPtr& InSubmix);
+		void StopSubmixEnvelopeFollowing(const FMixerSubmixPtr& InSubmix);
+
+		void AutomaticallyStartSubmixEnvelopeFollowers();
 #endif // UE_AUDIO_PROFILERTRACE_ENABLED
 
 		TArray<USoundSubmix*> RequiredSubmixes;
@@ -508,6 +511,9 @@ namespace Audio
 		IAudioLinkFactory* AudioLinkFactory = nullptr;
 
 #if UE_AUDIO_PROFILERTRACE_ENABLED
+		TSet<uint32> SubmixesWithAutoEnvelopeFollowers;
+		TSet<FString> SubmixesWithCVarEnvelopeFollowers;
+
 		IConsoleVariable* StartSubmixEnvelopeFollowerCVar = nullptr;
 		IConsoleVariable* StopSubmixEnvelopeFollowerCVar  = nullptr;
 
diff --git a/Engine/Source/Runtime/AudioMixer/Public/AudioMixerSubmix.h b/Engine/Source/Runtime/AudioMixer/Public/AudioMixerSubmix.h
--- a/Engine/Source/Runtime/AudioMixer/Public/AudioMixerSubmix.h
+++ b/Engine/Source/Runtime/AudioMixer/Public/AudioMixerSubmix.h
@@ -635,6 +635,9 @@ namespace Audio
 		// Bool set to true when envelope following is enabled
 		FThreadSafeBool bIsEnvelopeFollowing;
 
+		// Ref count for the number of times a request for submixes to start envelope following
+		FThreadSafeCounter EnvelopeFollowerRequestsRefCount = 0;
+
 		// Multi-cast delegate to broadcast envelope data from this submix instance
 		FOnSubmixEnvelope OnSubmixEnvelope;
 
