[AudioMixer] FMixerDevice: added logic to get TruePeak from main submix and send it as trace message to Audio Insights

#Synced-CL 50901149

diff --git a/Engine/Source/Runtime/AudioMixer/Private/AudioMixerDevice.cpp b/Engine/Source/Runtime/AudioMixer/Private/AudioMixerDevice.cpp
--- a/Engine/Source/Runtime/AudioMixer/Private/AudioMixerDevice.cpp
+++ b/Engine/Source/Runtime/AudioMixer/Private/AudioMixerDevice.cpp
@@ -135,6 +135,22 @@ FAutoConsoleVariableRef CVarSubmixBufferListenerForceRegistration(
 
 #if UE_AUDIO_PROFILERTRACE_ENABLED
 
+#if !WITH_EDITOR
+UE_TRACE_EVENT_BEGIN(Audio, MainSubmixLoaded)
+	UE_TRACE_EVENT_FIELD(uint32, DeviceId)
+	UE_TRACE_EVENT_FIELD(uint32, SubmixId)
+	UE_TRACE_EVENT_FIELD(uint64, Timestamp)
+	UE_TRACE_EVENT_FIELD(int32, NumChannels)
+UE_TRACE_EVENT_END()
+
+UE_TRACE_EVENT_BEGIN(Audio, MainSubmixTruePeakMaxValue)
+	UE_TRACE_EVENT_FIELD(uint32, DeviceId)
+	UE_TRACE_EVENT_FIELD(uint32, SubmixId)
+	UE_TRACE_EVENT_FIELD(uint64, Timestamp)
+	UE_TRACE_EVENT_FIELD(float, TruePeakMaxValueDb)
+UE_TRACE_EVENT_END()
+#endif // !WITH_EDITOR
+
 UE_TRACE_EVENT_BEGIN(Audio, SubmixAlivePing)
 	UE_TRACE_EVENT_FIELD(uint32, DeviceId)
 	UE_TRACE_EVENT_FIELD(uint32, SubmixId)
@@ -444,6 +460,73 @@ namespace Audio
 		SubmixesWithAutoEnvelopeFollowers.Reset();
 	}
 
+#if !WITH_EDITOR
+	void FMixerDevice::InitializeMainSubmixTruePeakEnvelopeFollower(int32 InSampleRate, int32 InNumChannels, int32 InNumSamples)
+	{
+		MainSubmixTruePeakEnvelopeFollower = FEnvelopeFollower(
+			{
+				.SampleRate = static_cast<float>(InSampleRate),
+				.NumChannels = InNumChannels,
+				.AttackTimeMsec = 0.0f,
+				.Mode = EPeakMode::TruePeak,
+				.bIsAnalog = false,
+			});
+		
+		MainSubmixTruePeakOutputEnvelope.SetNumUninitialized(InNumSamples);
+	}
+
+	void FMixerDevice::UpdateMainSubmixTruePeakMaxValue(const FAlignedFloatBuffer& InMainSubmixAudioBuffer)
+	{
+		MainSubmixTruePeakEnvelopeFollower.ProcessAudio(InMainSubmixAudioBuffer.GetData(), PlatformSettings.CallbackBufferFrameSize, MainSubmixTruePeakOutputEnvelope.GetData());
+
+		if (const float* MainSubmixTruePeakMaxValuePtr = Algo::MaxElement(MainSubmixTruePeakOutputEnvelope))
+		{
+			MainSubmixTruePeakMaxValueDb = ConvertToDecibels(*MainSubmixTruePeakMaxValuePtr);
+		}
+	}
+
+	void FMixerDevice::SendMainSubmixLoadedMessage()
+	{
+		const bool bChannelEnabled = UE_TRACE_CHANNELEXPR_IS_ENABLED(AudioMixerChannel);
+		if (!bChannelEnabled)
+		{
+			return;
+		}
+
+		const FMixerSubmixPtr MainSubmixPtr = GetMainSubmix().Pin();
+		if (!MainSubmixPtr.IsValid())
+		{
+			return;
+		}
+
+		UE_TRACE_LOG(Audio, MainSubmixLoaded, AudioMixerChannel)
+			<< MainSubmixLoaded.DeviceId(DeviceID)
+			<< MainSubmixLoaded.SubmixId(GetTypeHash(MainSubmixPtr->GetPathName()))
+			<< MainSubmixLoaded.Timestamp(FPlatformTime::Cycles64());
+	}
+
+	void FMixerDevice::SendMainSubmixTruePeakMaxValueMessage()
+	{
+		const bool bChannelEnabled = UE_TRACE_CHANNELEXPR_IS_ENABLED(AudioMixerChannel);
+		if (!bChannelEnabled)
+		{
+			return;
+		}
+
+		const FMixerSubmixPtr MainSubmixPtr = GetMainSubmix().Pin();
+		if (!MainSubmixPtr.IsValid())
+		{
+			return;
+		}
+
+		UE_TRACE_LOG(Audio, MainSubmixTruePeakMaxValue, AudioMixerChannel)
+			<< MainSubmixTruePeakMaxValue.DeviceId(DeviceID)
+			<< MainSubmixTruePeakMaxValue.SubmixId(GetTypeHash(MainSubmixPtr->GetPathName()))
+			<< MainSubmixTruePeakMaxValue.Timestamp(FPlatformTime::Cycles64())
+			<< MainSubmixTruePeakMaxValue.TruePeakMaxValueDb(MainSubmixTruePeakMaxValueDb);
+	}
+#endif // !WITH_EDITOR
+
 	void FMixerDevice::SendSubmixesAlivePing()
 	{
 		const bool bChannelEnabled = UE_TRACE_CHANNELEXPR_IS_ENABLED(AudioChannel);
@@ -1022,6 +1105,12 @@ namespace Audio
 				InitSoundSubmixes();
 
 #if UE_AUDIO_PROFILERTRACE_ENABLED
+
+#if !WITH_EDITOR
+				const int32 NumSamples = PlatformInfo.NumChannels * PlatformSettings.CallbackBufferFrameSize;
+				InitializeMainSubmixTruePeakEnvelopeFollower(PlatformSettings.SampleRate, PlatformInfo.NumChannels, NumSamples);
+#endif // !WITH_EDITOR
+
 				RegisterSubmixEnvelopeFollowerCVars();
 #endif // UE_AUDIO_PROFILERTRACE_ENABLED
 
@@ -1423,6 +1512,12 @@ namespace Audio
 		UpdateAudioClock();
 
 #if UE_AUDIO_PROFILERTRACE_ENABLED
+
+#if !WITH_EDITOR
+		UpdateMainSubmixTruePeakMaxValue(Output);
+		SendMainSubmixTruePeakMaxValueMessage();
+#endif // !WITH_EDITOR
+		
 		UpdateSubmixesAlivePingTime();
 		SendSubmixesSendLevelTraceMessages();
 #endif //UE_AUDIO_PROFILERTRACE_ENABLED
@@ -1637,6 +1732,10 @@ namespace Audio
 			// 1. Load or reload all sound submixes/instances
 			LoadRequiredSubmix(ERequiredSubmixes::Main, GetRequiredSubmixName(ERequiredSubmixes::Main), false /* DefaultMuteWhenBackgrounded */, AudioSettings->MasterSubmix);
 
+#if !WITH_EDITOR
+			SendMainSubmixLoadedMessage();
+#endif // !WITH_EDITOR
+
 			// BaseDefaultSubmix is an optional master submix type set by project settings
 			if (AudioSettings->BaseDefaultSubmix.IsValid())
 			{
diff --git a/Engine/Source/Runtime/AudioMixer/Public/AudioMixerDevice.h b/Engine/Source/Runtime/AudioMixer/Public/AudioMixerDevice.h
--- a/Engine/Source/Runtime/AudioMixer/Public/AudioMixerDevice.h
+++ b/Engine/Source/Runtime/AudioMixer/Public/AudioMixerDevice.h
@@ -14,6 +14,10 @@
 #include "UObject/GCObject.h"
 #include "UObject/StrongObjectPtr.h"
 
+#if UE_AUDIO_PROFILERTRACE_ENABLED && !WITH_EDITOR
+#include "DSP/EnvelopeFollower.h"
+#endif // UE_AUDIO_PROFILERTRACE_ENABLED && !WITH_EDITOR
+
 // Forward Declarations
 class FOnSubmixEnvelopeBP;
 class IAudioMixerPlatformInterface;
@@ -410,6 +414,15 @@ namespace Audio
 #if UE_AUDIO_PROFILERTRACE_ENABLED
 		void OnTraceStarted(FTraceAuxiliary::EConnectionType TraceType, const FString& TraceDestination);
 		void OnTraceStopped(FTraceAuxiliary::EConnectionType TraceType, const FString& TraceDestination);
+
+#if !WITH_EDITOR
+		void InitializeMainSubmixTruePeakEnvelopeFollower(int32 InSampleRate, int32 InNumChannels, int32 InNumSamples);
+		void UpdateMainSubmixTruePeakMaxValue(const FAlignedFloatBuffer& InMainSubmixAudioBuffer);
+
+		void SendMainSubmixLoadedMessage();
+		void SendMainSubmixTruePeakMaxValueMessage();
+#endif // !WITH_EDITOR
+	
 		void SendSubmixesAlivePing();
 		void UpdateSubmixesAlivePingTime();
 		void SendSubmixesSendLevelTraceMessages();
@@ -517,6 +530,12 @@ namespace Audio
 		IConsoleVariable* StartSubmixEnvelopeFollowerCVar = nullptr;
 		IConsoleVariable* StopSubmixEnvelopeFollowerCVar  = nullptr;
 
+#if !WITH_EDITOR
+		FEnvelopeFollower MainSubmixTruePeakEnvelopeFollower;
+		FAlignedFloatBuffer MainSubmixTruePeakOutputEnvelope;
+		float MainSubmixTruePeakMaxValueDb = MIN_VOLUME_DECIBELS;
+#endif // !WITH_EDITOR
+
 		double SubmixesAlivePingTimeSec = 0.0;
 		static constexpr double SubmixesAlivePingWaitTimeSec = 1.0;
 #endif // UE_AUDIO_PROFILERTRACE_ENABLED
