[AudioMixer] FMixerDevice: added missing #if UE_AUDIO_PROFILERTRACE_ENABLED

#rnx

#Synced-CL 50901149

diff --git a/Engine/Source/Runtime/AudioMixer/Private/AudioMixerDevice.cpp b/Engine/Source/Runtime/AudioMixer/Private/AudioMixerDevice.cpp
--- a/Engine/Source/Runtime/AudioMixer/Private/AudioMixerDevice.cpp
+++ b/Engine/Source/Runtime/AudioMixer/Private/AudioMixerDevice.cpp
@@ -1732,9 +1732,9 @@ namespace Audio
 			// 1. Load or reload all sound submixes/instances
 			LoadRequiredSubmix(ERequiredSubmixes::Main, GetRequiredSubmixName(ERequiredSubmixes::Main), false /* DefaultMuteWhenBackgrounded */, AudioSettings->MasterSubmix);
 
-#if !WITH_EDITOR
+#if UE_AUDIO_PROFILERTRACE_ENABLED && !WITH_EDITOR
 			SendMainSubmixLoadedMessage();
-#endif // !WITH_EDITOR
+#endif // UE_AUDIO_PROFILERTRACE_ENABLED && !WITH_EDITOR
 
 			// BaseDefaultSubmix is an optional master submix type set by project settings
 			if (AudioSettings->BaseDefaultSubmix.IsValid())
