[AudioMixer] Update to using TSpscQueue/TMpscQueue in AudioMixerDevice

#rnx

#tests Basic smoke testing in BR_Audio in PIE
#Synced-CL 50894658

diff --git a/Engine/Source/Runtime/AudioMixer/Public/AudioMixerDevice.h b/Engine/Source/Runtime/AudioMixer/Public/AudioMixerDevice.h
--- a/Engine/Source/Runtime/AudioMixer/Public/AudioMixerDevice.h
+++ b/Engine/Source/Runtime/AudioMixer/Public/AudioMixerDevice.h
@@ -6,6 +6,7 @@
 #include "AudioMixer.h"
 #include "AudioDevice.h"
 #include "Containers/MpscQueue.h"
+#include "Containers/SpscQueue.h"
 #include "Sound/SoundSubmix.h"
 #include "Sound/SoundGenerator.h"
 #include "DSP/BufferVectorOperations.h"
@@ -499,7 +500,7 @@ namespace Audio
 		TArray<USoundSubmix*> DelegateBoundSubmixes;
 
 		/** Queue of mixer source voices. */
-		TQueue<FMixerSourceVoice*> SourceVoices;
+		TSpscQueue<FMixerSourceVoice*> SourceVoices;
 
 		TMap<uint32, TArray<FSourceEffectChainEntry>> SourceEffectChainOverrides;
 
@@ -516,7 +517,7 @@ namespace Audio
 		mutable std::atomic<int32> AudioPlatformThreadId;
 
 		/** Command queue to send commands to audio render thread from game thread or audio thread. */
-		TQueue<TFunction<void()>> CommandQueue;
+		TMpscQueue<TFunction<void()>> CommandQueue;
 
 		/** MPSC command queue to send commands to the game thread */
 		TMpscQueue<TFunction<void()>> GameThreadCommandQueue;
