[Audio render scheduler] Enable by default
* Set ECVF_SaveForNextBoot flag on the cvar so it can be hotfixed
* Automatically dump the scheduler state one time in case of an error

#rnx

#Synced-CL 50894658

diff --git a/Engine/Source/Runtime/AudioMixer/Private/AudioMixerCVars.cpp b/Engine/Source/Runtime/AudioMixer/Private/AudioMixerCVars.cpp
--- a/Engine/Source/Runtime/AudioMixer/Private/AudioMixerCVars.cpp
+++ b/Engine/Source/Runtime/AudioMixer/Private/AudioMixerCVars.cpp
@@ -9,13 +9,13 @@
 
 namespace AudioMixerCVars
 {
-	int32 UseRenderScheduler = 0;
+	int32 UseRenderScheduler = 1;
 	FAutoConsoleVariableRef CVarUseAudioRenderScheduler(
 		TEXT("au.UseRenderScheduler"),
 		UseRenderScheduler,
 		TEXT("Use audio rendering scheduler to handle render dependencies.\n")
 		TEXT("0: Disabled, 1: Enabled"),
-		ECVF_Default);
+		ECVF_SaveForNextBoot);
 
 #if ENABLE_AUDIO_DEBUG
 	FAutoConsoleCommand AudioRenderSchedulerDumpCmd(
diff --git a/Engine/Source/Runtime/AudioMixer/Private/AudioRenderScheduler.cpp b/Engine/Source/Runtime/AudioMixer/Private/AudioRenderScheduler.cpp
--- a/Engine/Source/Runtime/AudioMixer/Private/AudioRenderScheduler.cpp
+++ b/Engine/Source/Runtime/AudioMixer/Private/AudioRenderScheduler.cpp
@@ -100,6 +100,9 @@ namespace Audio
 			UE_LOGFMT(LogAudioRenderScheduler, Error, "Tried to add duplicate render step '{Name}' with id {Id}",
 				("Name", Step->GetRenderStepName()),
 				("Id", Id.GetRawValue()));
+#if ENABLE_AUDIO_DEBUG
+			DumpOnError();
+#endif
 		}
 	}
 
@@ -123,6 +126,9 @@ namespace Audio
 		{
 			UE_LOGFMT(LogAudioRenderScheduler, Error, "Tried to remove unknown render step with id {Id}",
 				("Id", Id.GetRawValue()));
+#if ENABLE_AUDIO_DEBUG
+			DumpOnError();
+#endif
 		}
 	}
 
@@ -170,6 +176,9 @@ namespace Audio
 			UE_LOGFMT(LogAudioRenderScheduler, Error, "Unknown dependency removed from audio render scheduler: {Id1} -> {Id2}",
 				("Id1", FirstStep.GetRawValue()),
 				("Id2", SecondStep.GetRawValue()));
+#if ENABLE_AUDIO_DEBUG
+			DumpOnError();
+#endif
 		}
 
 		bDirty = true;
@@ -463,5 +472,15 @@ namespace Audio
 
 		UE_LOGFMT(LogAudioRenderScheduler, Display, "Step ordering: {Steps}", ("Steps", ActiveSteps));
 	}
+
+	void FAudioRenderScheduler::DumpOnError()
+	{
+		// Only do this once so we don't spam in a persistent error state.
+		if (!bDidDumpOnError)
+		{
+			bDidDumpOnError = true;
+			bShouldDumpSteps = true;
+		}
+	}
 #endif
 }
diff --git a/Engine/Source/Runtime/AudioMixer/Private/AudioRenderScheduler.h b/Engine/Source/Runtime/AudioMixer/Private/AudioRenderScheduler.h
--- a/Engine/Source/Runtime/AudioMixer/Private/AudioRenderScheduler.h
+++ b/Engine/Source/Runtime/AudioMixer/Private/AudioRenderScheduler.h
@@ -117,8 +117,10 @@ namespace Audio
 
 #if ENABLE_AUDIO_DEBUG
 		void DumpStepsInternal();
+		void DumpOnError();
 
 		bool bShouldDumpSteps = false;
+		bool bDidDumpOnError = false;
 #endif
 
 		struct FStepEntry
