[AudioMixer] FMixerDevice: added LKFS analyzer for main submix and send its resulting loudness values as a trace message to Audio Insights.

#Synced-CL 51090472

diff --git a/Engine/Source/Runtime/AudioMixer/Private/AudioMixerDevice.cpp b/Engine/Source/Runtime/AudioMixer/Private/AudioMixerDevice.cpp
--- a/Engine/Source/Runtime/AudioMixer/Private/AudioMixerDevice.cpp
+++ b/Engine/Source/Runtime/AudioMixer/Private/AudioMixerDevice.cpp
@@ -149,6 +149,17 @@ UE_TRACE_EVENT_BEGIN(Audio, MainSubmixTruePeakMaxValue)
 	UE_TRACE_EVENT_FIELD(uint64, Timestamp)
 	UE_TRACE_EVENT_FIELD(float, TruePeakMaxValueDb)
 UE_TRACE_EVENT_END()
+
+UE_TRACE_EVENT_BEGIN(Audio, MainSubmixLKFSValues)
+	UE_TRACE_EVENT_FIELD(uint32, DeviceId)
+	UE_TRACE_EVENT_FIELD(uint32, SubmixId)
+	UE_TRACE_EVENT_FIELD(uint64, Timestamp)
+	UE_TRACE_EVENT_FIELD(float, LongTermLoudness)
+	UE_TRACE_EVENT_FIELD(float, ShortTermLoudness)
+	UE_TRACE_EVENT_FIELD(float, MomentaryLoudness)
+	UE_TRACE_EVENT_FIELD(float, LoudnessRangeLowerBound)
+	UE_TRACE_EVENT_FIELD(float, LoudnessRangeUpperBound)
+UE_TRACE_EVENT_END()
 #endif // !WITH_EDITOR
 
 UE_TRACE_EVENT_BEGIN(Audio, SubmixAlivePing)
@@ -463,8 +474,8 @@ namespace Audio
 #if !WITH_EDITOR
 	void FMixerDevice::InitializeMainSubmixTruePeakEnvelopeFollower(int32 InSampleRate, int32 InNumChannels, int32 InNumSamples)
 	{
-		MainSubmixTruePeakEnvelopeFollower = FEnvelopeFollower(
-			{
+		MainSubmixTruePeakEnvelopeFollower = MakeUnique<FEnvelopeFollower>(
+			FEnvelopeFollowerInitParams{
 				.SampleRate = static_cast<float>(InSampleRate),
 				.NumChannels = InNumChannels,
 				.AttackTimeMsec = 0.0f,
@@ -475,14 +486,12 @@ namespace Audio
 		MainSubmixTruePeakOutputEnvelope.SetNumUninitialized(InNumSamples);
 	}
 
-	void FMixerDevice::UpdateMainSubmixTruePeakMaxValue(const FAlignedFloatBuffer& InMainSubmixAudioBuffer)
+	void FMixerDevice::InitializeMainSubmixLKFSAnalyzer(int32 InSampleRate, int32 InNumChannels)
 	{
-		MainSubmixTruePeakEnvelopeFollower.ProcessAudio(InMainSubmixAudioBuffer.GetData(), PlatformSettings.CallbackBufferFrameSize, MainSubmixTruePeakOutputEnvelope.GetData());
+		FLKFSAnalyzerSettings LKFSAnalyzerSettings;
+		LKFSAnalyzerSettings.bCalculateOverallLoudnessRange = true;
 
-		if (const float* MainSubmixTruePeakMaxValuePtr = Algo::MaxElement(MainSubmixTruePeakOutputEnvelope))
-		{
-			MainSubmixTruePeakMaxValueDb = ConvertToDecibels(*MainSubmixTruePeakMaxValuePtr);
-		}
+		MainSubmixLKFSAnalyzer = MakeUnique<FLKFSAnalyzer>(InNumChannels, static_cast<float>(InSampleRate), LKFSAnalyzerSettings);
 	}
 
 	void FMixerDevice::SendMainSubmixLoadedMessage()
@@ -505,7 +514,7 @@ namespace Audio
 			<< MainSubmixLoaded.Timestamp(FPlatformTime::Cycles64());
 	}
 
-	void FMixerDevice::SendMainSubmixTruePeakMaxValueMessage()
+	void FMixerDevice::UpdateAndSendMainSubmixTruePeakMaxValueMessage(const FAlignedFloatBuffer& InMainSubmixAudioBuffer)
 	{
 		const bool bChannelEnabled = UE_TRACE_CHANNELEXPR_IS_ENABLED(AudioMixerChannel);
 		if (!bChannelEnabled)
@@ -519,12 +528,69 @@ namespace Audio
 			return;
 		}
 
+		if (!MainSubmixTruePeakEnvelopeFollower.IsValid())
+		{
+			return;
+		}
+
+		MainSubmixTruePeakEnvelopeFollower->ProcessAudio(InMainSubmixAudioBuffer.GetData(), PlatformSettings.CallbackBufferFrameSize, MainSubmixTruePeakOutputEnvelope.GetData());
+		
+		const float* MainSubmixTruePeakMaxValuePtr = Algo::MaxElement(MainSubmixTruePeakOutputEnvelope);
+		const float MainSubmixTruePeakMaxValueDb   = MainSubmixTruePeakMaxValuePtr ? ConvertToDecibels(*MainSubmixTruePeakMaxValuePtr) : MIN_VOLUME_DECIBELS;
+
 		UE_TRACE_LOG(Audio, MainSubmixTruePeakMaxValue, AudioMixerChannel)
 			<< MainSubmixTruePeakMaxValue.DeviceId(DeviceID)
 			<< MainSubmixTruePeakMaxValue.SubmixId(GetTypeHash(MainSubmixPtr->GetPathName()))
 			<< MainSubmixTruePeakMaxValue.Timestamp(FPlatformTime::Cycles64())
 			<< MainSubmixTruePeakMaxValue.TruePeakMaxValueDb(MainSubmixTruePeakMaxValueDb);
 	}
+
+	void FMixerDevice::UpdateAndSendMainSubmixLKFSValuesMessage(const FAlignedFloatBuffer& InMainSubmixAudioBuffer)
+	{
+		const bool bChannelEnabled = UE_TRACE_CHANNELEXPR_IS_ENABLED(AudioMixerChannel);
+		if (!bChannelEnabled)
+		{
+			return;
+		}
+
+		const FMixerSubmixPtr MainSubmixPtr = GetMainSubmix().Pin();
+		if (!MainSubmixPtr.IsValid())
+		{
+			return;
+		}
+
+		if (!MainSubmixLKFSAnalyzer.IsValid())
+		{
+			return;
+		}
+
+		MainSubmixLKFSAnalyzer->Analyze(InMainSubmixAudioBuffer, MainSubmixLKFSAnalyzerResults);
+
+		const TArray<FLKFSResult>& OverallLKFSArray = MainSubmixLKFSAnalyzerResults.GetLoudnessResults();
+		if (OverallLKFSArray.Num() > 0)
+		{
+			const FLKFSResult& LatestLKFSResult = OverallLKFSArray[OverallLKFSArray.Num() - 1];
+
+			float LoudnessRangeLowerBound = 0.0f;
+			float LoudnessRangeUpperBound = 0.0f;
+
+			if (LatestLKFSResult.LoudnessRange.IsSet())
+			{
+				LoudnessRangeLowerBound = LatestLKFSResult.LoudnessRange.GetValue().Min;
+				LoudnessRangeUpperBound = LatestLKFSResult.LoudnessRange.GetValue().Max;
+			}
+
+			UE_TRACE_LOG(Audio, MainSubmixLKFSValues, AudioMixerChannel)
+				<< MainSubmixLKFSValues.DeviceId(DeviceID)
+				<< MainSubmixLKFSValues.SubmixId(GetTypeHash(MainSubmixPtr->GetPathName()))
+				<< MainSubmixLKFSValues.Timestamp(FPlatformTime::Cycles64())
+				<< MainSubmixLKFSValues.LongTermLoudness(LatestLKFSResult.GatedLoudness)
+				<< MainSubmixLKFSValues.ShortTermLoudness(LatestLKFSResult.ShortTermLoudness)
+				<< MainSubmixLKFSValues.MomentaryLoudness(LatestLKFSResult.Loudness)
+				<< MainSubmixLKFSValues.LoudnessRangeLowerBound(LoudnessRangeLowerBound)
+				<< MainSubmixLKFSValues.LoudnessRangeUpperBound(LoudnessRangeUpperBound);
+		}
+	}
 #endif // !WITH_EDITOR
 
 	void FMixerDevice::SendSubmixesAlivePing()
@@ -1108,7 +1174,9 @@ namespace Audio
 
 #if !WITH_EDITOR
 				const int32 NumSamples = PlatformInfo.NumChannels * PlatformSettings.CallbackBufferFrameSize;
+
 				InitializeMainSubmixTruePeakEnvelopeFollower(PlatformSettings.SampleRate, PlatformInfo.NumChannels, NumSamples);
+				InitializeMainSubmixLKFSAnalyzer(PlatformSettings.SampleRate, PlatformInfo.NumChannels);
 #endif // !WITH_EDITOR
 
 				RegisterSubmixEnvelopeFollowerCVars();
@@ -1197,6 +1265,11 @@ namespace Audio
 			AudioMixerPlatform->TeardownHardware();
 		}
 
+#if UE_AUDIO_PROFILERTRACE_ENABLED && !WITH_EDITOR
+		MainSubmixTruePeakEnvelopeFollower.Reset();
+		MainSubmixLKFSAnalyzer.Reset();
+#endif // UE_AUDIO_PROFILERTRACE_ENABLED && !WITH_EDITOR
+
 		// Reset existing submixes if they exist
 		RequiredSubmixInstances.Reset();
 		Submixes.Reset();
@@ -1514,8 +1587,8 @@ namespace Audio
 #if UE_AUDIO_PROFILERTRACE_ENABLED
 
 #if !WITH_EDITOR
-		UpdateMainSubmixTruePeakMaxValue(Output);
-		SendMainSubmixTruePeakMaxValueMessage();
+		UpdateAndSendMainSubmixTruePeakMaxValueMessage(Output);
+		UpdateAndSendMainSubmixLKFSValuesMessage(Output);
 #endif // !WITH_EDITOR
 		
 		UpdateSubmixesAlivePingTime();
diff --git a/Engine/Source/Runtime/AudioMixer/Public/AudioMixerDevice.h b/Engine/Source/Runtime/AudioMixer/Public/AudioMixerDevice.h
--- a/Engine/Source/Runtime/AudioMixer/Public/AudioMixerDevice.h
+++ b/Engine/Source/Runtime/AudioMixer/Public/AudioMixerDevice.h
@@ -17,6 +17,7 @@
 
 #if UE_AUDIO_PROFILERTRACE_ENABLED && !WITH_EDITOR
 #include "DSP/EnvelopeFollower.h"
+#include "DSP/LKFSAnalyzer.h"
 #endif // UE_AUDIO_PROFILERTRACE_ENABLED && !WITH_EDITOR
 
 // Forward Declarations
@@ -418,10 +419,12 @@ namespace Audio
 
 #if !WITH_EDITOR
 		void InitializeMainSubmixTruePeakEnvelopeFollower(int32 InSampleRate, int32 InNumChannels, int32 InNumSamples);
-		void UpdateMainSubmixTruePeakMaxValue(const FAlignedFloatBuffer& InMainSubmixAudioBuffer);
+		void InitializeMainSubmixLKFSAnalyzer(int32 InSampleRate, int32 InNumChannels);
 
 		void SendMainSubmixLoadedMessage();
-		void SendMainSubmixTruePeakMaxValueMessage();
+
+		void UpdateAndSendMainSubmixTruePeakMaxValueMessage(const FAlignedFloatBuffer& InMainSubmixAudioBuffer);
+		void UpdateAndSendMainSubmixLKFSValuesMessage(const FAlignedFloatBuffer& InMainSubmixAudioBuffer);
 #endif // !WITH_EDITOR
 	
 		void SendSubmixesAlivePing();
@@ -532,9 +535,11 @@ namespace Audio
 		IConsoleVariable* StopSubmixEnvelopeFollowerCVar  = nullptr;
 
 #if !WITH_EDITOR
-		FEnvelopeFollower MainSubmixTruePeakEnvelopeFollower;
+		TUniquePtr<FEnvelopeFollower> MainSubmixTruePeakEnvelopeFollower;
 		FAlignedFloatBuffer MainSubmixTruePeakOutputEnvelope;
-		float MainSubmixTruePeakMaxValueDb = MIN_VOLUME_DECIBELS;
+
+		TUniquePtr<FLKFSAnalyzer> MainSubmixLKFSAnalyzer;
+		FLKFSAnalyzerResults MainSubmixLKFSAnalyzerResults;
 #endif // !WITH_EDITOR
 
 		double SubmixesAlivePingTimeSec = 0.0;
