[Audio render scheduler] Only read cvar to enable scheduler at device creation time

#rnx
#tests Tested in BR_Audio with and without scheduler turned on

#Synced-CL 51097580

diff --git a/Engine/Source/Runtime/AudioMixer/Private/AudioMixerBus.cpp b/Engine/Source/Runtime/AudioMixer/Private/AudioMixerBus.cpp
--- a/Engine/Source/Runtime/AudioMixer/Private/AudioMixerBus.cpp
+++ b/Engine/Source/Runtime/AudioMixer/Private/AudioMixerBus.cpp
@@ -21,7 +21,7 @@ namespace Audio
 	{
 		SetNumOutputChannels(NumChannels);
 
-		if (AudioMixerCVars::UseRenderScheduler)
+		if (SourceManager->UseRenderScheduler())
 		{
 			SourceManager->GetRenderScheduler().AddStep(FAudioRenderStepId::FromAudioBusKey(BusKey), this);
 		}
@@ -29,7 +29,7 @@ namespace Audio
 
 	FMixerAudioBus::~FMixerAudioBus()
 	{
-		if (AudioMixerCVars::UseRenderScheduler)
+		if (SourceManager->UseRenderScheduler())
 		{
 			SourceManager->GetRenderScheduler().RemoveStep(FAudioRenderStepId::FromAudioBusKey(BusKey));
 		}
@@ -49,14 +49,14 @@ namespace Audio
 	void FMixerAudioBus::Update()
 	{
 		// When using the scheduler we flip buffers at the beginning of MixBuffer() instead.
-		check(!AudioMixerCVars::UseRenderScheduler);
+		check(!SourceManager->UseRenderScheduler());
 		CurrentBufferIndex = 1 - CurrentBufferIndex;
 	}
 
 	void FMixerAudioBus::AddInstanceId(const int32 InSourceId, const uint32 WaveInstancePlayOrder, int32 InNumOutputChannels)
 	{
 		InstanceIds.Add(InSourceId);
-		if (AudioMixerCVars::UseRenderScheduler)
+		if (SourceManager->UseRenderScheduler())
 		{
 			// Make sure this bus gets mixed before the source from it gets rendered
 			const FAudioRenderStepId SourceStepId = FAudioRenderStepId::FromWaveInstancePlayOrder(WaveInstancePlayOrder);
@@ -67,7 +67,7 @@ namespace Audio
 
 	bool FMixerAudioBus::RemoveInstanceId(const int32 InSourceId, const uint32 WaveInstancePlayOrder)
 	{
-		if (InstanceIds.Remove(InSourceId) > 0 && AudioMixerCVars::UseRenderScheduler)
+		if (InstanceIds.Remove(InSourceId) > 0 && SourceManager->UseRenderScheduler())
 		{
 			const FAudioRenderStepId SourceStepId = FAudioRenderStepId::FromWaveInstancePlayOrder(WaveInstancePlayOrder);
 			const FAudioRenderStepId BusStepId = FAudioRenderStepId::FromAudioBusKey(BusKey);
@@ -94,7 +94,7 @@ namespace Audio
 		// It's a new source id so just add it
 		AudioBusSends[(int32)BusSendType].Add(InAudioBusSend);
 
-		if (AudioMixerCVars::UseRenderScheduler)
+		if (SourceManager->UseRenderScheduler())
 		{
 			// Make sure the source sending to the bus gets rendered before the bus is mixed
 			const FAudioRenderStepId SourceStepId = FAudioRenderStepId::FromWaveInstancePlayOrder(InAudioBusSend.WaveInstancePlayOrder);
@@ -112,7 +112,7 @@ namespace Audio
 			// Remove this source id's send
 			if (Sends[i].SourceId == InSourceId)
 			{
-				if (AudioMixerCVars::UseRenderScheduler)
+				if (SourceManager->UseRenderScheduler())
 				{
 					const FAudioRenderStepId SourceStepId = FAudioRenderStepId::FromWaveInstancePlayOrder(Sends[i].WaveInstancePlayOrder);
 					const FAudioRenderStepId BusStepId = FAudioRenderStepId::FromAudioBusKey(BusKey);
@@ -136,7 +136,7 @@ namespace Audio
 		const int32 NumSamples = NumFrames * NumChannels;
 		const int32 NumOutputFrames = SourceManager->GetNumOutputFrames();
 
-		if (AudioMixerCVars::UseRenderScheduler)
+		if (SourceManager->UseRenderScheduler())
 		{
 			CurrentBufferIndex = 1 - CurrentBufferIndex;
 		}
@@ -156,7 +156,7 @@ namespace Audio
 				// If the audio source mixing to this audio bus is itself a source bus, we need to use the previous renderer buffer to avoid infinite recursion.
 				// With the render scheduler we don't need to treat source buses any differently from other sources -- the source bus should already have been
 				//  rendered.
-				if (!AudioMixerCVars::UseRenderScheduler && SourceManager->IsSourceBus(AudioBusSend.SourceId))
+				if (!SourceManager->UseRenderScheduler() && SourceManager->IsSourceBus(AudioBusSend.SourceId))
 				{
 					SourceBufferPtr = SourceManager->GetPreviousSourceBusBuffer(AudioBusSend.SourceId);
 				}
diff --git a/Engine/Source/Runtime/AudioMixer/Private/AudioMixerCVars.cpp b/Engine/Source/Runtime/AudioMixer/Private/AudioMixerCVars.cpp
--- a/Engine/Source/Runtime/AudioMixer/Private/AudioMixerCVars.cpp
+++ b/Engine/Source/Runtime/AudioMixer/Private/AudioMixerCVars.cpp
@@ -13,7 +13,7 @@ namespace AudioMixerCVars
 	FAutoConsoleVariableRef CVarUseAudioRenderScheduler(
 		TEXT("au.UseRenderScheduler"),
 		UseRenderScheduler,
-		TEXT("Use audio rendering scheduler to handle render dependencies.\n")
+		TEXT("Use audio rendering scheduler to handle render dependencies. Should be set in an ini file or on the command line, as changing this after audio mixer initialization has no effect.\n")
 		TEXT("0: Disabled, 1: Enabled"),
 		ECVF_SaveForNextBoot);
 
@@ -24,10 +24,12 @@ namespace AudioMixerCVars
 		FConsoleCommandDelegate::CreateStatic(
 			[]()
 			{
+				using namespace Audio;
 				FAudioDevice* AudioDevice = GEngine != nullptr ? GEngine->GetMainAudioDeviceRaw() : nullptr;
 				if (AudioDevice != nullptr)
 				{
-					static_cast<Audio::FMixerDevice*>(AudioDevice)->GetRenderScheduler().DumpSteps();
+					FMixerDevice* MixerDevice = static_cast<FMixerDevice*>(AudioDevice);
+					MixerDevice->GetRenderScheduler().DumpSteps();
 				}
 			}
 		)
diff --git a/Engine/Source/Runtime/AudioMixer/Private/AudioMixerCVars.h b/Engine/Source/Runtime/AudioMixer/Private/AudioMixerCVars.h
--- a/Engine/Source/Runtime/AudioMixer/Private/AudioMixerCVars.h
+++ b/Engine/Source/Runtime/AudioMixer/Private/AudioMixerCVars.h
@@ -6,5 +6,7 @@
 
 namespace AudioMixerCVars
 {
+	// This is used to decide whether or not to create a render scheduler when the mixer device is created.
+	//  Most code should test FMixerDevice::UseRenderScheduler() instead of using this.
 	extern int32 UseRenderScheduler;
 }
diff --git a/Engine/Source/Runtime/AudioMixer/Private/AudioMixerDevice.cpp b/Engine/Source/Runtime/AudioMixer/Private/AudioMixerDevice.cpp
--- a/Engine/Source/Runtime/AudioMixer/Private/AudioMixerDevice.cpp
+++ b/Engine/Source/Runtime/AudioMixer/Private/AudioMixerDevice.cpp
@@ -13,10 +13,12 @@
 #include "AudioMixerSubmix.h"
 #include "AudioMixerSourceVoice.h"
 #include "AudioPluginUtilities.h"
+#include "AudioMixerCVars.h"
 #include "AudioMixerEffectsManager.h"
 #include "AudioRenderScheduler.h"
 #include "DSP/Noise.h"
 #include "DSP/SinOsc.h"
+#include "Logging/StructuredLog.h"
 #include "Sound/AudioSettings.h"
 #include "Sound/SoundSubmix.h"
 #include "Sound/SoundSubmixSend.h"
@@ -396,7 +398,10 @@ namespace Audio
 		bAudioMixerModuleLoaded = true;
 
 		SourceManager = MakeUnique<FMixerSourceManager>(this);
+
 		RenderScheduler = MakeUnique<FAudioRenderScheduler>(this);
+		bUseRenderScheduler = AudioMixerCVars::UseRenderScheduler != 0;
+		UE_LOGFMT(LogAudioMixer, Display, "bUseRenderScheduler set to {UseRenderScheduler}.", ("UseRenderScheduler", bUseRenderScheduler));
 
 		// Register AudioLink Factory. 	
 		TArray<FName> Factories = IAudioLinkFactory::GetAllRegisteredFactoryNames();
diff --git a/Engine/Source/Runtime/AudioMixer/Private/AudioMixerSourceManager.cpp b/Engine/Source/Runtime/AudioMixer/Private/AudioMixerSourceManager.cpp
--- a/Engine/Source/Runtime/AudioMixer/Private/AudioMixerSourceManager.cpp
+++ b/Engine/Source/Runtime/AudioMixer/Private/AudioMixerSourceManager.cpp
@@ -1044,7 +1044,7 @@ namespace Audio
 		// Free the mixer source buffer data
 		if (SourceInfo.MixerSourceBuffer.IsValid())
 		{
-			if (AudioMixerCVars::UseRenderScheduler)
+			if (UseRenderScheduler())
 			{
 				RemoveBusPatchDependencies(SourceInfo.MixerSourceBuffer->GetInstanceID());
 			}
@@ -1570,9 +1570,9 @@ namespace Audio
 			SourceInfo.DebugName = InitParams.DebugName;
 #endif 
 
-			if (AudioMixerCVars::UseRenderScheduler)
+			if (UseRenderScheduler())
 			{
-				MixerDevice->GetRenderScheduler().AddStep(FAudioRenderStepId::FromWaveInstancePlayOrder(SourceInfo.PlayOrder), &SourceInfo);
+				GetRenderScheduler().AddStep(FAudioRenderStepId::FromWaveInstancePlayOrder(SourceInfo.PlayOrder), &SourceInfo);
 			}
 
 			AUDIO_MIXER_DEBUG_LOG(SourceId, TEXT("Is initializing"));
@@ -1606,11 +1606,11 @@ namespace Audio
 		{
 			AUDIO_MIXER_CHECK_AUDIO_PLAT_THREAD(MixerDevice);
 
-			if (AudioMixerCVars::UseRenderScheduler)
+			if (UseRenderScheduler())
 			{
 				FSourceInfo& SourceInfo = SourceInfos[SourceId];
 				check(SourceInfo.MixerSourceBuffer.IsValid());
-				MixerDevice->GetRenderScheduler().RemoveStep(FAudioRenderStepId::FromWaveInstancePlayOrder(SourceInfo.PlayOrder));
+				GetRenderScheduler().RemoveStep(FAudioRenderStepId::FromWaveInstancePlayOrder(SourceInfo.PlayOrder));
 			}
 
 			ReleaseSource(SourceId);
@@ -2870,7 +2870,7 @@ namespace Audio
 			switch (PendingAudioBusConnection->PatchVariant.GetIndex())
 			{
 			case FPendingAudioBusConnection::FPatchVariant::IndexOfType<FPatchInput>():
-				if (AudioMixerCVars::UseRenderScheduler)
+				if (UseRenderScheduler())
 				{
 					const uint32 PlayOrder = PendingAudioBusConnection->WaveInstancePlayOrder;
 					if (PlayOrder != INDEX_NONE)
@@ -2883,7 +2883,7 @@ namespace Audio
 				AudioBusPtr->AddNewPatchInput(PendingAudioBusConnection->PatchVariant.Get<FPatchInput>());
 				break;
 			case FPendingAudioBusConnection::FPatchVariant::IndexOfType<FPatchOutputStrongPtr>():
-				if (AudioMixerCVars::UseRenderScheduler)
+				if (UseRenderScheduler())
 				{
 					const uint32 PlayOrder = PendingAudioBusConnection->WaveInstancePlayOrder;
 					if (PlayOrder != INDEX_NONE)
@@ -2934,7 +2934,7 @@ namespace Audio
 				continue;
 			}
 
-			if (!AudioMixerCVars::UseRenderScheduler)
+			if (!UseRenderScheduler())
 			{
 				AudioBus->Update();
 			}
@@ -2990,7 +2990,7 @@ namespace Audio
 
 	void FMixerSourceManager::AddBusPatchDependency(const uint32 WaveInstancePlayOrder, const FAudioRenderStepId FirstStep, const FAudioRenderStepId SecondStep)
 	{
-		check(AudioMixerCVars::UseRenderScheduler);
+		check(UseRenderScheduler());
 		if (FirstStep.IsValid() && SecondStep.IsValid())
 		{
 			GetRenderScheduler().AddDependency(FirstStep, SecondStep);
@@ -3000,14 +3000,15 @@ namespace Audio
 
 	void FMixerSourceManager::RemoveBusPatchDependencies(const uint32 WaveInstancePlayOrder)
 	{
-		check(AudioMixerCVars::UseRenderScheduler);
+		check(UseRenderScheduler());
 		TArray<TPair<FAudioRenderStepId, FAudioRenderStepId>> Dependencies;
 		BusPatchDependencies.MultiFind(WaveInstancePlayOrder, Dependencies, false);
 		if (!Dependencies.IsEmpty())
 		{
+			FAudioRenderScheduler& Scheduler = GetRenderScheduler();
 			for (TPair<FAudioRenderStepId, FAudioRenderStepId>& Dependency : Dependencies)
 			{
-				GetRenderScheduler().RemoveDependency(Dependency.Get<0>(), Dependency.Get<1>());
+				Scheduler.RemoveDependency(Dependency.Get<0>(), Dependency.Get<1>());
 			}
 			BusPatchDependencies.Remove(WaveInstancePlayOrder);
 		}
@@ -3958,12 +3959,12 @@ namespace Audio
 		// Update pending tasks and release them if they're finished
 		UpdatePendingReleaseData();
 
-		if (AudioMixerCVars::UseRenderScheduler)
+		if (UseRenderScheduler())
 		{
 			ConnectBusPatches();
 
 			RenderThreadPhase = ESourceManagerRenderThreadPhase::RunScheduledRenderSteps;
-			MixerDevice->GetRenderScheduler().RenderBlock(DisableParallelSourceProcessingCvar != 0);
+			GetRenderScheduler().RenderBlock(DisableParallelSourceProcessingCvar != 0);
 		}
 		else
 		{
diff --git a/Engine/Source/Runtime/AudioMixer/Private/AudioMixerSourceManager.h b/Engine/Source/Runtime/AudioMixer/Private/AudioMixerSourceManager.h
--- a/Engine/Source/Runtime/AudioMixer/Private/AudioMixerSourceManager.h
+++ b/Engine/Source/Runtime/AudioMixer/Private/AudioMixerSourceManager.h
@@ -328,6 +328,11 @@ namespace Audio
 			PendingAudioBusConnections.Enqueue(FPendingAudioBusConnection{ FPendingAudioBusConnection::FPatchVariant(TInPlaceType<FPatchOutputStrongPtr>(), MoveTemp(PatchOutputStrongPtr)), MoveTemp(AudioBusKey), NumChannels, bIsAutomatic, WaveInstancePlayOrder });
 		}
 
+		bool UseRenderScheduler() const
+		{
+			return MixerDevice->UseRenderScheduler();
+		}
+
 		FAudioRenderScheduler& GetRenderScheduler() const
 		{
 			return MixerDevice->GetRenderScheduler();
diff --git a/Engine/Source/Runtime/AudioMixer/Public/AudioMixerDevice.h b/Engine/Source/Runtime/AudioMixer/Public/AudioMixerDevice.h
--- a/Engine/Source/Runtime/AudioMixer/Public/AudioMixerDevice.h
+++ b/Engine/Source/Runtime/AudioMixer/Public/AudioMixerDevice.h
@@ -366,6 +366,11 @@ namespace Audio
 		// Debug Commands
 		AUDIOMIXER_API void DrawSubmixes(FOutputDevice& InOutput, const TArray<FString>& InArgs) const;
 
+		bool UseRenderScheduler() const
+		{
+			return bUseRenderScheduler;
+		}
+
 		FAudioRenderScheduler& GetRenderScheduler() const
 		{
 			return *RenderScheduler;
@@ -509,6 +514,7 @@ namespace Audio
 
 		/** Scheduler responsible for launching rendering tasks */
 		TUniquePtr<FAudioRenderScheduler> RenderScheduler;
+		bool bUseRenderScheduler = true;
 
 		/** The mixer source manager. Should be cleaned up before the scheduler so dependencies can be removed correctly. */
 		TUniquePtr<FMixerSourceManager> SourceManager;
