[Subsonic] Generator sources can be played/stopped through subsonic actions, and get cleaned up when they're done playing.

Generator sources are a new experimental audio engine feature that allows an ISoundGenerator implementation to be played directly through the AudioMixer without going through the usual AudioComponent/SondBase/ActiveSound stack. To use them, set au.NumGeneratorSources to a nonzero number, then create a generator source with FMixerDevice::CreateGeneratorSource().

#tests Tested in PIE in BR_Audio
#rnx

#Synced-CL 51585013

#Synced-CL 51585013

diff --git a/Engine/Plugins/Experimental/Subsonic/Source/Subsonic/Private/CoreActions/SubsonicAction_GeneratorSource.cpp b/Engine/Plugins/Experimental/Subsonic/Source/Subsonic/Private/CoreActions/SubsonicAction_GeneratorSource.cpp
new file mode 100644
--- /dev/null
+++ b/Engine/Plugins/Experimental/Subsonic/Source/Subsonic/Private/CoreActions/SubsonicAction_GeneratorSource.cpp
@@ -0,0 +1,73 @@
+// Copyright Epic Games, Inc. All Rights Reserved.
+
+#include "CoreActions/SubsonicAction_GeneratorSource.h"
+
+#include "CoreEventSubscribers/SubsonicGeneratorSourceSubscriber.h"
+#include "IAudioMixerGeneratorSource.h"
+#include "SubsonicHandles.h"
+
+namespace UE::Subsonic
+{
+	void FSubsonicEventAction_GeneratorSourcePlay::Execute(const FSubsonicExecutor& InExecutor, const FActionHandle& InHandle) const
+	{
+		if (Sound == nullptr)
+		{
+			return;
+		}
+
+		const FActionScopeKey ScopeKey(InExecutor, InHandle);
+		if (USubsonicGeneratorSourceSubscriber* GeneratorSourceSubscriber = ScopeKey.FindDeviceSubsystem<USubsonicGeneratorSourceSubscriber>())
+		{
+			switch (Scope)
+			{
+			case ESubsonicExecutionScope::Global:
+				GeneratorSourceSubscriber->PlaySound(Name, *Sound);
+				break;
+
+			case ESubsonicExecutionScope::Collection:
+				GeneratorSourceSubscriber->PlaySound(InHandle.Event.Collection, Name, *Sound);
+				break;
+
+			case ESubsonicExecutionScope::Event:
+				GeneratorSourceSubscriber->PlaySound(InHandle.Event, Name, *Sound);
+				break;
+
+			case ESubsonicExecutionScope::Action:
+				GeneratorSourceSubscriber->PlaySound(InHandle, Name, *Sound);
+				break;
+
+			default:
+				checkNoEntry();
+			}
+		}
+	}
+
+	void FSubsonicEventAction_GeneratorSourceStop::Execute(const FSubsonicExecutor& InExecutor, const FActionHandle& InHandle) const
+	{
+		const FActionScopeKey ScopeKey(InExecutor, InHandle);
+		if (USubsonicGeneratorSourceSubscriber* GeneratorSourceSubscriber = ScopeKey.FindDeviceSubsystem<USubsonicGeneratorSourceSubscriber>())
+		{
+			switch (Scope)
+			{
+			case ESubsonicExecutionScope::Global:
+				GeneratorSourceSubscriber->StopSound(Name);
+				break;
+
+			case ESubsonicExecutionScope::Collection:
+				GeneratorSourceSubscriber->StopSound(InHandle.Event.Collection, Name);
+				break;
+
+			case ESubsonicExecutionScope::Event:
+				GeneratorSourceSubscriber->StopSound(InHandle.Event, Name);
+				break;
+
+			case ESubsonicExecutionScope::Action:
+				GeneratorSourceSubscriber->StopSound(InHandle, Name);
+				break;
+
+			default:
+				checkNoEntry();
+			}
+		}
+	}
+}
\ No newline at end of file
diff --git a/Engine/Plugins/Experimental/Subsonic/Source/Subsonic/Private/CoreActions/SubsonicAction_GeneratorSource.h b/Engine/Plugins/Experimental/Subsonic/Source/Subsonic/Private/CoreActions/SubsonicAction_GeneratorSource.h
new file mode 100644
--- /dev/null
+++ b/Engine/Plugins/Experimental/Subsonic/Source/Subsonic/Private/CoreActions/SubsonicAction_GeneratorSource.h
@@ -0,0 +1,57 @@
+// Copyright Epic Games, Inc. All Rights Reserved.
+#pragma once
+
+#include "Misc/CoreMiscDefines.h"
+#include "Sound/SoundWave.h"
+#include "SubsonicEventCollection.h"
+
+#include "SubsonicAction_GeneratorSource.generated.h"
+
+namespace UE::Subsonic
+{
+	USTRUCT(MinimalAPI, BlueprintType, DisplayName = "Play Sound (GeneratorSource)")
+		struct FSubsonicEventAction_GeneratorSourcePlay : public FSubsonicEventActionBase
+	{
+		GENERATED_BODY()
+
+		virtual ~FSubsonicEventAction_GeneratorSourcePlay() = default;
+
+		virtual void Execute(const FSubsonicExecutor& InExecutor, const FActionHandle& InHandle) const override;
+
+	public:
+		// Addressable name of Generator source instance.
+		UPROPERTY(EditAnywhere, Category = "Subsonic")
+		FName Name = "Sound01";
+
+		// Sound to play.
+		UPROPERTY(EditAnywhere, Category = "Subsonic")
+		TObjectPtr<USoundWave> Sound;
+
+		// Scope of source accessor.  If set to event, provided name is used
+		// to address the same source within the event, collection the same source
+		// within the collection, etc.
+		UPROPERTY(EditAnywhere, Category = "Subsonic")
+		ESubsonicExecutionScope Scope = ESubsonicExecutionScope::Collection;
+	};
+
+	USTRUCT(MinimalAPI, BlueprintType, DisplayName = "Stop Sound (GeneratorSource)")
+	struct FSubsonicEventAction_GeneratorSourceStop : public FSubsonicEventActionBase
+	{
+		GENERATED_BODY()
+
+		virtual ~FSubsonicEventAction_GeneratorSourceStop() = default;
+
+		virtual void Execute(const FSubsonicExecutor& InExecutor, const FActionHandle& InHandle) const override;
+
+	public:
+		// Addressable name of source instance.
+		UPROPERTY(EditAnywhere, Category = "Subsonic")
+		FName Name = "Sound01";
+
+		// Scope of source accessor.  If set to event, provided name is used
+		// to address the same source within the event, collection the same source
+		// within the collection, etc.
+		UPROPERTY(EditAnywhere, Category = "Subsonic")
+		ESubsonicExecutionScope Scope = ESubsonicExecutionScope::Collection;
+	};
+}
\ No newline at end of file
diff --git a/Engine/Plugins/Experimental/Subsonic/Source/Subsonic/Private/CoreEventSubscribers/SubsonicGeneratorSourceSubscriber.cpp b/Engine/Plugins/Experimental/Subsonic/Source/Subsonic/Private/CoreEventSubscribers/SubsonicGeneratorSourceSubscriber.cpp
new file mode 100644
--- /dev/null
+++ b/Engine/Plugins/Experimental/Subsonic/Source/Subsonic/Private/CoreEventSubscribers/SubsonicGeneratorSourceSubscriber.cpp
@@ -0,0 +1,200 @@
+// Copyright Epic Games, Inc. All Rights Reserved.
+
+#include "CoreEventSubscribers/SubsonicGeneratorSourceSubscriber.h"
+
+#include "Algo/Transform.h"
+#include "AudioDeviceManager.h"
+#include "AudioDevice.h"
+#include "AudioMixerDevice.h"
+#include "Engine/Engine.h"
+#include "IAudioMixerGeneratorSource.h"
+#include "SubsonicLog.h"
+#include "SubsonicSubsystem.h"
+#include "SubsonicWaveGenerator.h"
+#include "UObject/UObjectGlobals.h"
+
+
+namespace UE::Subsonic
+{
+	namespace SubscriberPrivate
+	{
+		template <typename HandleType>
+		Audio::IAudioMixerGeneratorSource* PlaySound(
+			USubsonicGeneratorSourceSubscriber& Subscriber,
+			USoundWave& Sound,
+			const HandleType& InHandle,
+			TMap<HandleType, FGeneratorSourceCache>& OutSourceCache,
+			FName Name,
+			bool bReleaseExisting)
+		{
+			Audio::IAudioMixerGeneratorSource* ReturnSource = nullptr;
+
+			if (FAudioDeviceManager* AudioDeviceManager = FAudioDeviceManager::Get())
+			{
+				if (Audio::FMixerDevice* AudioDevice = Subscriber.GetMutableMixerDevice())
+				{
+					FGeneratorSourceCache& Cache = OutSourceCache.FindOrAdd(InHandle);
+					FInstanceSources& InstanceSources = Cache.FindOrAdd(Name);
+
+					TSharedRef<FWaveGenerator> Generator = MakeShared<FWaveGenerator>(Sound.CreateSoundWaveProxy(), AudioDevice);
+					TUniquePtr<Audio::IAudioMixerGeneratorSource> NewSource = AudioDevice->CreateGeneratorSource(Generator);
+					if (NewSource.IsValid())
+					{
+						ReturnSource = NewSource.Get();
+						if (InstanceSources.ActiveSource != nullptr)
+						{
+							InstanceSources.ReleasedSources.Add(MoveTemp(InstanceSources.ActiveSource));
+						}
+
+						NewSource->Play();
+						InstanceSources.ActiveSource = MoveTemp(NewSource);
+					}
+				}
+			}
+
+			return ReturnSource;
+		}
+
+		template <typename HandleType>
+		void ClearStaleSounds(TMap<HandleType, FGeneratorSourceCache>& SourceCache)
+		{
+			TArray<HandleType> HandlesToClear;
+
+			for (TPair<HandleType, FGeneratorSourceCache>& CacheItem : SourceCache)
+			{
+				TArray<FName> NamesToClear;
+				for (TPair<FName, FInstanceSources>& SourcesItem : CacheItem.Value)
+				{
+					FInstanceSources& Instances = SourcesItem.Value;
+
+					Instances.ReleasedSources.RemoveAll([](TUniquePtr<Audio::IAudioMixerGeneratorSource>& Source)
+						{
+							return Source->IsDone();
+						});
+					if (Instances.ActiveSource != nullptr && Instances.ActiveSource->IsDone())
+					{
+						Instances.ActiveSource.Reset();
+					}
+					if (Instances.ActiveSource == nullptr && Instances.ReleasedSources.IsEmpty())
+					{
+						NamesToClear.Add(SourcesItem.Key);
+					}
+				}
+
+				for (FName Name : NamesToClear)
+				{
+					CacheItem.Value.Remove(Name);
+				}
+				if (CacheItem.Value.IsEmpty())
+				{
+					HandlesToClear.Add(CacheItem.Key);
+				}
+			}
+
+			for (HandleType Handle : HandlesToClear)
+			{
+				SourceCache.Remove(Handle);
+			}
+		}
+
+		template <typename HandleType>
+		void StopSound(
+			const HandleType& InHandle,
+			TMap<HandleType, FGeneratorSourceCache>& OutSourceCache,
+			FName Name)
+		{
+			FGeneratorSourceCache* Cache = OutSourceCache.Find(InHandle);
+			if (Cache != nullptr)
+			{
+				Cache->Remove(Name);
+			}
+		}
+	} // namespace SubscriberPrivate
+
+	USubsonicGeneratorSourceSubscriber::USubsonicGeneratorSourceSubscriber()
+		: FTickableGameObject(ETickableTickType::Never)
+	{
+	}
+
+	void USubsonicGeneratorSourceSubscriber::Initialize(FSubsystemCollectionBase& Collection)
+	{
+		Super::Initialize(Collection);
+		SetTickableTickType(ETickableTickType::Always);
+	}
+
+	void USubsonicGeneratorSourceSubscriber::Deinitialize()
+	{
+		SetTickableTickType(ETickableTickType::Never);
+
+		GlobalSources.Reset();
+		CollectionSources.Reset();
+		EventSources.Reset();
+		ActionSources.Reset();
+
+		Super::Deinitialize();
+	}
+
+	void USubsonicGeneratorSourceSubscriber::OnCollectionRegistered(const FCollectionHandle& InHandle)
+	{
+	}
+
+	void USubsonicGeneratorSourceSubscriber::OnCollectionUnregistered(const FCollectionHandle& InHandle)
+	{
+		CollectionSources.Remove(InHandle);
+	}
+
+	void USubsonicGeneratorSourceSubscriber::PlaySound(FName Name, USoundWave& Sound, bool bReleaseExisting)
+	{
+		const int32 FakeHandle = 0;
+		SubscriberPrivate::PlaySound(*this, Sound, FakeHandle, GlobalSources, Name, bReleaseExisting);
+	}
+
+	void USubsonicGeneratorSourceSubscriber::PlaySound(const FCollectionHandle& InHandle, FName Name, USoundWave& Sound, bool bReleaseExisting)
+	{
+		SubscriberPrivate::PlaySound(*this, Sound, InHandle, CollectionSources, Name, bReleaseExisting);
+	}
+
+	void USubsonicGeneratorSourceSubscriber::PlaySound(const FEventHandle& InHandle, FName Name, USoundWave& Sound, bool bReleaseExisting)
+	{
+		SubscriberPrivate::PlaySound(*this, Sound, InHandle, EventSources, Name, bReleaseExisting);
+	}
+
+	void USubsonicGeneratorSourceSubscriber::PlaySound(const FActionHandle& InHandle, FName Name, USoundWave& Sound, bool bReleaseExisting)
+	{
+		SubscriberPrivate::PlaySound(*this, Sound, InHandle, ActionSources, Name, bReleaseExisting);
+	}
+
+	void USubsonicGeneratorSourceSubscriber::StopSound(FName Name)
+	{
+		const int32 FakeHandle = 0;
+		SubscriberPrivate::StopSound(FakeHandle, GlobalSources, Name);
+	}
+
+	void USubsonicGeneratorSourceSubscriber::StopSound(const FCollectionHandle& InHandle, FName Name)
+	{
+		SubscriberPrivate::StopSound(InHandle, CollectionSources, Name);
+	}
+
+	void USubsonicGeneratorSourceSubscriber::StopSound(const FEventHandle& InHandle, FName Name)
+	{
+		SubscriberPrivate::StopSound(InHandle, EventSources, Name);
+	}
+
+	void USubsonicGeneratorSourceSubscriber::StopSound(const FActionHandle& InHandle, FName Name)
+	{
+		SubscriberPrivate::StopSound(InHandle, ActionSources, Name);
+	}
+
+	void USubsonicGeneratorSourceSubscriber::Tick(float DeltaTime)
+	{
+		SubscriberPrivate::ClearStaleSounds(GlobalSources);
+		SubscriberPrivate::ClearStaleSounds(CollectionSources);
+		SubscriberPrivate::ClearStaleSounds(EventSources);
+		SubscriberPrivate::ClearStaleSounds(ActionSources);
+	}
+
+	TStatId USubsonicGeneratorSourceSubscriber::GetStatId() const
+	{
+		RETURN_QUICK_DECLARE_CYCLE_STAT(USubsonicGeneratorSourceSubscriber, STATGROUP_Tickables);
+	}
+}
diff --git a/Engine/Plugins/Experimental/Subsonic/Source/Subsonic/Private/CoreEventSubscribers/SubsonicGeneratorSourceSubscriber.h b/Engine/Plugins/Experimental/Subsonic/Source/Subsonic/Private/CoreEventSubscribers/SubsonicGeneratorSourceSubscriber.h
new file mode 100644
--- /dev/null
+++ b/Engine/Plugins/Experimental/Subsonic/Source/Subsonic/Private/CoreEventSubscribers/SubsonicGeneratorSourceSubscriber.h
@@ -0,0 +1,69 @@
+// Copyright Epic Games, Inc. All Rights Reserved.
+#pragma once
+
+#include "CoreEventSubscribers/SubsonicEventSubscriberBase.h"
+#include "IAudioMixerGeneratorSource.h"
+#include "SubsonicEventSystemInterface.h"
+#include "SubsonicHandles.h"
+#include "Subsystems/AudioEngineSubsystem.h"
+#include "Tickable.h"
+
+#include "SubsonicGeneratorSourceSubscriber.generated.h"
+
+// Forward Declarations
+class USoundWave;
+
+namespace UE::Subsonic
+{
+	struct FInstanceSources
+	{
+		// Currently addressable generator source via handle
+		TUniquePtr<Audio::IAudioMixerGeneratorSource> ActiveSource;
+
+		// Sources no longer addressable from actions but actively playing.
+		TArray<TUniquePtr<Audio::IAudioMixerGeneratorSource>> ReleasedSources;
+	};
+
+	using FGeneratorSourceCache = TMap<FName, FInstanceSources>;
+
+	/**
+	 * Interface to GeneratorSource playback Subsonic implementation.
+	 */
+	UCLASS(MinimalAPI, BlueprintType)
+		class USubsonicGeneratorSourceSubscriber final : public USubsonicEventSubscriberBase, public FTickableGameObject
+	{
+		GENERATED_BODY()
+
+	public:
+		USubsonicGeneratorSourceSubscriber();
+
+		//~ Begin USubsystem interface
+		virtual void Initialize(FSubsystemCollectionBase& Collection) override;
+		virtual void Deinitialize() override;
+		//~ End USubsystem interface
+
+		virtual void OnCollectionRegistered(const FCollectionHandle& InEvent) override;
+		virtual void OnCollectionUnregistered(const FCollectionHandle& InEvent) override;
+
+		void PlaySound(FName Name, USoundWave& Sound, bool bReleaseExisting = true);
+		void PlaySound(const FCollectionHandle& InHandle, FName Name, USoundWave& Sound, bool bReleaseExisting = true);
+		void PlaySound(const FEventHandle& InHandle, FName Name, USoundWave& Sound, bool bReleaseExisting = true);
+		void PlaySound(const FActionHandle& InHandle, FName Name, USoundWave& Sound, bool bReleaseExisting = true);
+
+		void StopSound(FName Name);
+		void StopSound(const FCollectionHandle& InHandle, FName Name);
+		void StopSound(const FEventHandle& InHandle, FName Name);
+		void StopSound(const FActionHandle& InHandle, FName Name);
+
+		// FTickableGameObject interface
+		virtual void Tick(float DeltaTime) override;
+		virtual TStatId GetStatId() const override;
+	private:
+		TMap<int32, FGeneratorSourceCache> GlobalSources;
+		TMap<FCollectionHandle, FGeneratorSourceCache> CollectionSources;
+		TMap<FEventHandle, FGeneratorSourceCache> EventSources;
+		TMap<FActionHandle, FGeneratorSourceCache> ActionSources;
+	};
+}
+
+#undef UE_API
diff --git a/Engine/Plugins/Experimental/Subsonic/Source/Subsonic/Private/CoreEventSubscribers/SubsonicWaveGenerator.cpp b/Engine/Plugins/Experimental/Subsonic/Source/Subsonic/Private/CoreEventSubscribers/SubsonicWaveGenerator.cpp
new file mode 100644
--- /dev/null
+++ b/Engine/Plugins/Experimental/Subsonic/Source/Subsonic/Private/CoreEventSubscribers/SubsonicWaveGenerator.cpp
@@ -0,0 +1,74 @@
+// Copyright Epic Games, Inc. All Rights Reserved.
+
+#include "SubsonicWaveGenerator.h"
+
+#include "AudioMixerDevice.h"
+#include "DSP/MultichannelBuffer.h"
+#include "Sound/SoundGenerator.h"
+#include "Sound/SoundWave.h"
+#include "Sound/SoundWaveProxyPlayer.h"
+
+namespace UE::Subsonic
+{
+	namespace WaveGeneratorPrivate
+	{
+		void Interleave(const Audio::FMultichannelBufferView& Input, float* Output, const int32 NumOutputSamples)
+		{
+			check(Input.IsEmpty() || NumOutputSamples == Input.Num() * Input[0].Num());
+			for (int32 ChannelIndex = 0; ChannelIndex < Input.Num(); ++ChannelIndex)
+			{
+				float* OutputSample = Output + ChannelIndex;
+				TArrayView<float> InputChannel = Input[ChannelIndex];
+				for (float Sample : InputChannel)
+				{
+					*OutputSample = Sample;
+					OutputSample += Input.Num();
+				}
+			}
+		}
+	}
+
+	FWaveGenerator::FWaveGenerator(const TSharedPtr<FSoundWaveProxy>& SoundWaveProxy, Audio::FMixerDevice* MixerDevice)
+	{
+		const float SampleRate = MixerDevice->GetSampleRate();
+		check(SoundWaveProxy.IsValid());
+		FSoundWaveProxyPlayer::FSettings Settings(SampleRate, SoundWaveProxy->GetNumChannels());
+		Player = FSoundWaveProxyPlayer::Create(Settings);
+		Player->SetSoundWave(SoundWaveProxy);
+	}
+
+	FWaveGenerator::~FWaveGenerator()
+	{
+	}
+
+	int32 FWaveGenerator::OnGenerateAudio(float* OutAudio, int32 InNumSamples)
+	{
+		const int32 NumChannels = GetNumChannels();
+		check(InNumSamples >= 0 && NumChannels > 0 && InNumSamples % NumChannels == 0);
+		const int32 NumFrames = InNumSamples / NumChannels;
+
+		Audio::FMultichannelBuffer TempBuffer;
+		for (int32 i = 0; i < NumChannels; ++i)
+		{
+			Audio::FAlignedFloatBuffer ChannelBuffer;
+			ChannelBuffer.SetNumUninitialized(NumFrames);
+			TempBuffer.Add(MoveTemp(ChannelBuffer));
+		}
+
+		Audio::FMultichannelBufferView TempBufferView = Audio::MakeMultichannelBufferView(TempBuffer);
+		Player->RenderMultiChannelAudio(TempBufferView);
+		WaveGeneratorPrivate::Interleave(TempBufferView, OutAudio, InNumSamples);
+
+		return InNumSamples;
+	}
+
+	int32 FWaveGenerator::GetNumChannels() const
+	{
+		return Player->GetSourceNumChannels();
+	}
+
+	bool FWaveGenerator::IsFinished() const
+	{
+		return Player->IsFinished();
+	}
+}
\ No newline at end of file
diff --git a/Engine/Plugins/Experimental/Subsonic/Source/Subsonic/Private/CoreEventSubscribers/SubsonicWaveGenerator.h b/Engine/Plugins/Experimental/Subsonic/Source/Subsonic/Private/CoreEventSubscribers/SubsonicWaveGenerator.h
new file mode 100644
--- /dev/null
+++ b/Engine/Plugins/Experimental/Subsonic/Source/Subsonic/Private/CoreEventSubscribers/SubsonicWaveGenerator.h
@@ -0,0 +1,34 @@
+// Copyright Epic Games, Inc. All Rights Reserved.
+
+#pragma once
+
+#include "Sound/SoundGenerator.h"
+
+namespace Audio
+{
+	class FMixerDevice;
+}
+
+class FSoundWaveProxy;
+class FSoundWaveProxyPlayer;
+
+namespace UE::Subsonic
+{
+
+	// A sound generator that can be used to play back a wave file
+
+	class FWaveGenerator : public ISoundGenerator
+	{
+	public:
+		explicit FWaveGenerator(const TSharedPtr<FSoundWaveProxy>& SoundWaveProxy, Audio::FMixerDevice* MixerDevice);
+		virtual ~FWaveGenerator() override;
+
+		virtual int32 OnGenerateAudio(float* OutAudio, int32 InNumSamples) override;
+		virtual int32 GetNumChannels() const override;
+
+		virtual bool IsFinished() const override;
+	
+	protected:
+		TUniquePtr<FSoundWaveProxyPlayer> Player;
+	};
+}
\ No newline at end of file
diff --git a/Engine/Plugins/Experimental/Subsonic/Source/Subsonic/Subsonic.Build.cs b/Engine/Plugins/Experimental/Subsonic/Source/Subsonic/Subsonic.Build.cs
--- a/Engine/Plugins/Experimental/Subsonic/Source/Subsonic/Subsonic.Build.cs
+++ b/Engine/Plugins/Experimental/Subsonic/Source/Subsonic/Subsonic.Build.cs
@@ -22,8 +22,10 @@ public class Subsonic : ModuleRules
 			new string[]
 			{
 				"AudioExtensions",
+				"AudioMixer",
 				"CoreUObject",
-				"Engine"
+				"Engine",
+				"SignalProcessing",
 			}
 		);
 
diff --git a/Engine/Plugins/Media/PixelStreaming/Source/PixelStreaming/Public/PixelStreamingAudioComponent.h b/Engine/Plugins/Media/PixelStreaming/Source/PixelStreaming/Public/PixelStreamingAudioComponent.h
--- a/Engine/Plugins/Media/PixelStreaming/Source/PixelStreaming/Public/PixelStreamingAudioComponent.h
+++ b/Engine/Plugins/Media/PixelStreaming/Source/PixelStreaming/Public/PixelStreamingAudioComponent.h
@@ -36,7 +36,7 @@ public:
 	void AddAudio(const int16_t* AudioData, int InSampleRate, size_t NChannels, size_t NFrames);
 
 	int32 GetSampleRate() { return Params.SampleRate; }
-	int32 GetNumChannels() { return Params.NumChannels; }
+	virtual int32 GetNumChannels() const override { return Params.NumChannels; }
 	void EmptyBuffers();
 	void SetParameters(const FSoundGeneratorInitParams& InitParams);
 
diff --git a/Engine/Plugins/Runtime/Metasound/Source/MetasoundGenerator/Public/MetasoundGenerator.h b/Engine/Plugins/Runtime/Metasound/Source/MetasoundGenerator/Public/MetasoundGenerator.h
--- a/Engine/Plugins/Runtime/Metasound/Source/MetasoundGenerator/Public/MetasoundGenerator.h
+++ b/Engine/Plugins/Runtime/Metasound/Source/MetasoundGenerator/Public/MetasoundGenerator.h
@@ -262,7 +262,7 @@ namespace Metasound
 		FOnOutputChanged OnOutputChanged;
 		
 		/** Return the number of audio channels. */
-		UE_API int32 GetNumChannels() const;
+		UE_API virtual int32 GetNumChannels() const override;
 
 		//~ Begin FSoundGenerator
 		UE_API virtual int32 OnGenerateAudio(float* OutAudio, int32 NumSamples) override;
diff --git a/Engine/Plugins/Runtime/Synthesis/Source/Synthesis/Classes/SynthComponents/SynthComponentToneGenerator.h b/Engine/Plugins/Runtime/Synthesis/Source/Synthesis/Classes/SynthComponents/SynthComponentToneGenerator.h
--- a/Engine/Plugins/Runtime/Synthesis/Source/Synthesis/Classes/SynthComponents/SynthComponentToneGenerator.h
+++ b/Engine/Plugins/Runtime/Synthesis/Source/Synthesis/Classes/SynthComponents/SynthComponentToneGenerator.h
@@ -19,7 +19,7 @@ public:
 	virtual ~FToneGenerator() = default;
 
 	//~ Begin FSoundGenerator 
-	virtual int32 GetNumChannels() { return NumChannels; };
+	virtual int32 GetNumChannels() const override { return NumChannels; };
 	virtual int32 OnGenerateAudio(float* OutAudio, int32 NumSamples) override;
 	//~ End FSoundGenerator
 
diff --git a/Engine/Source/Runtime/AudioMixer/Private/AudioMixerDevice.cpp b/Engine/Source/Runtime/AudioMixer/Private/AudioMixerDevice.cpp
--- a/Engine/Source/Runtime/AudioMixer/Private/AudioMixerDevice.cpp
+++ b/Engine/Source/Runtime/AudioMixer/Private/AudioMixerDevice.cpp
@@ -7,6 +7,7 @@
 #include "AudioAnalytics.h"
 #include "AudioBusSubsystem.h"
 #include "AudioDeviceNotificationSubsystem.h"
+#include "AudioMixerGeneratorSource.h"
 #include "AudioMixerSource.h"
 #include "AudioMixerSourceManager.h"
 #include "AudioMixerSourceDecode.h"
@@ -135,6 +136,13 @@ FAutoConsoleVariableRef CVarSubmixBufferListenerForceRegistration(
 	TEXT("If true, will force register the referenced submix when Registering a Submix Buffer Listener."),
 	ECVF_Default);
 
+static int32 AudioMixerGeneratorSourcesCVar = 0;
+FAutoConsoleVariableRef CVarAudioMixerGeneratorSources(
+	TEXT("au.NumGeneratorSources"),
+	AudioMixerGeneratorSourcesCVar,
+	TEXT("Sets number of generator sources available in the AudioMixer."),
+	ECVF_Default);
+
 #if UE_AUDIO_PROFILERTRACE_ENABLED
 
 #if !WITH_EDITOR
@@ -1132,6 +1140,7 @@ namespace Audio
 
 				FSourceManagerInitParams SourceManagerInitParams;
 				SourceManagerInitParams.NumSources = GetMaxSources();
+				SourceManagerInitParams.NumGeneratorSources = AudioMixerGeneratorSourcesCVar;
 
 				AudioClock = 0.0;
 				AudioClockDelta = (double)OpenStreamParams.NumFrames / OpenStreamParams.SampleRate;
@@ -3703,6 +3712,11 @@ namespace Audio
 		return Audio::KickQueuedTasks(QueueId);
 	}
 
+	TUniquePtr<IAudioMixerGeneratorSource> FMixerDevice::CreateGeneratorSource(ISoundGeneratorRef InGenerator)
+	{
+		return MakeUnique<FAudioMixerGeneratorSource>(*this, InGenerator);
+	}
+
 	double FAudioClockTimingData::GetInterpolatedAudioClock(const double InAudioClock, const double InAudioClockDelta) const
 	{
 		if (UpdateTime > 0.0)
diff --git a/Engine/Source/Runtime/AudioMixer/Private/AudioMixerGeneratorSource.cpp b/Engine/Source/Runtime/AudioMixer/Private/AudioMixerGeneratorSource.cpp
new file mode 100644
--- /dev/null
+++ b/Engine/Source/Runtime/AudioMixer/Private/AudioMixerGeneratorSource.cpp
@@ -0,0 +1,49 @@
+// Copyright Epic Games, Inc. All Rights Reserved.
+
+#include "AudioMixerGeneratorSource.h"
+
+#include "AudioMixerDevice.h"
+#include "AudioMixerSourceManager.h"
+
+namespace Audio
+{
+	FAudioMixerGeneratorSource::FAudioMixerGeneratorSource(FMixerDevice& InMixerDevice, ISoundGeneratorRef InGenerator)
+		: MixerDevice(InMixerDevice)
+	{
+		SourceId = MixerDevice.GetSourceManager()->AllocateGeneratorSource(InGenerator);
+	}
+
+	FAudioMixerGeneratorSource::~FAudioMixerGeneratorSource()
+	{
+		if (SourceId != INDEX_NONE)
+		{
+			MixerDevice.GetSourceManager()->FreeGeneratorSource(SourceId);
+		}
+	}
+
+	void FAudioMixerGeneratorSource::Play()
+	{
+		if (SourceId != INDEX_NONE)
+		{
+			MixerDevice.GetSourceManager()->PlayGeneratorSource(SourceId);
+		}
+	}
+
+	void FAudioMixerGeneratorSource::Stop()
+	{
+		if (SourceId != INDEX_NONE)
+		{
+			MixerDevice.GetSourceManager()->StopGeneratorSource(SourceId);
+		}
+	}
+
+	bool FAudioMixerGeneratorSource::IsDone()
+	{
+		if (SourceId != INDEX_NONE)
+		{
+			return MixerDevice.GetSourceManager()->IsGeneratorSourceDone(SourceId);
+		}
+
+		return true;
+	}
+}
\ No newline at end of file
diff --git a/Engine/Source/Runtime/AudioMixer/Private/AudioMixerGeneratorSource.h b/Engine/Source/Runtime/AudioMixer/Private/AudioMixerGeneratorSource.h
new file mode 100644
--- /dev/null
+++ b/Engine/Source/Runtime/AudioMixer/Private/AudioMixerGeneratorSource.h
@@ -0,0 +1,28 @@
+// Copyright Epic Games, Inc. All Rights Reserved.
+
+#pragma once
+
+#include "IAudioMixerGeneratorSource.h"
+#include "Sound/SoundGenerator.h"
+
+namespace Audio
+{
+	class FMixerDevice;
+
+	class FAudioMixerGeneratorSource : public IAudioMixerGeneratorSource
+	{
+		UE_NONCOPYABLE(FAudioMixerGeneratorSource);
+
+	public:
+		FAudioMixerGeneratorSource(FMixerDevice& InMixerDevice, ISoundGeneratorRef InGenerator);
+		virtual ~FAudioMixerGeneratorSource() override;
+
+		virtual void Play() override;
+		virtual void Stop() override;
+
+		virtual bool IsDone() override;
+	private:
+		FMixerDevice& MixerDevice;
+		int32 SourceId = INDEX_NONE;
+	};
+}
diff --git a/Engine/Source/Runtime/AudioMixer/Private/AudioMixerSourceManager.cpp b/Engine/Source/Runtime/AudioMixer/Private/AudioMixerSourceManager.cpp
--- a/Engine/Source/Runtime/AudioMixer/Private/AudioMixerSourceManager.cpp
+++ b/Engine/Source/Runtime/AudioMixer/Private/AudioMixerSourceManager.cpp
@@ -410,6 +410,8 @@ namespace Audio
 		: MixerDevice(InMixerDevice)
 		, NumActiveSources(0)
 		, NumTotalSources(0)
+		, NumLegacySources(0)
+		, NumGeneratorSources(0)
 		, NumOutputFrames(0)
 		, NumOutputSamples(0)
 		, bInitialized(false)
@@ -663,12 +665,16 @@ namespace Audio
 
 		AUDIO_MIXER_CHECK(MixerDevice->GetSampleRate() > 0);
 
-		NumTotalSources = InitParams.NumSources;
+		NumLegacySources = InitParams.NumSources;
+		NumGeneratorSources = InitParams.NumGeneratorSources;
+		NumTotalSources = NumLegacySources + NumGeneratorSources;
 
 		NumOutputFrames = MixerDevice->PlatformSettings.CallbackBufferFrameSize;
 		NumOutputSamples = NumOutputFrames * MixerDevice->GetNumDeviceChannels();
+		
+		GeneratorSourceState.AddDefaulted(NumGeneratorSources);
 
-		MixerSources.Init(nullptr, NumTotalSources);
+		MixerSources.Init(nullptr, NumLegacySources);
 
 		// Populate output sources array with default data
 		SourceSubmixOutputBuffers.Reset();
@@ -757,18 +763,18 @@ namespace Audio
 			SourceInfo.NumInputFrames = 0;
 		}
 		
-		GameThreadInfo.bIsBusy.AddDefaulted(NumTotalSources);
-		GameThreadInfo.bNeedsSpeakerMap.AddDefaulted(NumTotalSources);
-		GameThreadInfo.bIsDebugMode.AddDefaulted(NumTotalSources);
-		GameThreadInfo.bIsUsingHRTFSpatializer.AddDefaulted(NumTotalSources);
+		GameThreadInfo.bIsBusy.AddDefaulted(NumLegacySources);
+		GameThreadInfo.bNeedsSpeakerMap.AddDefaulted(NumLegacySources);
+		GameThreadInfo.bIsDebugMode.AddDefaulted(NumLegacySources);
+		GameThreadInfo.bIsUsingHRTFSpatializer.AddDefaulted(NumLegacySources);
 #if ENABLE_AUDIO_DEBUG
-		GameThreadInfo.CPUCoreUtilization.AddZeroed(NumTotalSources);
+		GameThreadInfo.CPUCoreUtilization.AddZeroed(NumLegacySources);
 #endif // if ENABLE_AUDIO_DEBUG
 
-		GameThreadInfo.RelativeRenderCost.Reset(NumTotalSources);
-		GameThreadInfo.FreeSourceIndices.Reset(NumTotalSources);
-		GameThreadInfo.ModulationVolume.Reset(NumTotalSources);
-		for (int32 i = NumTotalSources - 1; i >= 0; --i)
+		GameThreadInfo.RelativeRenderCost.Reset(NumLegacySources);
+		GameThreadInfo.FreeSourceIndices.Reset(NumLegacySources);
+		GameThreadInfo.ModulationVolume.Reset(NumLegacySources);
+		for (int32 i = NumLegacySources - 1; i >= 0; --i)
 		{
 			GameThreadInfo.RelativeRenderCost.Add(1.0f);
 			GameThreadInfo.FreeSourceIndices.Add(i);
@@ -860,7 +866,7 @@ namespace Audio
 		AUDIO_MIXER_CHECK_GAME_THREAD(MixerDevice);
 
 #if VALIDATE_SOURCE_MIXER_STATE
-		for (int32 i = 0; i < NumTotalSources; ++i)
+		for (int32 i = 0; i < NumLegacySources; ++i)
 		{
 			if (!GameThreadInfo.bIsBusy[i])
 			{
@@ -940,7 +946,7 @@ namespace Audio
 		AUDIO_MIXER_CHECK(SourceId < NumTotalSources);
 		AUDIO_MIXER_CHECK(bInitialized);
 
-		if (MixerSources[SourceId] == nullptr)
+		if (SourceId < NumLegacySources && MixerSources[SourceId] == nullptr)
 		{
 			UE_LOG(LogAudioMixer, Warning, TEXT("Ignoring double release of SourceId: %i"), SourceId);
 			return;
@@ -950,6 +956,12 @@ namespace Audio
 		
 		FSourceInfo& SourceInfo = SourceInfos[SourceId];
 
+		if (UseRenderScheduler())
+		{
+			check(SourceId >= NumLegacySources || SourceInfo.MixerSourceBuffer.IsValid());
+			GetRenderScheduler().RemoveStep(FAudioRenderStepId::FromWaveInstancePlayOrder(SourceInfo.PlayOrder));
+		}
+
 #if AUDIO_MIXER_ENABLE_DEBUG_MODE
 		if (SourceInfo.bIsDebugMode)
 		{
@@ -1023,7 +1035,7 @@ namespace Audio
 			FMixerSubmixPtr SubmixPtr = SubmixSendItem.Submix.Pin();
 			if (SubmixPtr.IsValid())
 			{
-				SubmixPtr->RemoveSourceVoice(MixerSources[SourceId]);
+				SubmixPtr->RemoveSourceById(SourceId);
 			}
 		}
 		SourceInfo.SubmixSends.Reset();
@@ -1060,8 +1072,11 @@ namespace Audio
 		SourceInfo.bEffectTailsDone = true;
 
 		// Release the source voice back to the mixer device. This is pooled.
-		MixerDevice->ReleaseMixerSourceVoice(MixerSources[SourceId]);
-		MixerSources[SourceId] = nullptr;
+		if (SourceId < NumLegacySources)
+		{
+			MixerDevice->ReleaseMixerSourceVoice(MixerSources[SourceId]);
+			MixerSources[SourceId] = nullptr;
+		}
 
 		// Reset all state and data
 		SourceInfo.PitchSourceParam.Init();
@@ -1133,7 +1148,10 @@ namespace Audio
 		SourceInfo.NumInputChannels = 0;
 		SourceInfo.NumPostEffectChannels = 0;
 
-		GameThreadInfo.bNeedsSpeakerMap[SourceId] = false;
+		if (SourceId < NumLegacySources)
+		{
+			GameThreadInfo.bNeedsSpeakerMap[SourceId] = false;
+		}
 	}
 
 	void FMixerSourceManager::BuildSourceEffectChain(const int32 SourceId, FSoundEffectSourceInitData& InitData, const TArray<FSourceEffectChainEntry>& InSourceEffectChain, TArray<TSoundEffectSourcePtr>& OutSourceEffects)
@@ -1195,11 +1213,11 @@ namespace Audio
 		{
 			OutSourceId = GameThreadInfo.FreeSourceIndices.Pop();
 
-			AUDIO_MIXER_CHECK(OutSourceId < NumTotalSources);
+			AUDIO_MIXER_CHECK(OutSourceId < NumLegacySources);
 			AUDIO_MIXER_CHECK(!GameThreadInfo.bIsBusy[OutSourceId]);
 
 			AUDIO_MIXER_CHECK(!GameThreadInfo.bIsDebugMode[OutSourceId]);
-			AUDIO_MIXER_CHECK(NumActiveSources < NumTotalSources);
+			AUDIO_MIXER_CHECK(NumActiveSources < NumLegacySources);
 			++NumActiveSources;
 
 			GameThreadInfo.bIsBusy[OutSourceId] = true;
@@ -1221,7 +1239,7 @@ namespace Audio
 
 	void FMixerSourceManager::InitSource(const int32 SourceId, const FMixerSourceVoiceInitParams& InitParams)
 	{
-		AUDIO_MIXER_CHECK(SourceId < NumTotalSources);
+		AUDIO_MIXER_CHECK(SourceId < NumLegacySources);
 		AUDIO_MIXER_CHECK(GameThreadInfo.bIsBusy[SourceId]);
 		AUDIO_MIXER_CHECK(!GameThreadInfo.bIsDebugMode[SourceId]);
 		AUDIO_MIXER_CHECK(InitParams.SourceListener != nullptr);
@@ -1506,7 +1524,7 @@ namespace Audio
 						SourceInfo.bHasPreDistanceAttenuationSend = true;
 					}
 
-					SubmixPtr->AddOrSetSourceVoice(InitParams.SourceVoice, MixerSubmixSend.SendLevel, MixerSubmixSend.SubmixSendStage);
+					SubmixPtr->AddOrSetSourceById(SourceId, MixerSubmixSend.SendLevel, MixerSubmixSend.SubmixSendStage);
 					
 					if (SubmixPtr->IsSoundfieldSubmix())
 					{
@@ -1569,13 +1587,6 @@ namespace Audio
 		{
 			AUDIO_MIXER_CHECK_AUDIO_PLAT_THREAD(MixerDevice);
 
-			if (UseRenderScheduler())
-			{
-				FSourceInfo& SourceInfo = SourceInfos[SourceId];
-				check(SourceInfo.MixerSourceBuffer.IsValid());
-				GetRenderScheduler().RemoveStep(FAudioRenderStepId::FromWaveInstancePlayOrder(SourceInfo.PlayOrder));
-			}
-
 			ReleaseSource(SourceId);
 		});
 	}
@@ -1733,7 +1744,7 @@ namespace Audio
 
 	void FMixerSourceManager::Play(const int32 SourceId)
 	{
-		AUDIO_MIXER_CHECK(SourceId < NumTotalSources);
+		AUDIO_MIXER_CHECK(SourceId < NumLegacySources);
 		AUDIO_MIXER_CHECK(GameThreadInfo.bIsBusy[SourceId]);
 		AUDIO_MIXER_CHECK_GAME_THREAD(MixerDevice);
 
@@ -1769,21 +1780,8 @@ namespace Audio
 			StopInternal(SourceId);
 
 			// Verify we have a reasonable Source
-			AUDIO_MIXER_CHECK(SourceId < NumTotalSources);
-			FSourceInfo& SourceInfo = SourceInfos[SourceId];
-
-			//Update game thread state
-			SourceInfo.bIsDone = true;
-
-			// Notify that we're now done with this source
-			if (SourceInfo.SourceListener)
-			{
-				SourceInfo.SourceListener->OnDone();
-			}
-			if (SourceInfo.AudioLink)
-			{
-				SourceInfo.AudioLink->OnSourceDone(SourceId);
-			}
+			AUDIO_MIXER_CHECK(SourceId < NumLegacySources);
+			SetSourceDone(SourceId);
 		}
 	}
 
@@ -1794,7 +1792,7 @@ namespace Audio
 			return;
 		}
 
-		AUDIO_MIXER_CHECK(SourceId < NumTotalSources);
+		AUDIO_MIXER_CHECK(SourceId < NumLegacySources);
 
 		//Assert that we are being called from the GameThread and the
 		//source isn't busy.  Then call StopInternal() in a thread command
@@ -1830,7 +1828,7 @@ namespace Audio
 
 	void FMixerSourceManager::StopFade(const int32 SourceId, const int32 NumFrames)
 	{
-		AUDIO_MIXER_CHECK(SourceId < NumTotalSources);
+		AUDIO_MIXER_CHECK(SourceId < NumLegacySources);
 		AUDIO_MIXER_CHECK(GameThreadInfo.bIsBusy[SourceId]);
 		AUDIO_MIXER_CHECK(NumFrames > 0);
 		AUDIO_MIXER_CHECK_GAME_THREAD(MixerDevice);
@@ -1876,7 +1874,7 @@ namespace Audio
 	
 	void FMixerSourceManager::Pause(const int32 SourceId)
 	{
-		AUDIO_MIXER_CHECK(SourceId < NumTotalSources);
+		AUDIO_MIXER_CHECK(SourceId < NumLegacySources);
 		AUDIO_MIXER_CHECK(GameThreadInfo.bIsBusy[SourceId]);
 		AUDIO_MIXER_CHECK_GAME_THREAD(MixerDevice);
 
@@ -1893,7 +1891,7 @@ namespace Audio
 
 	void FMixerSourceManager::SetPitch(const int32 SourceId, const float Pitch)
 	{
-		AUDIO_MIXER_CHECK(SourceId < NumTotalSources);
+		AUDIO_MIXER_CHECK(SourceId < NumLegacySources);
 		AUDIO_MIXER_CHECK(GameThreadInfo.bIsBusy[SourceId]);
 
 		AudioMixerThreadCommand([this, SourceId, Pitch]()
@@ -1907,7 +1905,7 @@ namespace Audio
 
 	void FMixerSourceManager::SetVolume(const int32 SourceId, const float Volume)
 	{
-		AUDIO_MIXER_CHECK(SourceId < NumTotalSources);
+		AUDIO_MIXER_CHECK(SourceId < NumLegacySources);
 		AUDIO_MIXER_CHECK(GameThreadInfo.bIsBusy[SourceId]);
 		AUDIO_MIXER_CHECK_GAME_THREAD(MixerDevice);
 
@@ -1934,7 +1932,7 @@ namespace Audio
 
 	void FMixerSourceManager::SetDistanceAttenuation(const int32 SourceId, const float DistanceAttenuation)
 	{
-		AUDIO_MIXER_CHECK(SourceId < NumTotalSources);
+		AUDIO_MIXER_CHECK(SourceId < NumLegacySources);
 		AUDIO_MIXER_CHECK(GameThreadInfo.bIsBusy[SourceId]);
 		AUDIO_MIXER_CHECK_GAME_THREAD(MixerDevice);
 
@@ -1955,7 +1953,7 @@ namespace Audio
 
 	void FMixerSourceManager::SetSpatializationParams(const int32 SourceId, const FSpatializationParams& InParams)
 	{
-		AUDIO_MIXER_CHECK(SourceId < NumTotalSources);
+		AUDIO_MIXER_CHECK(SourceId < NumLegacySources);
 		AUDIO_MIXER_CHECK(GameThreadInfo.bIsBusy[SourceId]);
 		AUDIO_MIXER_CHECK_GAME_THREAD(MixerDevice);
 
@@ -1969,7 +1967,7 @@ namespace Audio
 
 	void FMixerSourceManager::SetChannelMap(const int32 SourceId, const uint32 NumInputChannels, const Audio::FAlignedFloatBuffer& ChannelMap, const bool bInIs3D, const bool bInIsCenterChannelOnly)
 	{
-		AUDIO_MIXER_CHECK(SourceId < NumTotalSources);
+		AUDIO_MIXER_CHECK(SourceId < NumLegacySources);
 		AUDIO_MIXER_CHECK(GameThreadInfo.bIsBusy[SourceId]);
 		AUDIO_MIXER_CHECK_GAME_THREAD(MixerDevice);
 
@@ -2000,7 +1998,7 @@ namespace Audio
 
 	void FMixerSourceManager::SetLPFFrequency(const int32 SourceId, const float InLPFFrequency)
 	{
-		AUDIO_MIXER_CHECK(SourceId < NumTotalSources);
+		AUDIO_MIXER_CHECK(SourceId < NumLegacySources);
 		AUDIO_MIXER_CHECK(GameThreadInfo.bIsBusy[SourceId]);
 		AUDIO_MIXER_CHECK_GAME_THREAD(MixerDevice);
 
@@ -2020,7 +2018,7 @@ namespace Audio
 
 	void FMixerSourceManager::SetHPFFrequency(const int32 SourceId, const float InHPFFrequency)
 	{
-		AUDIO_MIXER_CHECK(SourceId < NumTotalSources);
+		AUDIO_MIXER_CHECK(SourceId < NumLegacySources);
 		AUDIO_MIXER_CHECK(GameThreadInfo.bIsBusy[SourceId]);
 		AUDIO_MIXER_CHECK_GAME_THREAD(MixerDevice);
 
@@ -2039,7 +2037,7 @@ namespace Audio
 
 	void FMixerSourceManager::SetModLPFFrequency(const int32 SourceId, const float InLPFFrequency)
 	{
-		AUDIO_MIXER_CHECK(SourceId < NumTotalSources);
+		AUDIO_MIXER_CHECK(SourceId < NumLegacySources);
 		AUDIO_MIXER_CHECK(GameThreadInfo.bIsBusy[SourceId]);
 		AUDIO_MIXER_CHECK_GAME_THREAD(MixerDevice);
 
@@ -2055,7 +2053,7 @@ namespace Audio
 
 	void FMixerSourceManager::SetModHPFFrequency(const int32 SourceId, const float InHPFFrequency)
 	{
-		AUDIO_MIXER_CHECK(SourceId < NumTotalSources);
+		AUDIO_MIXER_CHECK(SourceId < NumLegacySources);
 		AUDIO_MIXER_CHECK(GameThreadInfo.bIsBusy[SourceId]);
 		AUDIO_MIXER_CHECK_GAME_THREAD(MixerDevice);
 
@@ -2111,7 +2109,7 @@ namespace Audio
 
 	void FMixerSourceManager::SetSourceBufferListener(const int32 SourceId, FSharedISourceBufferListenerPtr& InSourceBufferListener, bool InShouldSourceBufferListenerZeroBuffer)
 	{
-		AUDIO_MIXER_CHECK(SourceId < NumTotalSources);
+		AUDIO_MIXER_CHECK(SourceId < NumLegacySources);
 		AUDIO_MIXER_CHECK(GameThreadInfo.bIsBusy[SourceId]);
 		AUDIO_MIXER_CHECK_GAME_THREAD(MixerDevice);
 
@@ -2127,7 +2125,7 @@ namespace Audio
 
 	void FMixerSourceManager::SetModVolume(const int32 SourceId, const float InModVolume)
 	{
-		AUDIO_MIXER_CHECK(SourceId < NumTotalSources);
+		AUDIO_MIXER_CHECK(SourceId < NumLegacySources);
 		AUDIO_MIXER_CHECK(GameThreadInfo.bIsBusy[SourceId]);
 		AUDIO_MIXER_CHECK_GAME_THREAD(MixerDevice);
 
@@ -2142,7 +2140,7 @@ namespace Audio
 
 	void FMixerSourceManager::SetModPitch(const int32 SourceId, const float InModPitch)
 	{
-		AUDIO_MIXER_CHECK(SourceId < NumTotalSources);
+		AUDIO_MIXER_CHECK(SourceId < NumLegacySources);
 		AUDIO_MIXER_CHECK(GameThreadInfo.bIsBusy[SourceId]);
 		AUDIO_MIXER_CHECK_GAME_THREAD(MixerDevice);
 
@@ -2157,7 +2155,7 @@ namespace Audio
 
 	void FMixerSourceManager::SetSubmixSendInfo(const int32 SourceId, const FMixerSourceSubmixSend& InSubmixSend)
 	{
-		AUDIO_MIXER_CHECK(SourceId < NumTotalSources);
+		AUDIO_MIXER_CHECK(SourceId < NumLegacySources);
 		AUDIO_MIXER_CHECK(GameThreadInfo.bIsBusy[SourceId]);
 		AUDIO_MIXER_CHECK_GAME_THREAD(MixerDevice);
 
@@ -2209,18 +2207,14 @@ namespace Audio
 					SourceSubmixOutputBuffers[SourceId].SetPreAttenuationSourceBuffer(nullptr);
 				}
 
-				FMixerSourceVoice* SourceVoice = MixerSources[SourceId];
-				if (ensureAlways(nullptr != SourceVoice))
-				{
-					InSubmixPtr->AddOrSetSourceVoice(SourceVoice, InSubmixSend.SendLevel, InSubmixSend.SubmixSendStage);
-				}
+				InSubmixPtr->AddOrSetSourceById(SourceId, InSubmixSend.SendLevel, InSubmixSend.SubmixSendStage);
 			}
 		}, AUDIO_MIXER_THREAD_COMMAND_STRING("SetSubmixSendInfo()"));
 	}
 
 	void FMixerSourceManager::ClearSubmixSendInfo(const int32 SourceId, const FMixerSourceSubmixSend& InSubmixSend)
 	{
-		AUDIO_MIXER_CHECK(SourceId < NumTotalSources);
+		AUDIO_MIXER_CHECK(SourceId < NumLegacySources);
 		AUDIO_MIXER_CHECK(GameThreadInfo.bIsBusy[SourceId]);
 		AUDIO_MIXER_CHECK_GAME_THREAD(MixerDevice);
 
@@ -2261,14 +2255,14 @@ namespace Audio
 				}
 
 				// Now remove the source voice from the submix send list
-				InSubmixPtr->RemoveSourceVoice(MixerSources[SourceId]);
+				InSubmixPtr->RemoveSourceById(SourceId);
 			}
 		}, AUDIO_MIXER_THREAD_COMMAND_STRING("ClearSubmixSendInfo()"));
 	}
 
 	void FMixerSourceManager::SetBusSendInfo(const int32 SourceId, EBusSendType InAudioBusSendType, uint32 AudioBusId, float BusSendLevel, const FString& InBusName)
 	{
-		AUDIO_MIXER_CHECK(SourceId < NumTotalSources);
+		AUDIO_MIXER_CHECK(SourceId < NumLegacySources);
 		AUDIO_MIXER_CHECK(GameThreadInfo.bIsBusy[SourceId]);
 		AUDIO_MIXER_CHECK_GAME_THREAD(MixerDevice);
 
@@ -2345,6 +2339,7 @@ namespace Audio
 	float FMixerSourceManager::GetVolumeModulationValue(const int32 SourceId) const
 	{
 		AUDIO_MIXER_CHECK_GAME_THREAD(MixerDevice);
+		AUDIO_MIXER_CHECK(SourceId < NumLegacySources); 
 		return GameThreadInfo.ModulationVolume[SourceId];
 	}
 
@@ -2352,25 +2347,28 @@ namespace Audio
 	double FMixerSourceManager::GetCPUCoreUtilization(const int32 SourceId) const
 	{
 		AUDIO_MIXER_CHECK_GAME_THREAD(MixerDevice);
-		AUDIO_MIXER_CHECK(SourceId < NumTotalSources);
+		AUDIO_MIXER_CHECK(SourceId < NumLegacySources);
 		return GameThreadInfo.CPUCoreUtilization[SourceId];
 	}
 #endif // if ENABLE_AUDIO_DEBUG
 
 	float FMixerSourceManager::GetRelativeRenderCost(const int32 SourceId) const
 	{
+		AUDIO_MIXER_CHECK(SourceId < NumLegacySources);
 		return GameThreadInfo.RelativeRenderCost[SourceId];
 	}
 
 	bool FMixerSourceManager::IsUsingHRTFSpatializer(const int32 SourceId) const
 	{
 		AUDIO_MIXER_CHECK_GAME_THREAD(MixerDevice);
+		AUDIO_MIXER_CHECK(SourceId < NumLegacySources);
 		return GameThreadInfo.bIsUsingHRTFSpatializer[SourceId];
 	}
 
 	bool FMixerSourceManager::NeedsSpeakerMap(const int32 SourceId) const
 	{
 		AUDIO_MIXER_CHECK_GAME_THREAD(MixerDevice);
+		AUDIO_MIXER_CHECK(SourceId < NumLegacySources);
 		return GameThreadInfo.bNeedsSpeakerMap[SourceId];
 	}
 
@@ -2614,7 +2612,36 @@ namespace Audio
 			}
 #endif // UE_AUDIO_PROFILERTRACE_ENABLED
 
-			if (!PerSourceResampling)
+			const bool bIsGeneratorSource = SourceId >= NumLegacySources;
+			if (bIsGeneratorSource)
+			{
+				SourceInfo.bHasStarted = true;
+
+				const int32 GeneratorSourceIndex = SourceId - NumLegacySources;
+				const ISoundGeneratorPtr& SourceGenerator = GeneratorSourceState[GeneratorSourceIndex].Generator;
+				check(SourceGenerator.IsValid());
+				
+				TArrayView<float> WriteBuffer(PreDistanceAttenBufferPtr, NumSamples);
+				while (WriteBuffer.Num() > 0)
+				{
+					const int32 SamplesGenerated = SourceGenerator->OnGenerateAudio(WriteBuffer.GetData(), WriteBuffer.Num());
+					WriteBuffer = WriteBuffer.RightChop(SamplesGenerated);
+					if (SourceGenerator->IsFinished())
+					{
+						SourceInfo.bIsLastBuffer = true;
+						break;
+					}
+					else if (!ensureMsgf(SamplesGenerated > 0, TEXT("Source generator returned no samples but is not finished!")))
+					{
+						break;
+					}
+				}
+				
+				const int32 TotalSamplesGenerated = NumSamples - WriteBuffer.Num();
+				SourceInfo.CurrentFrameIndex += TotalSamplesGenerated / SourceInfo.NumInputChannels;
+				SourceInfo.NumFramesPlayed += TotalSamplesGenerated / SourceInfo.NumInputChannels;
+			}
+			else if (!PerSourceResampling)
 			{
 				int32 SampleIndex = 0;
 				float CurrentAlpha = SourceInfo.CurrentFrameAlpha;
@@ -3307,7 +3334,10 @@ namespace Audio
 						VolumeStart *= ModVolumeEnd;
 					}
 					VolumeDestination *= ModVolumeEnd;
-					GameThreadInfo.ModulationVolume[SourceId] = ModVolumeEnd;
+					if (SourceId < NumLegacySources)
+					{
+						GameThreadInfo.ModulationVolume[SourceId] = ModVolumeEnd;
+					}
 				}
 
 				TArrayView<float> PreDistanceAttenBufferFadeSamplesView(PreDistanceAttenBufferPtr, NumFadeSamples);
@@ -3344,7 +3374,10 @@ namespace Audio
 						VolumeStart *= ModVolumeEnd;
 					}
 					VolumeDestination *= ModVolumeEnd;
-					GameThreadInfo.ModulationVolume[SourceId] = ModVolumeEnd;
+					if (SourceId < NumLegacySources)
+					{
+						GameThreadInfo.ModulationVolume[SourceId] = ModVolumeEnd;
+					}
 				}
 
 				Audio::ArrayFade(PreDistanceAttenBufferView, VolumeStart, VolumeDestination);
@@ -3435,7 +3468,7 @@ namespace Audio
 			SourceInfo.bEffectTailsDone = true;
 		}
 
-		if (!bWasEffectTailsDone && SourceInfo.bEffectTailsDone)
+		if (SourceInfo.SourceListener && !bWasEffectTailsDone && SourceInfo.bEffectTailsDone)
 		{
 			SourceInfo.SourceListener->OnEffectTailsDone();
 		}
@@ -3723,7 +3756,10 @@ namespace Audio
 				// but it should reset in the next tick
 				if (SourceInfo.bIs3D)
 				{
-					GameThreadInfo.bNeedsSpeakerMap[SourceId] = true;
+					if (SourceId < NumLegacySources)
+					{
+						GameThreadInfo.bNeedsSpeakerMap[SourceId] = true;
+					}
 					SourceInfo.ScratchChannelMap.AddZeroed(NumSourceChannels * InNumOutputChannels);
 				}
 				// If it's a 2D sound, then just get a new channel map appropriate for the new device channel count
@@ -3805,7 +3841,7 @@ namespace Audio
 
 	void FMixerSourceManager::PauseSoundForQuantizationCommand(const int32 SourceId)
 	{
-		AUDIO_MIXER_CHECK(SourceId < NumTotalSources);
+		AUDIO_MIXER_CHECK(SourceId < NumLegacySources);
 		AUDIO_MIXER_CHECK_AUDIO_PLAT_THREAD(MixerDevice);
 
 		FSourceInfo& SourceInfo = SourceInfos[SourceId];
@@ -3816,7 +3852,7 @@ namespace Audio
 
 	void FMixerSourceManager::SetSubBufferDelayForSound(const int32 SourceId, const int32 FramesToDelay)
 	{
-		AUDIO_MIXER_CHECK(SourceId < NumTotalSources);
+		AUDIO_MIXER_CHECK(SourceId < NumLegacySources);
 		AUDIO_MIXER_CHECK_AUDIO_PLAT_THREAD(MixerDevice);
 
 		FSourceInfo& SourceInfo = SourceInfos[SourceId];
@@ -3826,7 +3862,7 @@ namespace Audio
 
 	void FMixerSourceManager::UnPauseSoundForQuantizationCommand(const int32 SourceId)
 	{
-		AUDIO_MIXER_CHECK(SourceId < NumTotalSources);
+		AUDIO_MIXER_CHECK(SourceId < NumLegacySources);
 		AUDIO_MIXER_CHECK_AUDIO_PLAT_THREAD(MixerDevice);
 
 		FSourceInfo& SourceInfo = SourceInfos[SourceId];
@@ -3982,15 +4018,7 @@ namespace Audio
 				
 				if (!SourceInfo.bIsDone)
 				{
-					SourceInfo.bIsDone = true;
-
-					// Notify that we're now done with this source
-					SourceInfo.SourceListener->OnDone();
-
-					if (SourceInfo.AudioLink)
-					{
-						SourceInfo.AudioLink->OnSourceDone(SourceId);
-					}
+					SetSourceDone(SourceId);
 				}
 			}
 		}
@@ -4006,14 +4034,37 @@ namespace Audio
 			if (!SourceInfo.bIsDone && SourceInfo.bIsStopping && SourceInfo.VolumeSourceDestination == 0.0f)
 			{
 				SourceInfo.bIsStopping = false;
-				SourceInfo.bIsDone = true;
-				SourceInfo.SourceListener->OnDone();
-				if (SourceInfo.AudioLink)
-				{
-					SourceInfo.AudioLink->OnSourceDone(SourceId);
-				}
+				SetSourceDone(SourceId);
 			}
+		}
+	}
+
+	void FMixerSourceManager::SetSourceDone(const int32 SourceId)
+	{
+		FSourceInfo& SourceInfo = SourceInfos[SourceId];
+
+		SourceInfo.bIsDone = true;
 
+		// Notify that we're now done with this source
+		if (SourceInfo.SourceListener)
+		{
+			SourceInfo.SourceListener->OnDone();
+		}
+
+		if (SourceInfo.AudioLink)
+		{
+			SourceInfo.AudioLink->OnSourceDone(SourceId);
+		}
+
+		if (SourceId >= NumLegacySources)
+		{
+			FGeneratorSourceState& State = GeneratorSourceState[SourceId - NumLegacySources];
+			if (State.Generator.IsValid())
+			{
+				State.Generator->OnEndGenerate();
+				State.Generator.Reset();
+			}
+			State.bSourceBusy = false;
 		}
 	}
 
@@ -4360,6 +4411,125 @@ namespace Audio
 		}
 	}
 
+	int32 FMixerSourceManager::AllocateGeneratorSource(ISoundGeneratorRef InGenerator)
+	{
+		auto IsFree = [this](const FGeneratorSourceState& State)
+			{
+				return !State.bAllocated && !State.bSourceBusy;
+			};
+		const int32 GeneratorSourceIndex = GeneratorSourceState.IndexOfByPredicate(IsFree);
+		if (GeneratorSourceIndex != INDEX_NONE)
+		{
+			GeneratorSourceState[GeneratorSourceIndex].bAllocated = true;
+			GeneratorSourceState[GeneratorSourceIndex].bSourceBusy = true;
+			const int32 SourceId = GeneratorSourceIndex + NumLegacySources;
+			check(SourceId >= NumLegacySources && SourceId < NumTotalSources);
+			AudioMixerThreadMPSCCommand([this, SourceId, InGenerator]
+				{
+					static int32 PlayOrder = 1 << 31;
+
+					FSourceInfo& SourceInfo = SourceInfos[SourceId];
+					SourceInfo.bIsBusy = true;
+					SourceInfo.NumInputChannels = InGenerator->GetNumChannels();
+					SourceInfo.PlayOrder = PlayOrder++;
+
+					check(SourceInfo.NumInputChannels > 0);
+					GeneratorSourceState[SourceId - NumLegacySources].Generator = InGenerator;
+					InGenerator->OnBeginGenerate();
+
+					// Initialize the submix output source for this source id
+					FMixerSourceSubmixOutputBufferSettings SourceSubmixOutputResetSettings = {};
+					SourceSubmixOutputResetSettings.NumOutputChannels = MixerDevice->GetDeviceOutputChannels();
+					SourceSubmixOutputResetSettings.NumSourceChannels = SourceInfo.NumInputChannels;
+					SourceSubmixOutputBuffers[SourceId].Reset(SourceSubmixOutputResetSettings);
+
+					MixerDevice->Get2DChannelMap(SourceInfo.bIsVorbis, SourceSubmixOutputResetSettings.NumSourceChannels, SourceSubmixOutputResetSettings.NumOutputChannels, SourceInfo.bIsCenterChannelOnly, SourceInfo.ScratchChannelMap);
+					SourceSubmixOutputBuffers[SourceId].SetChannelMap(SourceInfo.ScratchChannelMap, SourceInfo.bIsCenterChannelOnly);
+
+					FMixerSubmixPtr SubmixPtr = MixerDevice->GetMasterSubmix().Pin();
+					if (SubmixPtr.IsValid())
+					{
+						SubmixPtr->AddOrSetSourceById(SourceId, 1.0f, EMixerSourceSubmixSendStage::PostDistanceAttenuation);
+					}
+
+					if (UseRenderScheduler())
+					{
+						GetRenderScheduler().AddStep(FAudioRenderStepId::FromWaveInstancePlayOrder(SourceInfo.PlayOrder), &SourceInfo);
+					}
+				}
+			);
+
+			return SourceId;
+		}
+		else
+		{
+			return INDEX_NONE;
+		}
+	}
+
+	void FMixerSourceManager::FreeGeneratorSource(const int32 SourceId)
+	{
+		const int32 GeneratorSourceIndex = SourceId - NumLegacySources;
+		check(SourceId >= NumLegacySources && SourceId < NumTotalSources && GeneratorSourceState[GeneratorSourceIndex].bAllocated);
+		GeneratorSourceState[GeneratorSourceIndex].bAllocated = false;
+		AudioMixerThreadMPSCCommand([this, SourceId, GeneratorSourceIndex]
+			{
+				FMixerSubmixPtr SubmixPtr = MixerDevice->GetMasterSubmix().Pin();
+				if (SubmixPtr.IsValid())
+				{
+					SubmixPtr->RemoveSourceById(SourceId);
+				}
+				ISoundGeneratorPtr& Generator = GeneratorSourceState[GeneratorSourceIndex].Generator;
+				if (GeneratorSourceState[GeneratorSourceIndex].bSourceBusy && ensure(Generator.IsValid()))
+				{
+					Generator->OnEndGenerate();
+					Generator.Reset();
+				}
+				ReleaseSource(SourceId);
+				GeneratorSourceState[GeneratorSourceIndex].bSourceBusy = false;
+			}
+		);
+	}
+
+	void FMixerSourceManager::PlayGeneratorSource(const int32 SourceId)
+	{
+		check(SourceId >= NumLegacySources && SourceId < NumTotalSources && GeneratorSourceState[SourceId - NumLegacySources].bAllocated);
+		AudioMixerThreadMPSCCommand([this, SourceId]
+			{
+				FSourceInfo& SourceInfo = SourceInfos[SourceId];
+				SourceInfo.bIsDone = false;
+				SourceInfo.bIsPlaying = true;
+				SourceInfo.bIsPaused = false;
+				SourceInfo.bIsActive = true;
+				SourceInfo.bEnableBaseSubmix = true;
+			}
+		);
+	}
+
+	void FMixerSourceManager::StopGeneratorSource(const int32 SourceId)
+	{
+		check(SourceId >= NumLegacySources && SourceId < NumTotalSources && GeneratorSourceState[SourceId - NumLegacySources].bAllocated);
+		AudioMixerThreadMPSCCommand([this, SourceId]
+			{
+				FSourceInfo& SourceInfo = SourceInfos[SourceId];
+				SourceInfo.bIsPlaying = false;
+				SourceInfo.bIsPaused = false;
+				SourceInfo.bIsActive = false;
+				SourceInfo.bIsStopping = false;
+				if (!SourceInfo.bIsDone)
+				{
+					SetSourceDone(SourceId);
+				}
+			}
+		);
+	}
+
+	bool FMixerSourceManager::IsGeneratorSourceDone(const int32 SourceId)
+	{
+		check(SourceId >= NumLegacySources && SourceId < NumTotalSources && GeneratorSourceState[SourceId - NumLegacySources].bAllocated);
+		return !GeneratorSourceState[SourceId - NumLegacySources].bSourceBusy;
+	}
+
 	void FMixerSourceManager::FSourceInfo::DoRenderStep()
 	{
 		const int32 SourceId = this - SourceManager->SourceInfos.GetData();
@@ -4374,4 +4544,4 @@ namespace Audio
 		return TEXT("Render mixer source");
 #endif
 	}
-}
\ No newline at end of file
+}
diff --git a/Engine/Source/Runtime/AudioMixer/Private/AudioMixerSourceManager.h b/Engine/Source/Runtime/AudioMixer/Private/AudioMixerSourceManager.h
--- a/Engine/Source/Runtime/AudioMixer/Private/AudioMixerSourceManager.h
+++ b/Engine/Source/Runtime/AudioMixer/Private/AudioMixerSourceManager.h
@@ -186,6 +186,7 @@ namespace Audio
 	{
 		// Total number of sources to use in the source manager
 		int32 NumSources = 0;
+		int32 NumGeneratorSources = 0;
 	};
 
 	class FMixerSourceManager
@@ -345,6 +346,12 @@ namespace Audio
 			return bShouldResetTemp; 
 		}
 
+		int32 AllocateGeneratorSource(ISoundGeneratorRef InGenerator);
+		bool IsGeneratorSourceDone(const int32 SourceId);
+		void FreeGeneratorSource(const int32 SourceId);
+		void PlayGeneratorSource(const int32 SourceId);
+		void StopGeneratorSource(const int32 SourceId);
+
 	private:
 #define INVALID_AUDIO_RENDER_THREAD_ID static_cast<uint32>(-1)
 		uint32 AudioRenderThreadId = INVALID_AUDIO_RENDER_THREAD_ID;
@@ -352,6 +359,7 @@ namespace Audio
 		void BuildSourceEffectChain(const int32 SourceId, FSoundEffectSourceInitData& InitData, const TArray<FSourceEffectChainEntry>& SourceEffectChain, TArray<TSoundEffectSourcePtr>& OutSourceEffects);
 		void ResetSourceEffectChain(const int32 SourceId);
 		void ReadSourceFrame(const int32 SourceId);
+		void SetSourceDone(const int32 SourceId);
 
 		void GenerateSourceAudio(const bool bGenerateBuses);
 		void GenerateSourceAudio(const bool bGenerateBuses, const int32 SourceIdStart, const int32 SourceIdEnd);
@@ -688,9 +696,25 @@ namespace Audio
 
 		int32 NumActiveSources;
 		int32 NumTotalSources;
+		int32 NumLegacySources;
+		int32 NumGeneratorSources;
 		int32 NumOutputFrames;
 		int32 NumOutputSamples;
 
+		struct FGeneratorSourceState
+		{
+			// Generator used for rendering, owned by render thread.
+			ISoundGeneratorPtr Generator;
+
+			// Set on game thread, cleared from render thread when rendering is complete
+			std::atomic<uint8> bSourceBusy = false;
+
+			// Is this voice allocated on the game thread?
+			uint8 bAllocated:1 = false;
+		};
+
+		TArray<FGeneratorSourceState> GeneratorSourceState;
+
 		// Commands queued up to execute
 		FThreadSafeCounter NumCommands;
 
diff --git a/Engine/Source/Runtime/AudioMixer/Private/AudioMixerSourceVoice.cpp b/Engine/Source/Runtime/AudioMixer/Private/AudioMixerSourceVoice.cpp
--- a/Engine/Source/Runtime/AudioMixer/Private/AudioMixerSourceVoice.cpp
+++ b/Engine/Source/Runtime/AudioMixer/Private/AudioMixerSourceVoice.cpp
@@ -354,32 +354,6 @@ namespace Audio
 		return SourceManager->GetRelativeRenderCost(SourceId);
 	}
 
-	void FMixerSourceVoice::MixOutputBuffers(int32 InNumOutputChannels, const float SendLevel, EMixerSourceSubmixSendStage InSubmixSendStage, FAlignedFloatBuffer& OutWetBuffer) const
-	{
-		AUDIO_MIXER_CHECK_AUDIO_PLAT_THREAD(MixerDevice);
-
-		if (IsRenderingToSubmixes())
-		{
-			SourceManager->MixOutputBuffers(SourceId, InNumOutputChannels, SendLevel, InSubmixSendStage, OutWetBuffer);
-		}
-	}
-
-	const ISoundfieldAudioPacket* FMixerSourceVoice::GetEncodedOutput(const FSoundfieldEncodingKey& InKey) const
-	{
-		AUDIO_MIXER_CHECK_AUDIO_PLAT_THREAD(MixerDevice);
-
-		if (IsRenderingToSubmixes())
-		{
-			return SourceManager->GetEncodedOutput(SourceId, InKey);
-		}
-		return nullptr;
-	}
-
-	const FQuat FMixerSourceVoice::GetListenerRotationForVoice() const
-	{
-		return SourceManager->GetListenerRotation(SourceId);
-	}
-
 	void FMixerSourceVoice::SetSubmixSendInfo(FMixerSubmixWeakPtr Submix, const float SendLevel, const EMixerSourceSubmixSendStage SendStage/* = EMixerSourceSubmixSendStage::PostDistanceAttenuation*/)
 	{
 		AUDIO_MIXER_CHECK_GAME_THREAD(MixerDevice);
diff --git a/Engine/Source/Runtime/AudioMixer/Private/AudioMixerSourceVoice.h b/Engine/Source/Runtime/AudioMixer/Private/AudioMixerSourceVoice.h
--- a/Engine/Source/Runtime/AudioMixer/Private/AudioMixerSourceVoice.h
+++ b/Engine/Source/Runtime/AudioMixer/Private/AudioMixerSourceVoice.h
@@ -128,15 +128,6 @@ namespace Audio
 		// Useful as a metric more human-centric than CPU utilization.
 		float GetRelativeRenderCost() const;
 
-		// Mixes the dry and wet buffer audio into the given buffers.
-		void MixOutputBuffers(int32 InNumChannels, const float SendLevel, EMixerSourceSubmixSendStage InSubmixSendStage, FAlignedFloatBuffer& OutWetBuffer) const;
-
-		// For soundfield conversions, get the encoded audio.
-		const ISoundfieldAudioPacket* GetEncodedOutput(const FSoundfieldEncodingKey& InKey) const;
-
-		// This will return the listener rotation used for this source voice.
-		const FQuat GetListenerRotationForVoice() const;
-
 		// Sets the submix send levels
 		void SetSubmixSendInfo(FMixerSubmixWeakPtr Submix, const float SendLevel, const EMixerSourceSubmixSendStage SendStage = EMixerSourceSubmixSendStage::PostDistanceAttenuation);
 
diff --git a/Engine/Source/Runtime/AudioMixer/Private/AudioMixerSubmix.cpp b/Engine/Source/Runtime/AudioMixer/Private/AudioMixerSubmix.cpp
--- a/Engine/Source/Runtime/AudioMixer/Private/AudioMixerSubmix.cpp
+++ b/Engine/Source/Runtime/AudioMixer/Private/AudioMixerSubmix.cpp
@@ -642,7 +642,7 @@ namespace Audio
 
 	int32 FMixerSubmix::GetNumSourceVoices() const
 	{
-		return MixerSourceVoices.Num();
+		return MixerSourceIds.Num();
 	}
 
 	int32 FMixerSubmix::GetNumEffects() const
@@ -664,6 +664,11 @@ namespace Audio
 	}
 
 	void FMixerSubmix::AddOrSetSourceVoice(FMixerSourceVoice* InSourceVoice, const float InSendLevel, EMixerSourceSubmixSendStage InSubmixSendStage)
+	{
+		AddOrSetSourceById(InSourceVoice->GetSourceId(), InSendLevel, InSubmixSendStage);
+	}
+
+	void FMixerSubmix::AddOrSetSourceById(const int32 InSourceId, const float InSendLevel, EMixerSourceSubmixSendStage InSubmixSendStage)
 	{
 		AUDIO_MIXER_CHECK_AUDIO_PLAT_THREAD(MixerDevice);
 
@@ -671,7 +676,7 @@ namespace Audio
 		NewVoiceData.SendLevel = InSendLevel;
 		NewVoiceData.SubmixSendStage = InSubmixSendStage;
 
-		MixerSourceVoices.Add(InSourceVoice, NewVoiceData);
+		MixerSourceIds.Add(InSourceId, NewVoiceData);
 	}
 
 	FPatchOutputStrongPtr FMixerSubmix::AddPatch(float InGain)
@@ -686,17 +691,22 @@ namespace Audio
 	}
 
 	void FMixerSubmix::RemoveSourceVoice(FMixerSourceVoice* InSourceVoice)
+	{
+		RemoveSourceById(InSourceVoice->GetSourceId());
+	}
+
+	void FMixerSubmix::RemoveSourceById(const int32 InSourceId)
 	{
 		AUDIO_MIXER_CHECK_AUDIO_PLAT_THREAD(MixerDevice);
 		
 		// If the source has a corresponding ambisonics encoder, close it out.
 		uint32 SourceEncoderID = INDEX_NONE;
-		const FSubmixVoiceData* MixerSourceVoiceData = MixerSourceVoices.Find(InSourceVoice);
+		const FSubmixVoiceData* MixerSourceVoiceData = MixerSourceIds.Find(InSourceId);
 
 		// If we did find a valid corresponding FSubmixVoiceData, remove it from the map.
 		if (MixerSourceVoiceData)
 		{
-			int32 NumRemoved = MixerSourceVoices.Remove(InSourceVoice);
+			int32 NumRemoved = MixerSourceIds.Remove(InSourceId);
 			AUDIO_MIXER_CHECK(NumRemoved == 1);
 		}
 	}
@@ -1023,7 +1033,7 @@ namespace Audio
 				ScratchBuffer.AddZeroed(ChildSubmixSharedPtr->NumSamples);
 				
 
-				// If this is true, the Soundfield Factory explicitly requested that a seperate encoder stream was set up for every
+				// If this is true, the Soundfield Factory explicitly requested that a separate encoder stream was set up for every
 				// non-soundfield child submix.
 				if (Child.Encoder.IsValid())
 				{
@@ -1235,7 +1245,7 @@ namespace Audio
 		}
 
 		// If this submix is not rendering any sources directly and silence has been detected, we need to check it's children submixes
-		if (MixerSourceVoices.Num() == 0 && SilenceTimeStartSeconds >= 0.0)
+		if (MixerSourceIds.Num() == 0 && SilenceTimeStartSeconds >= 0.0)
 		{
 			// Use the audio clock to check against the disablement timeout
 			double AudioClock = MixerDevice->GetAudioClock();
@@ -1376,7 +1386,7 @@ namespace Audio
 
 				UE_CLOG(LogSubmixEnablementCVar == 1, LogAudioMixer, Verbose,
 					TEXT("Submix Disabled. Num Sources: %d, Time Silent: %.2f, Disablement Threshold: %.2f, Submix Name: %s"), 
-					MixerSourceVoices.Num(),
+					MixerSourceIds.Num(),
 					(float)(MixerDevice->GetAudioClock() - SilenceTimeStartSeconds),
 					AutoDisableTime,
 					*SubmixName);
@@ -1436,16 +1446,16 @@ namespace Audio
 
 		{
 			CSV_SCOPED_TIMING_STAT(Audio, SubmixSource);
-			CONDITIONAL_SCOPE_CYCLE_COUNTER(STAT_AudioMixerSubmixSource, (MixerSourceVoices.Num() > 0));
+			CONDITIONAL_SCOPE_CYCLE_COUNTER(STAT_AudioMixerSubmixSource, (MixerSourceIds.Num() > 0));
 
 			// Loop through this submix's sound sources
-			for (const auto& MixerSourceVoiceIter : MixerSourceVoices)
+			for (const auto& MixerSourceVoiceIter : MixerSourceIds)
 			{
-				const FMixerSourceVoice* MixerSourceVoice = MixerSourceVoiceIter.Key;
+				const int32 SourceId = MixerSourceVoiceIter.Key;
 				const float SendLevel = MixerSourceVoiceIter.Value.SendLevel;
 				const EMixerSourceSubmixSendStage SubmixSendStage = MixerSourceVoiceIter.Value.SubmixSendStage;
 
-				MixerSourceVoice->MixOutputBuffers(NumChannels, SendLevel, SubmixSendStage, InputBuffer);
+				MixerDevice->GetSourceManager()->MixOutputBuffers(SourceId, NumChannels, SendLevel, SubmixSendStage, InputBuffer);
 				
 				// Check the buffer after each voice mix to catch any bad values.
 				AUDIO_CHECK_BUFFER_NAMED_MSG(InputBuffer, TEXT("Submix SourceMix"), TEXT("Submix: %s"), *GetName());
@@ -1880,6 +1890,8 @@ namespace Audio
 		check(IsSoundfieldSubmix());
 		PumpCommandQueue();
 
+		FMixerSourceManager* SourceManager = MixerDevice->GetSourceManager();
+
 		// Mix all submix audio into OutputAudio.
 		{
 			CSV_SCOPED_TIMING_STAT(Audio, SubmixSoundfieldChildren);
@@ -1919,19 +1931,19 @@ namespace Audio
 		// Mix all source sends into OutputAudio.
 		{
 			CSV_SCOPED_TIMING_STAT(Audio, SubmixSoundfieldSources);
-			CONDITIONAL_SCOPE_CYCLE_COUNTER(STAT_AudioMixerSubmixSoundfieldSources, (MixerSourceVoices.Num() > 0));
+			CONDITIONAL_SCOPE_CYCLE_COUNTER(STAT_AudioMixerSubmixSoundfieldSources, (MixerSourceIds.Num() > 0));
 
 			check(SoundfieldStreams.Mixer.IsValid());
 
 			// Loop through this submix's sound sources
-			for (const auto& MixerSourceVoiceIter : MixerSourceVoices)
+			for (const auto& MixerSourceVoiceIter : MixerSourceIds)
 			{
-				const FMixerSourceVoice* MixerSourceVoice = MixerSourceVoiceIter.Key;
+				const int32 SourceId = MixerSourceVoiceIter.Key;
 				const float SendLevel = MixerSourceVoiceIter.Value.SendLevel;
 
 				// if this voice has a valid encoded packet, mix it in.
-				const ISoundfieldAudioPacket* Packet = MixerSourceVoice->GetEncodedOutput(GetKeyForSubmixEncoding());
-				UpdateListenerRotation(MixerSourceVoice->GetListenerRotationForVoice());
+				const ISoundfieldAudioPacket* Packet = SourceManager->GetEncodedOutput(SourceId, GetKeyForSubmixEncoding());
+				UpdateListenerRotation(SourceManager->GetListenerRotation(SourceId));
 
 				if (Packet)
 				{
diff --git a/Engine/Source/Runtime/AudioMixer/Public/AudioMixerDevice.h b/Engine/Source/Runtime/AudioMixer/Public/AudioMixerDevice.h
--- a/Engine/Source/Runtime/AudioMixer/Public/AudioMixerDevice.h
+++ b/Engine/Source/Runtime/AudioMixer/Public/AudioMixerDevice.h
@@ -49,6 +49,7 @@ namespace Audio
 	class FMixerSubmix;
 	class FAudioFormatSettings;
 	class FAudioRenderScheduler;
+	class IAudioMixerGeneratorSource;
 
 	typedef TSharedPtr<FMixerSubmix, ESPMode::ThreadSafe> FMixerSubmixPtr;
 	typedef TWeakPtr<FMixerSubmix, ESPMode::ThreadSafe> FMixerSubmixWeakPtr;
@@ -376,6 +377,9 @@ namespace Audio
 			return *RenderScheduler;
 		}
 
+		UE_EXPERIMENTAL(5.8, "AudioMixer Generator Source API is experimental and likely to change.")
+		AUDIOMIXER_API TUniquePtr<IAudioMixerGeneratorSource> CreateGeneratorSource(ISoundGeneratorRef InGenerator);
+
 	protected:
 		AUDIOMIXER_API virtual void InitSoundSubmixes() override;
 
diff --git a/Engine/Source/Runtime/AudioMixer/Public/AudioMixerSubmix.h b/Engine/Source/Runtime/AudioMixer/Public/AudioMixerSubmix.h
--- a/Engine/Source/Runtime/AudioMixer/Public/AudioMixerSubmix.h
+++ b/Engine/Source/Runtime/AudioMixer/Public/AudioMixerSubmix.h
@@ -175,14 +175,20 @@ namespace Audio
 		// Returns the size of the submix chain. 
 		AUDIOMIXER_API int32 GetSizeOfSubmixChain() const;
 
-		// Add (if not already added) or sets the amount of the source voice's send amount
+		UE_DEPRECATED(5.8, "Use FMixerSubmix::AddOrSetSourceById instead.")
 		AUDIOMIXER_API void AddOrSetSourceVoice(FMixerSourceVoice* InSourceVoice, const float SendLevel, EMixerSourceSubmixSendStage InSubmixSendStage);
 
+		// Add (if not already added) or sets the amount of the source voice's send amount
+		void AddOrSetSourceById(const int32 InSourceId, const float SendLevel, EMixerSourceSubmixSendStage InSubmixSendStage);
+
 		AUDIOMIXER_API FPatchOutputStrongPtr AddPatch(float InGain);
 
-		/** Removes the given source voice from the submix. */
+		UE_DEPRECATED(5.8, "Use FMixerSubmix::RemoveSourceById instead.")
 		AUDIOMIXER_API void RemoveSourceVoice(FMixerSourceVoice* InSourceVoice);
 
+		/** Removes the given source voice from the submix. */
+		void RemoveSourceById(const int32 InSourceId);
+
 		/** Appends the effect submix to the effect submix chain. */
 		AUDIOMIXER_API void AddSoundEffectSubmix(FSoundEffectSubmixPtr InSoundEffectSubmix);
 
@@ -429,8 +435,8 @@ namespace Audio
 		// Owning mixer device. 
 		FMixerDevice* MixerDevice;
 
-		// Map of mixer source voices with a given send level for this submix
-		TMap<FMixerSourceVoice*, FSubmixVoiceData> MixerSourceVoices;
+		// Map of mixer source voices identified by SourceId with a given send level for this submix
+		TMap<int32, FSubmixVoiceData> MixerSourceIds;
 
 		FAlignedFloatBuffer ScratchBuffer;
 		FAlignedFloatBuffer SubmixChainMixBuffer;
diff --git a/Engine/Source/Runtime/AudioMixer/Public/IAudioMixerGeneratorSource.h b/Engine/Source/Runtime/AudioMixer/Public/IAudioMixerGeneratorSource.h
new file mode 100644
--- /dev/null
+++ b/Engine/Source/Runtime/AudioMixer/Public/IAudioMixerGeneratorSource.h
@@ -0,0 +1,20 @@
+// Copyright Epic Games, Inc. All Rights Reserved.
+
+#pragma once
+
+#include "CoreMinimal.h"
+
+namespace Audio
+{
+	UE_EXPERIMENTAL(5.8, "AudioMixer Generator Source API is experimental and likely to change.")
+	class IAudioMixerGeneratorSource
+	{
+	public:
+		virtual ~IAudioMixerGeneratorSource() = default;
+
+		virtual void Play() = 0;
+		virtual void Stop() = 0;
+
+		virtual bool IsDone() = 0;
+	};
+}
diff --git a/Engine/Source/Runtime/Engine/Classes/Sound/SoundGenerator.h b/Engine/Source/Runtime/Engine/Classes/Sound/SoundGenerator.h
--- a/Engine/Source/Runtime/Engine/Classes/Sound/SoundGenerator.h
+++ b/Engine/Source/Runtime/Engine/Classes/Sound/SoundGenerator.h
@@ -61,6 +61,9 @@ public:
 
 	virtual Audio::AudioTaskQueueId GetSynchronizedRenderQueueId() const { return 0; }
 
+	// Currently only used by generator source rendering, must be implemented to be used with a generator source.
+	virtual int32 GetNumChannels() const { return 0; }
+
 protected:
 
 	// Protected method to execute lambda in audio render thread
@@ -89,3 +92,4 @@ public:
 };
 
 typedef TSharedPtr<ISoundGenerator, ESPMode::ThreadSafe> ISoundGeneratorPtr;
+typedef TSharedRef<ISoundGenerator, ESPMode::ThreadSafe> ISoundGeneratorRef;
