Remove AudioMixer code deprecated in 5.5 or earlier

#rnx
#Synced-CL 51585013

diff --git a/Engine/Source/Runtime/AudioMixer/Private/AudioMixerDevice.cpp b/Engine/Source/Runtime/AudioMixer/Private/AudioMixerDevice.cpp
--- a/Engine/Source/Runtime/AudioMixer/Private/AudioMixerDevice.cpp
+++ b/Engine/Source/Runtime/AudioMixer/Private/AudioMixerDevice.cpp
@@ -3402,58 +3402,6 @@ namespace Audio
 		return *MainSubmix;
 	}
 
-	void FMixerDevice::RegisterSubmixBufferListener(ISubmixBufferListener* InSubmixBufferListener, USoundSubmix* InSubmix)
-	{
-		DECLARE_CYCLE_STAT(TEXT("FAudioThreadTask.RegisterSubmixBufferListener"), STAT_RegisterSubmixBufferListener, STATGROUP_AudioThreadCommands);
-
-		const bool bUseMaster = InSubmix == nullptr;
-		const TWeakObjectPtr<USoundSubmix> SubmixPtr(InSubmix);
-
-		auto RegisterLambda = [this, InSubmixBufferListener, bUseMaster, SubmixPtr]()
-		{
-			CSV_SCOPED_TIMING_STAT(Audio, RegisterSubmixBufferListener);
-
-			FMixerSubmixPtr FoundSubmix = bUseMaster
-				? GetMasterSubmix().Pin()
-				: GetSubmixInstance(SubmixPtr.Get()).Pin();
-
-			// Submix is not currently registered, so add this listener to a pending list for this submix
-			if (!bUseMaster && !FoundSubmix.IsValid() && SubmixPtr.IsValid())
-			{
-
-				FString SubmixName = SubmixPtr->GetPathName();
-				if (SubmixBufferListenerForceRegistrationCVar)
-				{
-					RegisterSoundSubmix(SubmixPtr.Get(), true /* bInit */);
-					FoundSubmix = GetSubmixInstance(SubmixPtr.Get()).Pin();
-					
-					UE_LOG(LogAudioMixer, Warning, TEXT("Submix '%s' force-enabled by Submix Buffer Listener '%s' due to CVar 'au.SubmixBufferListenerForceRegistration'."
-						"Please disable the CVar and manually ensure the submix is registered before registering the submix buffer listener."), *SubmixPtr->GetName(), *InSubmixBufferListener->GetListenerName());
-				}
-				else
-				{
-					UE_LOG(LogAudioMixer, Display, TEXT("Submix '%s' not currently loaded, Submix buffer listener '%s' added to pending list."), *SubmixName, *InSubmixBufferListener->GetListenerName());
-					if (!PendingSubmixBufferListeners.Contains(SubmixName))
-					{
-						PendingSubmixBufferListeners.Add(SubmixName, TMap<FString, TSharedRef<ISubmixBufferListener, ESPMode::ThreadSafe>>());
-					}
-					TMap<FString, TSharedRef<ISubmixBufferListener, ESPMode::ThreadSafe>>& PendingMap = PendingSubmixBufferListeners[SubmixName];
-					PendingMap.Add(InSubmixBufferListener->GetListenerName(), InSubmixBufferListener->AsShared());
-				}
-
-			}
-
-			if (FoundSubmix.IsValid())
-			{
-				PRAGMA_DISABLE_DEPRECATION_WARNINGS
-				FoundSubmix->RegisterBufferListener(InSubmixBufferListener);
-				PRAGMA_ENABLE_DEPRECATION_WARNINGS
-			}
-		};
-
-		FAudioThread::RunCommandOnAudioThread(MoveTemp(RegisterLambda));
-	}
-
 	void FMixerDevice::RegisterSubmixBufferListener(TSharedRef<ISubmixBufferListener, ESPMode::ThreadSafe> InSubmixBufferListener, USoundSubmix& InSubmix)
 	{
 		DECLARE_CYCLE_STAT(TEXT("FAudioThreadTask.RegisterSubmixBufferListener"), STAT_RegisterSubmixBufferListener, STATGROUP_AudioThreadCommands);
@@ -3509,44 +3457,6 @@ namespace Audio
 		FAudioThread::RunCommandOnAudioThread(MoveTemp(RegisterLambda));
 	}
 
-	void FMixerDevice::UnregisterSubmixBufferListener(ISubmixBufferListener* InSubmixBufferListener, USoundSubmix* InSubmix)
-	{
-		DECLARE_CYCLE_STAT(TEXT("FAudioThreadTask.UnregisterSubmixBufferListener"), STAT_UnregisterSubmixBufferListener, STATGROUP_AudioThreadCommands);
-
-		const bool bUseMaster = InSubmix == nullptr;
-		const TWeakObjectPtr<USoundSubmix> SubmixPtr(InSubmix);
-
-		auto UnregisterLambda = [this, InSubmixBufferListener, bUseMaster, SubmixPtr]()
-		{
-			CSV_SCOPED_TIMING_STAT(Audio, UnregisterSubmixBufferListener);
-
-			FMixerSubmixPtr FoundSubmix = bUseMaster
-				? GetMasterSubmix().Pin()
-				: GetSubmixInstance(SubmixPtr.Get()).Pin();
-
-			if (SubmixPtr.IsValid())
-			{
-				if (PendingSubmixBufferListeners.Contains(SubmixPtr->GetPathName()))
-				{
-					PendingSubmixBufferListeners[SubmixPtr->GetPathName()].Remove(InSubmixBufferListener->GetListenerName());
-				}
-			}
-				
-			if (FoundSubmix.IsValid())
-			{
-				PRAGMA_DISABLE_DEPRECATION_WARNINGS
-				FoundSubmix->UnregisterBufferListener(InSubmixBufferListener);
-				PRAGMA_ENABLE_DEPRECATION_WARNINGS
-			}
-			else
-			{
-				UE_LOG(LogAudioMixer, Display, TEXT("Submix buffer listener not unregistered. Submix not loaded."));
-			}
-		};
-
-		FAudioThread::RunCommandOnAudioThread(MoveTemp(UnregisterLambda));
-	}
-
 	void FMixerDevice::UnregisterSubmixBufferListener(TSharedRef<ISubmixBufferListener, ESPMode::ThreadSafe> InSubmixBufferListener, USoundSubmix& InSubmix)
 	{
 		DECLARE_CYCLE_STAT(TEXT("FAudioThreadTask.UnregisterSubmixBufferListener"), STAT_UnregisterSubmixBufferListener, STATGROUP_AudioThreadCommands);
diff --git a/Engine/Source/Runtime/AudioMixer/Private/AudioMixerSubmix.cpp b/Engine/Source/Runtime/AudioMixer/Private/AudioMixerSubmix.cpp
--- a/Engine/Source/Runtime/AudioMixer/Private/AudioMixerSubmix.cpp
+++ b/Engine/Source/Runtime/AudioMixer/Private/AudioMixerSubmix.cpp
@@ -2329,26 +2329,12 @@ namespace Audio
 		bIsRecording = true;
 	}
 
-	void FMixerSubmix::RegisterBufferListener(ISubmixBufferListener* BufferListener)
-	{
-		FScopeLock Lock(&BufferListenerCriticalSection);
-		check(BufferListener);
-		BufferListenerPtrs.AddUnique(BufferListener->AsShared());
-	}
-
 	void FMixerSubmix::RegisterBufferListener(TSharedRef<ISubmixBufferListener, ESPMode::ThreadSafe> BufferListener)
 	{
 		FScopeLock Lock(&BufferListenerCriticalSection);
 		BufferListenerPtrs.AddUnique(BufferListener);
 	}
 
-	void FMixerSubmix::UnregisterBufferListener(ISubmixBufferListener* BufferListener)
-	{
-		FScopeLock Lock(&BufferListenerCriticalSection);
-		check(BufferListener);
-		BufferListenerPtrs.Remove(BufferListener->AsShared());
-	}
-
 	void FMixerSubmix::UnregisterBufferListener(TSharedRef<ISubmixBufferListener, ESPMode::ThreadSafe> BufferListener)
 	{
 		FScopeLock Lock(&BufferListenerCriticalSection);
diff --git a/Engine/Source/Runtime/AudioMixer/Public/AudioMixerDevice.h b/Engine/Source/Runtime/AudioMixer/Public/AudioMixerDevice.h
--- a/Engine/Source/Runtime/AudioMixer/Public/AudioMixerDevice.h
+++ b/Engine/Source/Runtime/AudioMixer/Public/AudioMixerDevice.h
@@ -203,12 +203,6 @@ namespace Audio
 		AUDIOMIXER_API virtual void RemoveSpectralAnalysisDelegate(USoundSubmix* InSubmix, const FOnSubmixSpectralAnalysisBP& OnSubmixSpectralAnalysisBP) override;
 
 		// Submix buffer listener callbacks
-		UE_DEPRECATED(5.4, "Use RegisterSubmixBufferListener version that requires a shared reference to a listener and provide explicit reference to a submix: use GetMainSubmixObject to register with the Main Output Submix (rather than nullptr for safety), and instantiate buffer listener via the shared pointer API.")
-		AUDIOMIXER_API virtual void RegisterSubmixBufferListener(ISubmixBufferListener* InSubmixBufferListener, USoundSubmix* InSubmix = nullptr) override;
-
-		UE_DEPRECATED(5.4, "Use UnregisterSubmixBufferListener version that requires a shared reference to a listener and provide explicit reference to a submix: use GetMainSubmixObject to unregister from the Main Output Submix (rather than nullptr for safety), and instantiate buffer listener via the shared pointer API.")
-		AUDIOMIXER_API virtual void UnregisterSubmixBufferListener(ISubmixBufferListener* InSubmixBufferListener, USoundSubmix* InSubmix = nullptr) override;
-
 		AUDIOMIXER_API virtual void RegisterSubmixBufferListener(TSharedRef<ISubmixBufferListener, ESPMode::ThreadSafe> InSubmixBufferListener, USoundSubmix& InSubmix) override;
 
 		// This is optional, and should only be used to manually unregister before the ISubmixBufferListener is getting destroyed
diff --git a/Engine/Source/Runtime/AudioMixer/Public/AudioMixerSubmix.h b/Engine/Source/Runtime/AudioMixer/Public/AudioMixerSubmix.h
--- a/Engine/Source/Runtime/AudioMixer/Public/AudioMixerSubmix.h
+++ b/Engine/Source/Runtime/AudioMixer/Public/AudioMixerSubmix.h
@@ -266,15 +266,9 @@ namespace Audio
 		AUDIOMIXER_API void ResumeRecordingOutput();
 
 		// Register buffer listener with this submix
-		// Unregister buffer listener with this submix
-
-		UE_DEPRECATED(5.4, "This function is deprecated. Use RegisterBufferListener version that is provided a shared reference to a listener.")
-		AUDIOMIXER_API void RegisterBufferListener(ISubmixBufferListener* BufferListener);
-		
-		UE_DEPRECATED(5.4, "This function is deprecated. Use UnregisterBufferListener version that is provided a shared reference to a listener.")
-		AUDIOMIXER_API void UnregisterBufferListener(ISubmixBufferListener* BufferListener);
-
 		AUDIOMIXER_API void RegisterBufferListener(TSharedRef<ISubmixBufferListener, ESPMode::ThreadSafe> BufferListener);
+
+		// Unregister buffer listener with this submix
 		AUDIOMIXER_API void UnregisterBufferListener(TSharedRef<ISubmixBufferListener, ESPMode::ThreadSafe> BufferListener);
 
 		// Starts envelope following with the given attack time and release time
diff --git a/Engine/Source/Runtime/AudioMixer/Public/Components/SynthComponent.h b/Engine/Source/Runtime/AudioMixer/Public/Components/SynthComponent.h
--- a/Engine/Source/Runtime/AudioMixer/Public/Components/SynthComponent.h
+++ b/Engine/Source/Runtime/AudioMixer/Public/Components/SynthComponent.h
@@ -339,12 +339,6 @@ protected:
 	// Called when synth is created.
 	virtual bool Init(int32& SampleRate) { return true; }
 
-	UE_DEPRECATED(4.26, "Use OnBeginGenerate to get a callback before audio is generating on the audio render thread")
-	virtual void OnStart() {}
-
-	UE_DEPRECATED(4.26, "Use OnEndGenerate to get a callback when audio stops generating on the audio render thread")
-	virtual void OnStop() {}
-
 	// Called when the synth component begins generating audio in render thread
 	virtual void OnBeginGenerate() {}
 
diff --git a/Engine/Source/Runtime/AudioMixer/Public/Quartz/AudioMixerClockHandle.h b/Engine/Source/Runtime/AudioMixer/Public/Quartz/AudioMixerClockHandle.h
--- a/Engine/Source/Runtime/AudioMixer/Public/Quartz/AudioMixerClockHandle.h
+++ b/Engine/Source/Runtime/AudioMixer/Public/Quartz/AudioMixerClockHandle.h
@@ -168,12 +168,6 @@ public:
 		return RawHandle.DoesClockExist();
 	}
 
-	UE_DEPRECATED(5.1, "This function should not be called directly, and the original functionality has been moved into FQuartzTickable")
-	virtual void ProcessCommand(const Audio::FQuartzQuantizedCommandDelegateData& Data) override {}
-
-	UE_DEPRECATED(5.1, "This function should not be called directly, and the original functionality has been moved into FQuartzTickable")
-	virtual void ProcessCommand(const Audio::FQuartzMetronomeDelegateData& Data) override {};
-
 	AUDIOMIXER_API bool GetCurrentTickRate(const UObject* WorldContextObject, Audio::FQuartzClockTickRate& OutTickRate) const;
 
 private:
