Rename FMixerSubmix::IsValid() to IsOwnerValid() to avoid confusion with TSharedPtr::IsValid(). Clean up some unnecessary string copies.

#rnx
#Synced-CL 51682935

diff --git a/Engine/Source/Runtime/AudioMixer/Private/AudioMixerDevice.cpp b/Engine/Source/Runtime/AudioMixer/Private/AudioMixerDevice.cpp
--- a/Engine/Source/Runtime/AudioMixer/Private/AudioMixerDevice.cpp
+++ b/Engine/Source/Runtime/AudioMixer/Private/AudioMixerDevice.cpp
@@ -817,7 +817,7 @@ namespace Audio
 	{
 		ensure(IsInAudioThread());
 
-		if (!InSubmix || !InSubmix->IsValid())
+		if (!InSubmix || !InSubmix->IsOwnerValid())
 		{
 			return;
 		}
@@ -833,7 +833,7 @@ namespace Audio
 	{
 		ensure(IsInAudioThread());
 
-		if (!InSubmix || !InSubmix->IsValid() || !InSubmix->bIsEnvelopeFollowing)
+		if (!InSubmix || !InSubmix->IsOwnerValid() || !InSubmix->bIsEnvelopeFollowing)
 		{
 			return;
 		}
@@ -3410,23 +3410,18 @@ namespace Audio
 
 		const TWeakObjectPtr<USoundSubmix> SubmixPtr(&InSubmix);
 
-		// Pass the name vs. reconciling it inline with the command lambda, as occasionally
-		// deprecated submix buffer listeners are not constructed as a shared pointer and
-		// therefore getting destroyed before the lambda is executed on the AudioThread.
-		// This means no name is provided and thus not apparent who the caller is. If
-		// the Buffer Listener is invalid here, at least the callstack will show the
-		// requesting client directly.
-		const FString ListenerName = InSubmixBufferListener->GetListenerName();
-		auto RegisterLambda = [this, InSubmixBufferListener, ListenerName, SubmixPtr]()
+		auto RegisterLambda = [this, InSubmixBufferListener, SubmixPtr]()
 		{
 			CSV_SCOPED_TIMING_STAT(Audio, RegisterSubmixBufferListener);
 
 			FMixerSubmixPtr FoundSubmix = GetSubmixInstance(SubmixPtr.Get()).Pin();
+			const FString& ListenerName = InSubmixBufferListener->GetListenerName();
 
 			// Submix is not currently registered
 			if (!FoundSubmix.IsValid() && SubmixPtr.IsValid())
 			{
-				FString SubmixName = SubmixPtr->GetPathName();
+				const FString SubmixName = SubmixPtr->GetPathName();
+
 				if (SubmixBufferListenerForceRegistrationCVar)
 				{
 					RegisterSoundSubmix(SubmixPtr.Get(), true /* bInit */);
@@ -3455,7 +3450,7 @@ namespace Audio
 			}
 		};
 
-		UE_LOG(LogAudioMixer, Display, TEXT("Sending SubmixBufferListener '%s' register command..."), *ListenerName);
+		UE_LOG(LogAudioMixer, Display, TEXT("Sending SubmixBufferListener '%s' register command..."), *InSubmixBufferListener->GetListenerName());
 		FAudioThread::RunCommandOnAudioThread(MoveTemp(RegisterLambda));
 	}
 
diff --git a/Engine/Source/Runtime/AudioMixer/Private/AudioMixerSource.cpp b/Engine/Source/Runtime/AudioMixer/Private/AudioMixerSource.cpp
--- a/Engine/Source/Runtime/AudioMixer/Private/AudioMixerSource.cpp
+++ b/Engine/Source/Runtime/AudioMixer/Private/AudioMixerSource.cpp
@@ -1455,7 +1455,7 @@ namespace Audio
 				{
 					FMixerSubmixWeakPtr CurrReturnSubmixWeakPtr = MixerDevice->GetSubmixInstance(ReturnSubmix);
 					FMixerSubmixPtr CurrReturnSubmixPtr = CurrReturnSubmixWeakPtr.Pin();
-					while (CurrReturnSubmixPtr && CurrReturnSubmixPtr->IsValid())
+					while (CurrReturnSubmixPtr && CurrReturnSubmixPtr->IsOwnerValid())
 					{
 						ReturnSubmixAncestors.Add(CurrReturnSubmixPtr->GetId());
 
@@ -1471,7 +1471,7 @@ namespace Audio
 		FMixerSubmixWeakPtr CurrSubmixWeakPtr = MixerDevice->GetSubmixInstance(WaveInstance->SoundSubmix);
 		FMixerSubmixPtr CurrSubmixPtr = CurrSubmixWeakPtr.Pin();
 		// Check the submix and all its parents in the graph for active modulation
-		while (CurrSubmixPtr && CurrSubmixPtr->IsValid())
+		while (CurrSubmixPtr && CurrSubmixPtr->IsOwnerValid())
 		{
 			// Matching ID means the external spatializer has returned to the submix graph at this point,
 			// so we no longer need to manually apply volume modulation
diff --git a/Engine/Source/Runtime/AudioMixer/Private/AudioMixerSubmix.cpp b/Engine/Source/Runtime/AudioMixer/Private/AudioMixerSubmix.cpp
--- a/Engine/Source/Runtime/AudioMixer/Private/AudioMixerSubmix.cpp
+++ b/Engine/Source/Runtime/AudioMixer/Private/AudioMixerSubmix.cpp
@@ -1216,6 +1216,11 @@ namespace Audio
 		CommandQueue.Enqueue(MoveTemp(Command));
 	}
 
+	bool FMixerSubmix::IsOwnerValid() const
+	{
+		return OwningSubmixObject.IsValid();
+	}
+
 	bool FMixerSubmix::IsValid() const
 	{
 		return OwningSubmixObject.IsValid();
@@ -1425,7 +1430,7 @@ namespace Audio
 
 				// Owning submix can become invalid prior to BeginDestroy being called if object is
 				// forcibly deleted in editor, so submix validity (in addition to pointer validity) is checked before processing
-				if (ChildSubmix.IsValid() && ChildSubmix->IsValid())
+				if (ChildSubmix.IsValid() && ChildSubmix->IsOwnerValid())
 				{
 					ChildSubmix->ProcessAudio(InputBuffer);
 
diff --git a/Engine/Source/Runtime/AudioMixer/Public/AudioMixerSubmix.h b/Engine/Source/Runtime/AudioMixer/Public/AudioMixerSubmix.h
--- a/Engine/Source/Runtime/AudioMixer/Public/AudioMixerSubmix.h
+++ b/Engine/Source/Runtime/AudioMixer/Public/AudioMixerSubmix.h
@@ -216,6 +216,9 @@ namespace Audio
 		/** Checks to see if submix is valid.  Submix can be considered invalid if the OwningSubmix
 		  * pointer is stale.
 		  */
+		AUDIOMIXER_API bool IsOwnerValid() const;
+
+		UE_DEPRECATED(5.8, "Use IsOwnerValid() instead, to avoid confusion with TSharedPtr::IsValid().")
 		AUDIOMIXER_API bool IsValid() const;
 
 		// Function which processes audio.
