Post-Attenuation Bus Send
#Synced-CL 51714887

diff --git a/Engine/Source/Runtime/AudioMixer/Private/AudioMixerBus.cpp b/Engine/Source/Runtime/AudioMixer/Private/AudioMixerBus.cpp
--- a/Engine/Source/Runtime/AudioMixer/Private/AudioMixerBus.cpp
+++ b/Engine/Source/Runtime/AudioMixer/Private/AudioMixerBus.cpp
@@ -75,7 +75,7 @@ namespace Audio
 		}
 
 		// Return true if there is no more instances or sends
-		return bIsAutomatic && !InstanceIds.Num() && !AudioBusSends[(int32)EBusSendType::PreEffect].Num() && !AudioBusSends[(int32)EBusSendType::PostEffect].Num();
+		return bIsAutomatic && !InstanceIds.Num() && !AudioBusSends[(int32)EBusSendType::PreEffect].Num() && !AudioBusSends[(int32)EBusSendType::PostEffect].Num() && !AudioBusSends[(int32)EBusSendType::PostAttenuation].Num();
 	}
 
 	void FMixerAudioBus::AddSend(EBusSendType BusSendType, const FAudioBusSend& InAudioBusSend)
@@ -160,16 +160,20 @@ namespace Audio
 				{
 					SourceBufferPtr = SourceManager->GetPreviousSourceBusBuffer(AudioBusSend.SourceId);
 				}
-				// If the source mixing into this is not itself a bus, then simply mix the pre-attenuation audio of the source into the bus
+				// If the source mixing into this is not itself a bus, then simply mix the pre/post-attenuation audio of the source into the bus
 				// The source will have already computed its buffers for this frame
 				else if (BusSendType == (int32)EBusSendType::PostEffect)
 				{
 					SourceBufferPtr = SourceManager->GetPreDistanceAttenuationBuffer(AudioBusSend.SourceId);
 				}
-				else
+				else if (BusSendType == (int32)EBusSendType::PreEffect)
 				{
 					SourceBufferPtr = SourceManager->GetPreEffectBuffer(AudioBusSend.SourceId);
 				}
+				else
+				{
+					SourceBufferPtr = SourceManager->GetPostDistanceAttenuationBuffer(AudioBusSend.SourceId);
+				}
 
 				// It's possible we may not have a source buffer ptr here if the sound is not playing
 				if (SourceBufferPtr)
diff --git a/Engine/Source/Runtime/AudioMixer/Private/AudioMixerSourceManager.cpp b/Engine/Source/Runtime/AudioMixer/Private/AudioMixerSourceManager.cpp
--- a/Engine/Source/Runtime/AudioMixer/Private/AudioMixerSourceManager.cpp
+++ b/Engine/Source/Runtime/AudioMixer/Private/AudioMixerSourceManager.cpp
@@ -716,6 +716,7 @@ namespace Audio
 		
 			SourceInfo.AudioBusSends[(int32)EBusSendType::PreEffect].Reset();
 			SourceInfo.AudioBusSends[(int32)EBusSendType::PostEffect].Reset();
+			SourceInfo.AudioBusSends[(int32)EBusSendType::PostAttenuation].Reset();
 
 			SourceInfo.SourceEffectChainId = INDEX_NONE;
 
@@ -1106,6 +1107,7 @@ namespace Audio
 		SourceInfo.CurrentPCMBuffer = nullptr;
 		SourceInfo.SourceBuffer.Reset();
 		SourceInfo.PreDistanceAttenuationBuffer.Reset();
+		SourceInfo.PostDistanceAttenuationBuffer.Empty();
 		SourceInfo.SourceEffectScratchBuffer.Reset();
 		SourceInfo.AudioPluginOutputData.AudioBuffer.Reset();
 		SourceInfo.CurrentFrameValues.Reset();
@@ -1338,6 +1340,9 @@ namespace Audio
 			SourceInfo.PreDistanceAttenuationBuffer.Reset();
 			SourceInfo.PreDistanceAttenuationBuffer.AddZeroed(BufferSize);
 
+			SourceInfo.PostDistanceAttenuationBuffer.Reset();
+			SourceInfo.PostDistanceAttenuationBuffer.AddZeroed(BufferSize);
+
 			// Initialize the number of per-source LPF filters based on input channels
 			SourceInfo.LowPassFilter.Init(MixerDevice->SampleRate, InitParams.NumInputChannels);
 			SourceInfo.HighPassFilter.Init(MixerDevice->SampleRate, InitParams.NumInputChannels);
@@ -3570,6 +3575,9 @@ namespace Audio
 
 			FMixerSourceSubmixOutputBuffer& SourceSubmixOutputBuffer = SourceSubmixOutputBuffers[SourceId];
 
+			// Copy output buffer to Post-Attenuation Bus Send
+			SourceInfo.PostDistanceAttenuationBuffer = SourceInfo.SourceBuffer;
+
 			// Send source audio to plugins
 			ComputePluginAudio(SourceInfo, SourceSubmixOutputBuffer, SourceId, NumSamples);
 		}
@@ -3885,6 +3893,18 @@ namespace Audio
 		return SourceInfo.PreDistanceAttenuationBuffer.GetData();
 	}
 
+	const float* FMixerSourceManager::GetPostDistanceAttenuationBuffer(const int32 SourceId) const
+	{
+		const FSourceInfo& SourceInfo = SourceInfos[SourceId];
+
+		if (SourceInfo.bIsPaused || SourceInfo.bIsPausedForQuantization)
+		{
+			return nullptr;
+		}
+
+		return SourceInfo.PostDistanceAttenuationBuffer.GetData();
+	}
+
 	const float* FMixerSourceManager::GetPreEffectBuffer(const int32 SourceId) const
 	{
 		const FSourceInfo& SourceInfo = SourceInfos[SourceId];
diff --git a/Engine/Source/Runtime/AudioMixer/Private/AudioMixerSourceManager.h b/Engine/Source/Runtime/AudioMixer/Private/AudioMixerSourceManager.h
--- a/Engine/Source/Runtime/AudioMixer/Private/AudioMixerSourceManager.h
+++ b/Engine/Source/Runtime/AudioMixer/Private/AudioMixerSourceManager.h
@@ -306,6 +306,7 @@ namespace Audio
 
 		// Buffer getters
 		const float* GetPreDistanceAttenuationBuffer(const int32 SourceId) const;
+		const float* GetPostDistanceAttenuationBuffer(const int32 SourceId) const;
 		const float* GetPreEffectBuffer(const int32 SourceId) const;
 		const float* GetPreviousSourceBusBuffer(const int32 SourceId) const;
 		const float* GetPreviousAudioBusBuffer(const int32 AudioBusId) const;
@@ -459,6 +460,7 @@ namespace Audio
 			Audio::FAlignedFloatBuffer SourceBuffer;
 			Audio::FAlignedFloatBuffer PreEffectBuffer;
 			Audio::FAlignedFloatBuffer PreDistanceAttenuationBuffer;
+			Audio::FAlignedFloatBuffer PostDistanceAttenuationBuffer;
 			Audio::FAlignedFloatBuffer SourceEffectScratchBuffer;
 
 			// Data used for delaying the rendering of source audio for sample-accurate quantization
diff --git a/Engine/Source/Runtime/Engine/Classes/Components/AudioComponent.h b/Engine/Source/Runtime/Engine/Classes/Components/AudioComponent.h
--- a/Engine/Source/Runtime/Engine/Classes/Components/AudioComponent.h
+++ b/Engine/Source/Runtime/Engine/Classes/Components/AudioComponent.h
@@ -661,6 +661,15 @@ public:
 	UFUNCTION(BlueprintCallable, Category = "Audio|Components|Audio")
 	ENGINE_API void SetSourceBusSendPostEffect(USoundSourceBus* SoundSourceBus, float SourceBusSendLevel);
 
+	/** Allows designers to target a specific Audio Component instance’s sound and set the send level (volume of sound copied)
+	 *  to the indicated Source Bus. If the Source Bus is not already part of the sound’s sends, the reference will be added to
+	 *  this instance’s Override sends. This particular send occurs after attenuation processing.
+	 * @param SoundSourceBus The Bus to send the signal to
+	 * @param SourceBusSendLevel The scalar used to alter the volume of the copied signal
+	 */
+	UFUNCTION(BlueprintCallable, Category = "Audio|Components|Audio")
+	ENGINE_API void SetSourceBusSendPostAttenuation(USoundSourceBus* SoundSourceBus, float SourceBusSendLevel);
+
 	/** Sets how much audio the sound should send to the given Audio Bus (PRE Source Effects).
 	 *  if the Bus Send doesn't already exist, it will be added to the overrides on the active sound. 
 	 * @param AudioBus The Bus to send the signal to
@@ -677,6 +686,14 @@ public:
 	UFUNCTION(BlueprintCallable, Category = "Audio|Components|Audio")
 	ENGINE_API void SetAudioBusSendPostEffect(UAudioBus* AudioBus, float AudioBusSendLevel);
 
+	/** Sets how much audio the sound should send to the given Audio Bus (POST Attenuation).
+	 *  if the Audio Bus Send doesn't already exist, it will be added to the overrides on the active sound. 
+	 * @param AudioBus The Bus to send the signal to
+	 * @param AudioBusSendLevel The scalar used to alter the volume of the copied signal
+	 */
+	UFUNCTION(BlueprintCallable, Category = "Audio|Components|Audio")
+	ENGINE_API void SetAudioBusSendPostAttenuation(UAudioBus* AudioBus, float AudioBusSendLevel);
+
 	/** When set to TRUE, enables an additional Low Pass Filter Frequency to be calculated in with the
 	 *  sound instance’s LPF total, allowing designers to set filter settings for the targeted Audio Component’s
 	 *  sound instance.
diff --git a/Engine/Source/Runtime/Engine/Classes/Sound/SoundBase.h b/Engine/Source/Runtime/Engine/Classes/Sound/SoundBase.h
--- a/Engine/Source/Runtime/Engine/Classes/Sound/SoundBase.h
+++ b/Engine/Source/Runtime/Engine/Classes/Sound/SoundBase.h
@@ -242,6 +242,10 @@ public:
 	UPROPERTY(EditAnywhere, BlueprintReadWrite, Category = "Routing", meta = (DisplayName = "Pre-Effect Bus Sends", EditCondition = "bEnableBusSends"))
 	TArray<FSoundSourceBusSendInfo> PreEffectBusSends;
 
+	/** This sound will send its audio output to this list of buses if there are bus instances playing after attenuation is applied */
+	UPROPERTY(EditAnywhere, BlueprintReadWrite, Category = "Routing", meta = (DisplayName = "Post-Attenuation Bus Sends", EditCondition = "bEnableBusSends"))
+	TArray<FSoundSourceBusSendInfo> PostAttenuationBusSends;
+
 	/** Array of user data stored with the asset */
 	UPROPERTY(EditAnywhere, Instanced, Category = "Advanced")
 	TArray<TObjectPtr<UAssetUserData>> AssetUserData;
diff --git a/Engine/Source/Runtime/Engine/Private/Components/AudioComponent.cpp b/Engine/Source/Runtime/Engine/Private/Components/AudioComponent.cpp
--- a/Engine/Source/Runtime/Engine/Private/Components/AudioComponent.cpp
+++ b/Engine/Source/Runtime/Engine/Private/Components/AudioComponent.cpp
@@ -1635,6 +1635,11 @@ void UAudioComponent::SetSourceBusSendPostEffect(USoundSourceBus* SoundSourceBus
 	SetBusSendEffectInternal(SoundSourceBus, nullptr, SourceBusSendLevel, EBusSendType::PostEffect);
 }
 
+void UAudioComponent::SetSourceBusSendPostAttenuation(USoundSourceBus* SoundSourceBus, float SourceBusSendLevel)
+{
+	SetBusSendEffectInternal(SoundSourceBus, nullptr, SourceBusSendLevel, EBusSendType::PostAttenuation);
+}
+
 void UAudioComponent::SetAudioBusSendPreEffect(UAudioBus* AudioBus, float AudioBusSendLevel)
 {
 	SetBusSendEffectInternal(nullptr, AudioBus, AudioBusSendLevel, EBusSendType::PreEffect);
@@ -1645,6 +1650,11 @@ void UAudioComponent::SetAudioBusSendPostEffect(UAudioBus* AudioBus, float Audio
 	SetBusSendEffectInternal(nullptr, AudioBus, AudioBusSendLevel, EBusSendType::PostEffect);
 }
 
+void UAudioComponent::SetAudioBusSendPostAttenuation(UAudioBus* AudioBus, float AudioBusSendLevel)
+{
+	SetBusSendEffectInternal(nullptr, AudioBus, AudioBusSendLevel, EBusSendType::PostAttenuation);
+}
+
 void UAudioComponent::SetLowPassFilterEnabled(bool InLowPassFilterEnabled)
 {
 	if (FAudioDevice* AudioDevice = GetAudioDevice())
diff --git a/Engine/Source/Runtime/Engine/Private/SoundBase.cpp b/Engine/Source/Runtime/Engine/Private/SoundBase.cpp
--- a/Engine/Source/Runtime/Engine/Private/SoundBase.cpp
+++ b/Engine/Source/Runtime/Engine/Private/SoundBase.cpp
@@ -145,6 +145,10 @@ void USoundBase::GetSoundSourceBusSends(EBusSendType BusSendType, TArray<FSoundS
 	{
 		OutSends = PreEffectBusSends;
 	}
+	else if (BusSendType == EBusSendType::PostAttenuation)
+	{
+		OutSends = PostAttenuationBusSends;
+	}
 	else
 	{
 		OutSends = BusSends;
diff --git a/Engine/Source/Runtime/Engine/Public/Audio.h b/Engine/Source/Runtime/Engine/Public/Audio.h
--- a/Engine/Source/Runtime/Engine/Public/Audio.h
+++ b/Engine/Source/Runtime/Engine/Public/Audio.h
@@ -151,6 +151,7 @@ enum class EBusSendType : uint8
 {
 	PreEffect,
 	PostEffect,
+	PostAttenuation,
 	Count
 };
 
