[AudioMixer] FMixerDevice: removed #if WITH_EDITOR compiler conditionals so main submix analyzer messages can be send also when in Editor + added cvar to reset the main submix analyzers.

#Synced-CL 52380821

diff --git a/Engine/Source/Runtime/AudioMixer/Private/AudioMixerDevice.cpp b/Engine/Source/Runtime/AudioMixer/Private/AudioMixerDevice.cpp
--- a/Engine/Source/Runtime/AudioMixer/Private/AudioMixerDevice.cpp
+++ b/Engine/Source/Runtime/AudioMixer/Private/AudioMixerDevice.cpp
@@ -144,7 +144,6 @@ FAutoConsoleVariableRef CVarAudioMixerGeneratorSources(
 
 #if UE_AUDIO_PROFILERTRACE_ENABLED
 
-#if !WITH_EDITOR
 UE_TRACE_EVENT_BEGIN(Audio, MainSubmixLoaded)
 	UE_TRACE_EVENT_FIELD(uint32, DeviceId)
 	UE_TRACE_EVENT_FIELD(uint32, SubmixId)
@@ -169,7 +168,6 @@ UE_TRACE_EVENT_BEGIN(Audio, MainSubmixLKFSValues)
 	UE_TRACE_EVENT_FIELD(float, LoudnessRangeLowerBound)
 	UE_TRACE_EVENT_FIELD(float, LoudnessRangeUpperBound)
 UE_TRACE_EVENT_END()
-#endif // !WITH_EDITOR
 
 UE_TRACE_EVENT_BEGIN(Audio, SubmixAlivePing)
 	UE_TRACE_EVENT_FIELD(uint32, DeviceId)
@@ -483,7 +481,6 @@ namespace Audio
 		SubmixesWithAutoEnvelopeFollowers.Reset();
 	}
 
-#if !WITH_EDITOR
 	void FMixerDevice::InitializeMainSubmixTruePeakEnvelopeFollower(int32 InSampleRate, int32 InNumChannels, int32 InNumSamples)
 	{
 		MainSubmixTruePeakEnvelopeFollower = MakeUnique<FEnvelopeFollower>(
@@ -605,7 +602,6 @@ namespace Audio
 				<< MainSubmixLKFSValues.LoudnessRangeUpperBound(LoudnessRangeUpperBound);
 		}
 	}
-#endif // !WITH_EDITOR
 
 	void FMixerDevice::SendSubmixesAlivePing()
 	{
@@ -727,6 +723,9 @@ namespace Audio
 	{
 		IConsoleManager::Get().UnregisterConsoleObject(StartSubmixEnvelopeFollowerCVar);
 		IConsoleManager::Get().UnregisterConsoleObject(StopSubmixEnvelopeFollowerCVar);
+
+		StartSubmixEnvelopeFollowerCVar	= nullptr;
+		StopSubmixEnvelopeFollowerCVar	= nullptr;
 	}
 
 	void FMixerDevice::OnStartSubmixEnvelopeFollowerCVarChanged(IConsoleVariable* InStartSubmixEnvelopeFollowerCVar)
@@ -881,6 +880,48 @@ namespace Audio
 			}
 		}
 	}
+
+	void FMixerDevice::RegisterResetMainSubmixAnalyzersCmd()
+	{
+		if (ResetMainSubmixAnalyzersCmd == nullptr)
+		{
+			ResetMainSubmixAnalyzersCmd = IConsoleManager::Get().RegisterConsoleCommand(
+				*FString::Printf(TEXT("au.ResetMainSubmixAnalyzers_AD%d"), DeviceID),
+				TEXT("Resets the main submix LKFS and TruePeak analyzers"),
+				FConsoleCommandDelegate::CreateRaw(this, &FMixerDevice::OnResetMainSubmixAnalyzersCmdInvoked),
+				ECVF_Default
+			);
+		}
+	}
+
+	void FMixerDevice::UnregisterResetMainSubmixAnalyzersCmd()
+	{
+		IConsoleManager::Get().UnregisterConsoleObject(ResetMainSubmixAnalyzersCmd);
+		ResetMainSubmixAnalyzersCmd = nullptr;
+	}
+
+	void FMixerDevice::OnResetMainSubmixAnalyzersCmdInvoked()
+	{
+		ensure(IsInGameThread());
+
+		FAudioThread::RunCommandOnAudioThread([this]()
+		{
+			ResetMainSubmixAnalyzers();
+		});
+	}
+
+	void FMixerDevice::ResetMainSubmixAnalyzers()
+	{
+		ensure(IsInAudioThread());
+
+		const int32 NumSamples = PlatformInfo.NumChannels * PlatformSettings.CallbackBufferFrameSize;
+
+		MainSubmixLKFSAnalyzer.Reset();
+		InitializeMainSubmixLKFSAnalyzer(PlatformSettings.SampleRate, PlatformInfo.NumChannels);
+
+		MainSubmixTruePeakEnvelopeFollower.Reset();
+		InitializeMainSubmixTruePeakEnvelopeFollower(PlatformSettings.SampleRate, PlatformInfo.NumChannels, NumSamples);
+	}
 #endif // UE_AUDIO_PROFILERTRACE_ENABLED
 
 	void FMixerDevice::AddReferencedObjects(FReferenceCollector& Collector)
@@ -1186,15 +1227,13 @@ namespace Audio
 				InitSoundSubmixes();
 
 #if UE_AUDIO_PROFILERTRACE_ENABLED
-
-#if !WITH_EDITOR
 				const int32 NumSamples = PlatformInfo.NumChannels * PlatformSettings.CallbackBufferFrameSize;
 
 				InitializeMainSubmixTruePeakEnvelopeFollower(PlatformSettings.SampleRate, PlatformInfo.NumChannels, NumSamples);
 				InitializeMainSubmixLKFSAnalyzer(PlatformSettings.SampleRate, PlatformInfo.NumChannels);
-#endif // !WITH_EDITOR
 
 				RegisterSubmixEnvelopeFollowerCVars();
+				RegisterResetMainSubmixAnalyzersCmd();
 #endif // UE_AUDIO_PROFILERTRACE_ENABLED
 
 				AudioMixerPlatform->PostInitializeHardware();
@@ -1247,6 +1286,7 @@ namespace Audio
 		{
 #if UE_AUDIO_PROFILERTRACE_ENABLED
 			UnregisterSubmixEnvelopeFollowerCVars();
+			UnregisterResetMainSubmixAnalyzersCmd();
 #endif // UE_AUDIO_PROFILERTRACE_ENABLED
 
 			for (TObjectIterator<USoundSubmix> It; It; ++It)
@@ -1275,10 +1315,10 @@ namespace Audio
 			AudioMixerPlatform->TeardownHardware();
 		}
 
-#if UE_AUDIO_PROFILERTRACE_ENABLED && !WITH_EDITOR
+#if UE_AUDIO_PROFILERTRACE_ENABLED
 		MainSubmixTruePeakEnvelopeFollower.Reset();
 		MainSubmixLKFSAnalyzer.Reset();
-#endif // UE_AUDIO_PROFILERTRACE_ENABLED && !WITH_EDITOR
+#endif // UE_AUDIO_PROFILERTRACE_ENABLED
 
 		// Reset existing submixes if they exist
 		RequiredSubmixInstances.Reset();
@@ -1595,12 +1635,9 @@ namespace Audio
 		UpdateAudioClock();
 
 #if UE_AUDIO_PROFILERTRACE_ENABLED
-
-#if !WITH_EDITOR
 		UpdateAndSendMainSubmixTruePeakMaxValueMessage(Output);
 		UpdateAndSendMainSubmixLKFSValuesMessage(Output);
-#endif // !WITH_EDITOR
-		
+
 		UpdateSubmixesAlivePingTime();
 		SendSubmixesSendLevelTraceMessages();
 #endif //UE_AUDIO_PROFILERTRACE_ENABLED
@@ -1812,9 +1849,9 @@ namespace Audio
 			// 1. Load or reload all sound submixes/instances
 			LoadRequiredSubmix(ERequiredSubmixes::Main, GetRequiredSubmixName(ERequiredSubmixes::Main), false /* DefaultMuteWhenBackgrounded */, AudioSettings->MasterSubmix);
 
-#if UE_AUDIO_PROFILERTRACE_ENABLED && !WITH_EDITOR
+#if UE_AUDIO_PROFILERTRACE_ENABLED
 			SendMainSubmixLoadedMessage();
-#endif // UE_AUDIO_PROFILERTRACE_ENABLED && !WITH_EDITOR
+#endif // UE_AUDIO_PROFILERTRACE_ENABLED
 
 			// BaseDefaultSubmix is an optional master submix type set by project settings
 			if (AudioSettings->BaseDefaultSubmix.IsValid())
diff --git a/Engine/Source/Runtime/AudioMixer/Public/AudioMixerDevice.h b/Engine/Source/Runtime/AudioMixer/Public/AudioMixerDevice.h
--- a/Engine/Source/Runtime/AudioMixer/Public/AudioMixerDevice.h
+++ b/Engine/Source/Runtime/AudioMixer/Public/AudioMixerDevice.h
@@ -15,10 +15,10 @@
 #include "UObject/GCObject.h"
 #include "UObject/StrongObjectPtr.h"
 
-#if UE_AUDIO_PROFILERTRACE_ENABLED && !WITH_EDITOR
+#if UE_AUDIO_PROFILERTRACE_ENABLED
 #include "DSP/EnvelopeFollower.h"
 #include "DSP/LKFSAnalyzer.h"
-#endif // UE_AUDIO_PROFILERTRACE_ENABLED && !WITH_EDITOR
+#endif // UE_AUDIO_PROFILERTRACE_ENABLED
 
 // Forward Declarations
 class FOnSubmixEnvelopeBP;
@@ -420,7 +420,6 @@ namespace Audio
 		void OnTraceStarted(FTraceAuxiliary::EConnectionType TraceType, const FString& TraceDestination);
 		void OnTraceStopped(FTraceAuxiliary::EConnectionType TraceType, const FString& TraceDestination);
 
-#if !WITH_EDITOR
 		void InitializeMainSubmixTruePeakEnvelopeFollower(int32 InSampleRate, int32 InNumChannels, int32 InNumSamples);
 		void InitializeMainSubmixLKFSAnalyzer(int32 InSampleRate, int32 InNumChannels);
 
@@ -428,7 +427,6 @@ namespace Audio
 
 		void UpdateAndSendMainSubmixTruePeakMaxValueMessage(const FAlignedFloatBuffer& InMainSubmixAudioBuffer);
 		void UpdateAndSendMainSubmixLKFSValuesMessage(const FAlignedFloatBuffer& InMainSubmixAudioBuffer);
-#endif // !WITH_EDITOR
 	
 		void SendSubmixesAlivePing();
 		void UpdateSubmixesAlivePingTime();
@@ -444,6 +442,13 @@ namespace Audio
 		void StopSubmixEnvelopeFollowing(const FMixerSubmixPtr& InSubmix);
 
 		void AutomaticallyStartSubmixEnvelopeFollowers();
+
+		void RegisterResetMainSubmixAnalyzersCmd();
+		void UnregisterResetMainSubmixAnalyzersCmd();
+
+		void OnResetMainSubmixAnalyzersCmdInvoked();
+
+		void ResetMainSubmixAnalyzers();
 #endif // UE_AUDIO_PROFILERTRACE_ENABLED
 
 		TArray<USoundSubmix*> RequiredSubmixes;
@@ -537,14 +542,13 @@ namespace Audio
 
 		IConsoleVariable* StartSubmixEnvelopeFollowerCVar = nullptr;
 		IConsoleVariable* StopSubmixEnvelopeFollowerCVar  = nullptr;
+		IConsoleObject*   ResetMainSubmixAnalyzersCmd     = nullptr;
 
-#if !WITH_EDITOR
 		TUniquePtr<FEnvelopeFollower> MainSubmixTruePeakEnvelopeFollower;
 		FAlignedFloatBuffer MainSubmixTruePeakOutputEnvelope;
 
 		TUniquePtr<FLKFSAnalyzer> MainSubmixLKFSAnalyzer;
 		FLKFSAnalyzerResults MainSubmixLKFSAnalyzerResults;
-#endif // !WITH_EDITOR
 
 		double SubmixesAlivePingTimeSec = 0.0;
 		static constexpr double SubmixesAlivePingWaitTimeSec = 1.0;
