[AudioMixer][WASAPI] Fix ~2s game thread hitch during audio device swap

During a WASAPI device swap, the game thread was hitching for ~2 seconds due to two
independent issues:

1. FlushCommandQueue waiting on a stopped render thread:
After the device swap completes and the audio render thread is stopped,
FMixerDevice::UpdateHardware calls FlushAudioRenderingCommands to pump pending
source manager commands. FlushCommandQueue waits on CommandsProcessedEvent twice,
but the render thread is dead and will never signal it ? both waits hit their
~1s timeouts before falling back to inline pumping.

Fix: Added bRenderThreadStopped parameter to FlushCommandQueue and
FlushAudioRenderingCommands. When true, skips the Wait() calls and pumps
commands inline immediately. The game thread is already designated as the
audio thread via ResetAudioRenderingThreadId() at this point.

2. WASAPI device thread blocking on WaitForMultipleObjects timeout:
FAudioMixerWasapiRunnable::Stop() sets bIsRunning=false but does not wake the
thread from WaitForMultipleObjects (1s timeout). When the old device is no longer
present (e.g. headphones unplugged), WASAPI won't signal the event handles, so the
thread blocks for the full timeout before exiting.

Fix: Signal all event handles in Stop() so the thread wakes immediately, checks
bIsRunning, and exits. Safe because the while(bIsRunning) check precedes the
callback invocation.

#Synced-CL 52474377

diff --git a/Engine/Source/Runtime/AudioMixer/Private/AudioMixerDevice.cpp b/Engine/Source/Runtime/AudioMixer/Private/AudioMixerDevice.cpp
--- a/Engine/Source/Runtime/AudioMixer/Private/AudioMixerDevice.cpp
+++ b/Engine/Source/Runtime/AudioMixer/Private/AudioMixerDevice.cpp
@@ -1395,7 +1395,9 @@ namespace Audio
 				std::atomic<int32> CurrentAudioPlatformThreadId = AudioPlatformThreadId.load();
 
 				// Force source manager to incorporate device channel count change.
-				FlushAudioRenderingCommands(true /* bPumpSynchronously */);
+				// The render thread is stopped at this point, so skip waiting on it
+				// to avoid a ~1s timeout per Wait() call in FlushCommandQueue.
+				FlushAudioRenderingCommands(true /* bPumpSynchronously */, true /* bRenderThreadStopped */);
 
 				if (UAudioDeviceNotificationSubsystem* AudioDeviceNotifSubsystem = UAudioDeviceNotificationSubsystem::Get())
 				{
@@ -2450,11 +2452,11 @@ namespace Audio
 		}
 	}
 	
-	void FMixerDevice::FlushAudioRenderingCommands(bool bPumpSynchronously)
+	void FMixerDevice::FlushAudioRenderingCommands(bool bPumpSynchronously, bool bRenderThreadStopped)
 	{
 		if (IsInitialized() && (FPlatformProcess::SupportsMultithreading() && !AudioMixerPlatform->IsNonRealtime()))
 		{
-			SourceManager->FlushCommandQueue(bPumpSynchronously);
+			SourceManager->FlushCommandQueue(bPumpSynchronously, bRenderThreadStopped);
 		}
 		else if (AudioMixerPlatform->IsNonRealtime())
 		{
diff --git a/Engine/Source/Runtime/AudioMixer/Private/AudioMixerSourceManager.cpp b/Engine/Source/Runtime/AudioMixer/Private/AudioMixerSourceManager.cpp
--- a/Engine/Source/Runtime/AudioMixer/Private/AudioMixerSourceManager.cpp
+++ b/Engine/Source/Runtime/AudioMixer/Private/AudioMixerSourceManager.cpp
@@ -4147,7 +4147,7 @@ namespace Audio
 		}
 	}
 
-	void FMixerSourceManager::FlushCommandQueue(bool bPumpInCommand)
+	void FMixerSourceManager::FlushCommandQueue(bool bPumpInCommand, bool bSkipWaitForRenderThread)
 	{
 		check(CommandsProcessedEvent != nullptr);
 
@@ -4158,9 +4158,18 @@ namespace Audio
 			return;
 		}
 
-		// Make sure current current executing 
 		bool bTimedOut = false;
-		if (!CommandsProcessedEvent->Wait(CommandBufferFlushWaitTimeMsCvar))
+
+		if (bSkipWaitForRenderThread)
+		{
+			// The render thread is known to be stopped (e.g. during device swap).
+			// Skip waiting on CommandsProcessedEvent since no thread will signal it.
+			// We'll pump commands inline below.
+			CommandsProcessedEvent->Trigger();
+			bTimedOut = true;
+		}
+		// Make sure render thread is still executing
+		else if (!CommandsProcessedEvent->Wait(CommandBufferFlushWaitTimeMsCvar))
 		{
 			CommandsProcessedEvent->Trigger();
 			bTimedOut = true;
@@ -4179,8 +4188,13 @@ namespace Audio
 			PumpCommandQueue();
 		}
 
-		// Wait one more time for the double pump
-		if (!CommandsProcessedEvent->Wait(1000))
+		if (bSkipWaitForRenderThread)
+		{
+			// Already pumped inline, no second wait needed.
+			CommandsProcessedEvent->Trigger();
+		}
+		// wait one more time for the double pump
+		else if (!CommandsProcessedEvent->Wait(1000))
 		{
 			CommandsProcessedEvent->Trigger();
 			UE_LOG(LogAudioMixer, Warning, TEXT("Timed out waiting to flush the source manager command queue (2)."));
diff --git a/Engine/Source/Runtime/AudioMixer/Private/AudioMixerSourceManager.h b/Engine/Source/Runtime/AudioMixer/Private/AudioMixerSourceManager.h
--- a/Engine/Source/Runtime/AudioMixer/Private/AudioMixerSourceManager.h
+++ b/Engine/Source/Runtime/AudioMixer/Private/AudioMixerSourceManager.h
@@ -324,7 +324,7 @@ namespace Audio
 		bool IsSourceBus(const int32 SourceId) const;
 		void PumpCommandQueue();
 		void UpdatePendingReleaseData(bool bForceWait = false);
-		void FlushCommandQueue(bool bPumpCommandQueue = false);
+		void FlushCommandQueue(bool bPumpCommandQueue = false, bool bSkipWaitForRenderThread = false);
 
 		// Pushes a TFUnction command into an MPSC queue from an arbitrary thread to the audio render thread
 		void AudioMixerThreadMPSCCommand(TFunction<void()>&& InCommand, const char* InDebugString=nullptr);
diff --git a/Engine/Source/Runtime/AudioMixer/Public/AudioMixerDevice.h b/Engine/Source/Runtime/AudioMixer/Public/AudioMixerDevice.h
--- a/Engine/Source/Runtime/AudioMixer/Public/AudioMixerDevice.h
+++ b/Engine/Source/Runtime/AudioMixer/Public/AudioMixerDevice.h
@@ -212,7 +212,7 @@ namespace Audio
 		AUDIOMIXER_API virtual FPatchOutputStrongPtr AddPatchForSubmix(uint32 InObjectId, float InPatchGain) override;
 
 		AUDIOMIXER_API virtual void FlushExtended(UWorld* WorldToFlush, bool bClearActivatedReverb);
-		AUDIOMIXER_API virtual void FlushAudioRenderingCommands(bool bPumpSynchronously = false) override;
+		AUDIOMIXER_API virtual void FlushAudioRenderingCommands(bool bPumpSynchronously = false, bool bRenderThreadStopped = false) override;
 
 		// Audio Device Properties
 		AUDIOMIXER_API virtual bool IsNonRealtime() const override;
diff --git a/Engine/Source/Runtime/Engine/Public/AudioDevice.h b/Engine/Source/Runtime/Engine/Public/AudioDevice.h
--- a/Engine/Source/Runtime/Engine/Public/AudioDevice.h
+++ b/Engine/Source/Runtime/Engine/Public/AudioDevice.h
@@ -590,8 +590,9 @@ public:
 	/**
 	 * Allows audio rendering command queue to flush during audio device flush.
 	 * @param bPumpSynchronously must be called in situations where the audio render thread is not being called.
+	 * @param bRenderThreadStopped if true, skips waiting on the render thread event and pumps commands inline immediately.
 	 */
-	virtual void FlushAudioRenderingCommands(bool bPumpSynchronously = false) {}
+	virtual void FlushAudioRenderingCommands(bool bPumpSynchronously = false, bool bRenderThreadStopped = false) {}
 
 	ENGINE_API void OnPreGarbageCollect();
 
diff --git a/Engine/Source/Runtime/Windows/AudioMixerWasapi/Private/AudioMixerWasapiDeviceThread.cpp b/Engine/Source/Runtime/Windows/AudioMixerWasapi/Private/AudioMixerWasapiDeviceThread.cpp
--- a/Engine/Source/Runtime/Windows/AudioMixerWasapi/Private/AudioMixerWasapiDeviceThread.cpp
+++ b/Engine/Source/Runtime/Windows/AudioMixerWasapi/Private/AudioMixerWasapiDeviceThread.cpp
@@ -87,6 +87,19 @@ namespace Audio
 	void FAudioMixerWasapiRunnable::Stop()
 	{
 		bIsRunning = false;
+
+		// Signal event handles to wake the thread from WaitForMultipleObjects
+		// so it can exit promptly instead of waiting for the 1s timeout.
+		// This matters when the device is no longer present (e.g. headphones unplugged)
+		// and WASAPI won't signal the event handles anymore.
+		for (HANDLE Handle : EventHandles)
+		{
+			if (Handle != nullptr)
+			{
+				::SetEvent(Handle);
+			}
+		}
+
 		if (OutputStreamTimeoutsDetected > 0)
 		{
 			UE_LOG(LogAudioMixerWasapi, Display, TEXT("FAudioMixerWasapiRunnable::Stop render stream reported %d timeouts"), OutputStreamTimeoutsDetected);
