Unshelved CL 52494576

Replace SoundWaveProxyReader usage with SoundWaveProxyPlayer

[RN] Deprecate FSoundWaveProxyReader. Use FSoundWaveProxyPlayer instead

#Synced-CL 52474605

#Synced-CL 52496418

diff --git a/Engine/Plugins/Experimental/MetasoundExperimental/Source/MetasoundExperimentalEngineRuntime/Private/MetasoundGranulatorNode.cpp b/Engine/Plugins/Experimental/MetasoundExperimental/Source/MetasoundExperimentalEngineRuntime/Private/MetasoundGranulatorNode.cpp
--- a/Engine/Plugins/Experimental/MetasoundExperimental/Source/MetasoundExperimentalEngineRuntime/Private/MetasoundGranulatorNode.cpp
+++ b/Engine/Plugins/Experimental/MetasoundExperimental/Source/MetasoundExperimentalEngineRuntime/Private/MetasoundGranulatorNode.cpp
@@ -14,7 +14,7 @@
 #include "MetasoundWave.h"
 #include "DSP/Vbap.h"
 #include "Sound/SoundWave.h"
-#include "Sound/SoundWaveProxyReader.h"
+#include "Sound/SoundWaveProxyPlayer.h"
 #include "TypeFamily/ChannelTypeFamily.h"
 #include "MetasoundCatPannerNode.h"
 
@@ -288,13 +288,12 @@ namespace Metasound
 			SourceFileNumFrames = SoundWaveProxyPtr->GetNumFrames();
 			SourceFileNumChannels = SoundWaveProxyPtr->GetNumChannels();
 
-			FSoundWaveProxyReader::FSettings ProxyReaderSettings;
-			ProxyReaderSettings.MaxDecodeSizeInFrames = DecodedAudioSizeInSeconds * SourceFileSampleRate;
-			ProxyReaderSettings.StartTimeInSeconds = 0.0f;
-			ProxyReaderSettings.bIsLooping = true;
-			ProxyReaderSettings.LoopDurationInSeconds = SourceFileDurationSeconds;
-			SoundWaveProxyReaderPtr = FSoundWaveProxyReader::Create(SoundWaveProxyPtr.ToSharedRef(), ProxyReaderSettings);
-
+			FSoundWaveProxyPlayer::FSettings ProxyPlayerSettings(SourceFileSampleRate, SourceFileNumChannels);
+			ProxyPlayerSettings.MaxDecodeSizeFrames = DecodedAudioSizeInSeconds * SourceFileSampleRate;
+			SoundWaveProxyPlayerPtr = FSoundWaveProxyPlayer::Create(ProxyPlayerSettings);
+			SoundWaveProxyPlayerPtr->SetSoundWave(SoundWaveProxyPtr);
+			SoundWaveProxyPlayerPtr->SetLoop(true);
+			
 			DecodedChunks.Reset();
 			DecodedChunks.AddDefaulted(3);
 
@@ -344,10 +343,10 @@ namespace Metasound
 			check(SourceFileNumChannels > 0);
 			check(SoundWaveProxyPtr.IsValid());
 
-			check(SoundWaveProxyReaderPtr.IsValid());
+			check(SoundWaveProxyPlayerPtr.IsValid());
 			float DecoderSeekTimeSeconds = FMath::Max(InDecoderSeekTimeSeconds - 0.5f * DecodedAudioSizeInSeconds, 0.0f);
-			// If proxy reader already exists, then simply seek the decoder to the desired location
-			SoundWaveProxyReaderPtr->SeekToTime(DecoderSeekTimeSeconds);
+			// If proxy player already exists, then simply seek the decoder to the desired location
+			SoundWaveProxyPlayerPtr->SeekToTime(DecoderSeekTimeSeconds);
 
 			float AudioMixerSampleRate = Settings.GetSampleRate();
 
@@ -360,7 +359,7 @@ namespace Metasound
 			InOutDataChunk.PCMAudio.AddUninitialized(DecodedAudioSize);
 
 			// This does the actual decoding of the audio to match the size of the input buffer
-			SoundWaveProxyReaderPtr->PopAudio(InOutDataChunk.PCMAudio);
+			SoundWaveProxyPlayerPtr->GenerateSourceAudio(InOutDataChunk.PCMAudio);
 
 			// Check if we need to do SRC. If we do, this will change the allocated 
 			// size (potentially expanding or shrinking) to match 
@@ -931,7 +930,7 @@ namespace Metasound
 
 		FWaveAsset CurrentWaveAsset;
 		FSoundWaveProxyPtr SoundWaveProxyPtr;
-		TUniquePtr<FSoundWaveProxyReader> SoundWaveProxyReaderPtr;
+		TUniquePtr<FSoundWaveProxyPlayer> SoundWaveProxyPlayerPtr;
 		TUniquePtr<Audio::ISampleRateConverter> SRC;
 		EGranularEnvelope GranularEnvelopeType = EGranularEnvelope::Hann;
 		TArray<Audio::Grain::FEnvelope> GrainEnvelopeData;
diff --git a/Engine/Plugins/Runtime/Metasound/Source/MetasoundEngine/Private/MetasoundWavePlayerNode.cpp b/Engine/Plugins/Runtime/Metasound/Source/MetasoundEngine/Private/MetasoundWavePlayerNode.cpp
--- a/Engine/Plugins/Runtime/Metasound/Source/MetasoundEngine/Private/MetasoundWavePlayerNode.cpp
+++ b/Engine/Plugins/Runtime/Metasound/Source/MetasoundEngine/Private/MetasoundWavePlayerNode.cpp
@@ -20,6 +20,8 @@
 #include "MetasoundWavePlayerNodeV2.h"
 #include "Sound/SoundWaveProxyReader.h"
 
+PRAGMA_DISABLE_DEPRECATION_WARNINGS
+
 #define LOCTEXT_NAMESPACE "MetasoundWaveNode"
 
 namespace Metasound
@@ -1378,4 +1380,6 @@ namespace Metasound
 
 #undef LOCTEXT_NAMESPACE // MetasoundWaveNode
 
+PRAGMA_ENABLE_DEPRECATION_WARNINGS
+
 
diff --git a/Engine/Source/Runtime/AudioMixer/Classes/Generators/SoundWaveScrubber.h b/Engine/Source/Runtime/AudioMixer/Classes/Generators/SoundWaveScrubber.h
--- a/Engine/Source/Runtime/AudioMixer/Classes/Generators/SoundWaveScrubber.h
+++ b/Engine/Source/Runtime/AudioMixer/Classes/Generators/SoundWaveScrubber.h
@@ -8,13 +8,14 @@
 #include "Sound/SoundBase.h"
 #include "Sound/SoundGenerator.h"
 #include "Sound/SoundWave.h"
-#include "Sound/SoundWaveProxyReader.h"
 #include "AudioDynamicParameter.h"
 
 #include "SoundWaveScrubber.generated.h"
 
 #define UE_API AUDIOMIXER_API
 
+class FSoundWaveProxyPlayer;
+
 namespace Audio
 {
 
@@ -87,8 +88,8 @@ private:
 
 	// Sound wave proxy to safely query and use the referenced sound wave asset
 	FSoundWaveProxyPtr SoundWaveProxyPtr;
-	// Sound wave proxy reader is a decoder we can use to decode audio into chunks
-	TUniquePtr<FSoundWaveProxyReader> SoundWaveProxyReaderPtr;
+	// Sound wave proxy player is a decoder we can use to decode audio into chunks
+	TUniquePtr<FSoundWaveProxyPlayer> SoundWaveProxyPlayerPtr;
 	// Simple SRC interface. Decoded audio is SRC'd to match the audio mixer sample rate
 	TUniquePtr<Audio::ISampleRateConverter> SRC;
 
diff --git a/Engine/Source/Runtime/AudioMixer/Private/Generators/SoundWaveScrubber.cpp b/Engine/Source/Runtime/AudioMixer/Private/Generators/SoundWaveScrubber.cpp
--- a/Engine/Source/Runtime/AudioMixer/Private/Generators/SoundWaveScrubber.cpp
+++ b/Engine/Source/Runtime/AudioMixer/Private/Generators/SoundWaveScrubber.cpp
@@ -1,6 +1,7 @@
 // Copyright Epic Games, Inc. All Rights Reserved.
 
 #include "Generators/SoundWaveScrubber.h"
+#include "Sound/SoundWaveProxyPlayer.h"
 #include "ProfilingDebugging/CpuProfilerTrace.h"
 
 #include UE_INLINE_GENERATED_CPP_BY_NAME(SoundWaveScrubber)
@@ -108,30 +109,33 @@ void FSoundWaveScrubber::DecodeToDataChunk(FDecodedDataChunk& InOutDataChunk, fl
 	check(NumChannels > 0);
 	check(SoundWaveProxyPtr.IsValid());
 
-	if (!SoundWaveProxyReaderPtr.IsValid())
+	if (!SoundWaveProxyPlayerPtr.IsValid())
 	{
 		// Make sure we haven't already tried to create a proxy reader (decoder)
 		check(!bHasErrorWithDecoder);
 
-		// Create the proxy reader, which is our decoder. Initialize it at the initial playhead time
-		FSoundWaveProxyReader::FSettings ProxyReaderSettings;
-		ProxyReaderSettings.MaxDecodeSizeInFrames = DecodedAudioSizeInSeconds * SourceFileSampleRate;
-		ProxyReaderSettings.StartTimeInSeconds = InDecoderSeekTimeSeconds;
-		SoundWaveProxyReaderPtr = FSoundWaveProxyReader::Create(SoundWaveProxyPtr.ToSharedRef(), ProxyReaderSettings);
+		// Create the proxy player, which is our decoder. Initialize it at the initial playhead time
+		int32 SourceFileNumChanels = SoundWaveProxyPtr->GetNumChannels();
+		FSoundWaveProxyPlayer::FSettings ProxyPlayerSettings(SourceFileSampleRate, SourceFileNumChanels);
+		ProxyPlayerSettings.MaxDecodeSizeFrames = DecodedAudioSizeInSeconds * SourceFileSampleRate;
+		SoundWaveProxyPlayerPtr = FSoundWaveProxyPlayer::Create(ProxyPlayerSettings);
+		SoundWaveProxyPlayerPtr->SetSoundWave(SoundWaveProxyPtr);
 
-		if (!SoundWaveProxyReaderPtr.IsValid())
+		if (!SoundWaveProxyPlayerPtr->IsPlayerValid())
 		{
 			bHasErrorWithDecoder = true;
-			UE_LOG(LogAudioMixer, Warning, TEXT("Unable to make a sound wave proxy reader for sound wave '%s' in the sound wave scrubber."), *SoundWaveProxyPtr->GetFName().ToString());
+			UE_LOGF(LogAudioMixer, Warning, "Unable to make a Sound Wave Proxy Player for sound wave '%ls' in the sound wave scrubber.", *SoundWaveProxyPtr->GetFName().ToString());
 			return;
 		}
+
+		SoundWaveProxyPlayerPtr->SeekToTime(InDecoderSeekTimeSeconds);
 	}
 	else
 	{
-		check(SoundWaveProxyReaderPtr.IsValid());
+		check(SoundWaveProxyPlayerPtr.IsValid());
 		float DecoderSeekTimeSeconds = InDecoderSeekTimeSeconds;
-		// If proxy reader already exists, then simply seek the decoder to the desired location
-		SoundWaveProxyReaderPtr->SeekToTime(DecoderSeekTimeSeconds);
+		// If proxy player already exists, then simply seek the decoder to the desired location
+		SoundWaveProxyPlayerPtr->SeekToTime(DecoderSeekTimeSeconds);
 	}
 
 	InOutDataChunk.FrameStart = InDecoderSeekTimeSeconds * AudioMixerSampleRate;
@@ -143,7 +147,7 @@ void FSoundWaveScrubber::DecodeToDataChunk(FDecodedDataChunk& InOutDataChunk, fl
 	InOutDataChunk.PCMAudio.AddUninitialized(DecodedAudioSize);
 
 	// This does the actual decoding of the audio to match the size of the input buffer
-	SoundWaveProxyReaderPtr->PopAudio(InOutDataChunk.PCMAudio);
+	SoundWaveProxyPlayerPtr->GenerateSourceAudio(InOutDataChunk.PCMAudio);
 
 	// Check if we need to do SRC. If we do, this will change the allocated 
 	// size (potentially expanding or shrinking) to match 
diff --git a/Engine/Source/Runtime/Engine/Classes/Sound/SoundWaveProxyReader.h b/Engine/Source/Runtime/Engine/Classes/Sound/SoundWaveProxyReader.h
--- a/Engine/Source/Runtime/Engine/Classes/Sound/SoundWaveProxyReader.h
+++ b/Engine/Source/Runtime/Engine/Classes/Sound/SoundWaveProxyReader.h
@@ -16,7 +16,7 @@ class FSoundWaveProxy;
 * FSoundWaveProxyReader provides controls for looping and relevant frame index
 * values.
 */
-class FSoundWaveProxyReader
+class UE_DEPRECATED(5.8, "Use FSoundWaveProxyPlayer instead.") FSoundWaveProxyReader
 {
 
 public:
@@ -210,5 +210,8 @@ private:
 
 	bool bIsFirstDecode = true;
 	bool bPreviousIsStreaming = true;
+	
+PRAGMA_DISABLE_DEPRECATION_WARNINGS
 	uint32 NumFramesPerDecode = MinNumFramesPerDecode;
+PRAGMA_ENABLE_DEPRECATION_WARNINGS
 };
diff --git a/Engine/Source/Runtime/Engine/Private/SoundWaveProxyReader.cpp b/Engine/Source/Runtime/Engine/Private/SoundWaveProxyReader.cpp
--- a/Engine/Source/Runtime/Engine/Private/SoundWaveProxyReader.cpp
+++ b/Engine/Source/Runtime/Engine/Private/SoundWaveProxyReader.cpp
@@ -30,6 +30,8 @@ FAutoConsoleVariableRef CVarSoundWaveProxyReaderSimulateSeekOnNonSeekable(
 	TEXT("0: Do not simulate seek, !0: Simulate seek"),
 	ECVF_Default);
 
+PRAGMA_DISABLE_DEPRECATION_WARNINGS
+
 uint32 FSoundWaveProxyReader::ConformDecodeSize(uint32 InMaxDesiredDecodeSizeInFrames)
 {
 	static_assert(DefaultMinDecodeSizeInFrames >= DecodeSizeQuantizationInFrames, "Min decode size less than decode size quantization");
@@ -619,3 +621,5 @@ FSoundWaveProxyReader::EDecodeResult FSoundWaveProxyReader::Decode()
 
 	return EDecodeResult::Finished;
 }
+
+PRAGMA_ENABLE_DEPRECATION_WARNINGS
\ No newline at end of file
