[AudioRenderScheduler] Change THREAD_CHECK macro from a check to an ensure to be consistent with other such macros in AudioMixer

#rnx
#Synced-CL 52523372

diff --git a/Engine/Source/Runtime/AudioMixer/Private/AudioRenderScheduler.cpp b/Engine/Source/Runtime/AudioMixer/Private/AudioRenderScheduler.cpp
--- a/Engine/Source/Runtime/AudioMixer/Private/AudioRenderScheduler.cpp
+++ b/Engine/Source/Runtime/AudioMixer/Private/AudioRenderScheduler.cpp
@@ -11,7 +11,7 @@
 DEFINE_LOG_CATEGORY(LogAudioRenderScheduler);
 
 // If MixerDevice is null, don't check the thread. This is used in unit tests.
-#define AUDIO_RENDER_SCHEDULER_CHECK_THREAD() check(MixerDevice == nullptr || MixerDevice->IsAudioRenderingThread())
+#define AUDIO_RENDER_SCHEDULER_CHECK_THREAD() ensure(MixerDevice == nullptr || MixerDevice->IsAudioRenderingThread())
 
 namespace Audio
 {
