[AudioMixer] MainSubmixLoaded trace event: removed unused NumChannels field.

#Synced-CL 52673159

diff --git a/Engine/Source/Runtime/AudioMixer/Private/AudioMixerDevice.cpp b/Engine/Source/Runtime/AudioMixer/Private/AudioMixerDevice.cpp
--- a/Engine/Source/Runtime/AudioMixer/Private/AudioMixerDevice.cpp
+++ b/Engine/Source/Runtime/AudioMixer/Private/AudioMixerDevice.cpp
@@ -148,7 +148,6 @@ UE_TRACE_EVENT_BEGIN(Audio, MainSubmixLoaded)
 	UE_TRACE_EVENT_FIELD(uint32, DeviceId)
 	UE_TRACE_EVENT_FIELD(uint32, SubmixId)
 	UE_TRACE_EVENT_FIELD(uint64, Timestamp)
-	UE_TRACE_EVENT_FIELD(int32, NumChannels)
 UE_TRACE_EVENT_END()
 
 UE_TRACE_EVENT_BEGIN(Audio, MainSubmixTruePeakMaxValue)
