[AudioInsights] Add BusPatch trace events and engine-side emission for MetaSound AudioBus Writer/Reader nodes

#Synced-CL 52724280

diff --git a/Engine/Source/Runtime/AudioMixer/Private/AudioBusSubsystem.cpp b/Engine/Source/Runtime/AudioMixer/Private/AudioBusSubsystem.cpp
--- a/Engine/Source/Runtime/AudioMixer/Private/AudioBusSubsystem.cpp
+++ b/Engine/Source/Runtime/AudioMixer/Private/AudioBusSubsystem.cpp
@@ -163,7 +163,15 @@ Audio::FPatchInput UAudioBusSubsystem::AddPatchInputForSoundAndAudioBus(uint64 S
 	if (Audio::FPatchOutputStrongPtr PatchOutput = MixerDevice->MakePatch(InFrames, NumChannels, InGain))
 	{
 		Audio::FPatchInput PatchInput = MoveTemp(PatchOutput);
-		AddPendingConnection(SoundInstanceID, FPendingConnection{ FPendingConnection::FPatchVariant(TInPlaceType<Audio::FPatchInput>(), PatchInput), MoveTemp(AudioBusKey), InFrames, NumChannels });
+		AddPendingConnection(SoundInstanceID, FPendingConnection{ .PatchVariant = FPendingConnection::FPatchVariant(TInPlaceType<Audio::FPatchInput>(), PatchInput),
+																  .AudioBusKey = MoveTemp(AudioBusKey),
+																  .BlockSizeFrames = InFrames,
+																  .NumChannels = NumChannels,
+#if UE_AUDIO_PROFILERTRACE_ENABLED
+																  .InstanceID = SoundInstanceID 
+#endif // UE_AUDIO_PROFILERTRACE_ENABLED
+		});
+
 		return PatchInput;
 	}
 
@@ -180,7 +188,15 @@ Audio::FPatchOutputStrongPtr UAudioBusSubsystem::AddPatchOutputForSoundAndAudioB
 
 	if (Audio::FPatchOutputStrongPtr PatchOutput = MixerDevice->MakePatch(InFrames, NumChannels, InGain))
 	{
-		AddPendingConnection(SoundInstanceID, FPendingConnection{ FPendingConnection::FPatchVariant(TInPlaceType<Audio::FPatchOutputStrongPtr>(), PatchOutput), MoveTemp(AudioBusKey), InFrames, NumChannels });
+		AddPendingConnection(SoundInstanceID, FPendingConnection{ .PatchVariant = FPendingConnection::FPatchVariant(TInPlaceType<Audio::FPatchOutputStrongPtr>(), PatchOutput),
+																  .AudioBusKey = MoveTemp(AudioBusKey),
+																  .BlockSizeFrames = InFrames,
+																  .NumChannels = NumChannels,
+#if UE_AUDIO_PROFILERTRACE_ENABLED
+																  .InstanceID = SoundInstanceID
+#endif // UE_AUDIO_PROFILERTRACE_ENABLED
+		});
+
 		return PatchOutput;
 	}
 
@@ -206,10 +222,28 @@ void UAudioBusSubsystem::ConnectPatches(uint64 SoundInstanceID)
 			switch (PendingConnection.PatchVariant.GetIndex())
 			{
 			case FPendingConnection::FPatchVariant::IndexOfType<Audio::FPatchInput>():
-				SourceManager->AddPendingAudioBusConnection(MoveTemp(PendingConnection.AudioBusKey), PendingConnection.NumChannels, PendingConnection.bIsAutomatic, MoveTemp(PendingConnection.PatchVariant.Get<Audio::FPatchInput>()));
+				SourceManager->AddPendingAudioBusConnection(MoveTemp(PendingConnection.AudioBusKey)
+															, PendingConnection.NumChannels
+															, PendingConnection.bIsAutomatic
+															, MoveTemp(PendingConnection.PatchVariant.Get<Audio::FPatchInput>())
+															, INDEX_NONE
+#if UE_AUDIO_PROFILERTRACE_ENABLED
+															, PendingConnection.InstanceID
+#endif // UE_AUDIO_PROFILERTRACE_ENABLED
+				);
+
 				break;
 			case FPendingConnection::FPatchVariant::IndexOfType<Audio::FPatchOutputStrongPtr>():
-				SourceManager->AddPendingAudioBusConnection(MoveTemp(PendingConnection.AudioBusKey), PendingConnection.NumChannels, PendingConnection.bIsAutomatic, MoveTemp(PendingConnection.PatchVariant.Get<Audio::FPatchOutputStrongPtr>()));
+				SourceManager->AddPendingAudioBusConnection(MoveTemp(PendingConnection.AudioBusKey)
+															, PendingConnection.NumChannels
+															, PendingConnection.bIsAutomatic
+															, MoveTemp(PendingConnection.PatchVariant.Get<Audio::FPatchOutputStrongPtr>())
+															, INDEX_NONE
+#if UE_AUDIO_PROFILERTRACE_ENABLED
+															, PendingConnection.InstanceID
+#endif // UE_AUDIO_PROFILERTRACE_ENABLED
+				);
+
 				break;
 			}
 		}
diff --git a/Engine/Source/Runtime/AudioMixer/Private/AudioMixerSourceManager.cpp b/Engine/Source/Runtime/AudioMixer/Private/AudioMixerSourceManager.cpp
--- a/Engine/Source/Runtime/AudioMixer/Private/AudioMixerSourceManager.cpp
+++ b/Engine/Source/Runtime/AudioMixer/Private/AudioMixerSourceManager.cpp
@@ -285,6 +285,22 @@ UE_TRACE_EVENT_BEGIN(Audio, SoundToBusSend)
 	UE_TRACE_EVENT_FIELD(float[], SendLevels)
 UE_TRACE_EVENT_END()
 
+UE_TRACE_EVENT_BEGIN(Audio, BusPatchWriterConnected)
+	UE_TRACE_EVENT_FIELD(uint32, DeviceId)
+	UE_TRACE_EVENT_FIELD(uint64, Timestamp)
+	UE_TRACE_EVENT_FIELD(uint32, SenderID)
+	UE_TRACE_EVENT_FIELD(uint32[], ReceiverIDs)
+	UE_TRACE_EVENT_FIELD(float[], SendLevels)
+UE_TRACE_EVENT_END()
+
+UE_TRACE_EVENT_BEGIN(Audio, BusPatchReaderConnected)
+	UE_TRACE_EVENT_FIELD(uint32, DeviceId)
+	UE_TRACE_EVENT_FIELD(uint64, Timestamp)
+	UE_TRACE_EVENT_FIELD(uint32, SenderID)
+	UE_TRACE_EVENT_FIELD(uint32[], ReceiverIDs)
+	UE_TRACE_EVENT_FIELD(float[], SendLevels)
+UE_TRACE_EVENT_END()
+
 // Audio Bus messages
 UE_TRACE_EVENT_BEGIN(Audio, AudioBusStart)
 	UE_TRACE_EVENT_FIELD(uint32, DeviceId)
@@ -391,6 +407,40 @@ namespace Audio
 				<< AudioBusStop.Timestamp(FPlatformTime::Cycles64());
 		}
 	}
+
+	void TraceBusPatchWriterConnectedMessage(const uint32 InDeviceId, const uint32 InPlayOrder, const uint32 InBusObjectId)
+	{
+		const bool bChannelEnabled = UE_TRACE_CHANNELEXPR_IS_ENABLED(AudioChannel);
+		if (!bChannelEnabled)
+		{
+			return;
+		}
+
+		constexpr float SendLevel = 1.0f;
+		UE_TRACE_LOG(Audio, BusPatchWriterConnected, AudioChannel)
+			<< BusPatchWriterConnected.DeviceId(InDeviceId)
+			<< BusPatchWriterConnected.Timestamp(FPlatformTime::Cycles64())
+			<< BusPatchWriterConnected.SenderID(InPlayOrder)
+			<< BusPatchWriterConnected.ReceiverIDs(&InBusObjectId, 1)
+			<< BusPatchWriterConnected.SendLevels(&SendLevel, 1);
+	}
+
+	void TraceBusPatchReaderConnectedMessage(const uint32 InDeviceId, const uint32 InBusObjectId, const uint32 InPlayOrder)
+	{
+		const bool bChannelEnabled = UE_TRACE_CHANNELEXPR_IS_ENABLED(AudioChannel);
+		if (!bChannelEnabled)
+		{
+			return;
+		}
+
+		constexpr float SendLevel = 1.0f;
+		UE_TRACE_LOG(Audio, BusPatchReaderConnected, AudioChannel)
+			<< BusPatchReaderConnected.DeviceId(InDeviceId)
+			<< BusPatchReaderConnected.Timestamp(FPlatformTime::Cycles64())
+			<< BusPatchReaderConnected.SenderID(InBusObjectId)
+			<< BusPatchReaderConnected.ReceiverIDs(&InPlayOrder, 1)
+			<< BusPatchReaderConnected.SendLevels(&SendLevel, 1);
+	}
 #endif // UE_AUDIO_PROFILERTRACE_ENABLED
 
 	/*************************************************************************
@@ -643,6 +693,21 @@ namespace Audio
 				<< SoundToBusSend.SendLevels(BusSendLevels.GetData(), BusSendLevels.Num());
 		}
 	}
+
+	uint32 FMixerSourceManager::ResolvePlayOrder(const FPendingAudioBusConnection& Connection) const
+	{
+		uint32 PlayOrder = Connection.WaveInstancePlayOrder;
+		if (PlayOrder == INDEX_NONE)
+		{
+			const uint32* FoundPlayOrder = InstanceIdToPlayOrderMap.Find(Connection.InstanceID);
+			if (FoundPlayOrder)
+			{
+				PlayOrder = *FoundPlayOrder;
+			}
+		}
+
+		return PlayOrder;
+	}
 #endif // UE_AUDIO_PROFILERTRACE_ENABLED
 
 	void FMixerSourceManager::Init(const FSourceManagerInitParams& InitParams)
@@ -1010,6 +1075,12 @@ namespace Audio
 		// Free the mixer source buffer data
 		if (SourceInfo.MixerSourceBuffer.IsValid())
 		{
+#if UE_AUDIO_PROFILERTRACE_ENABLED
+			InstanceIdToPlayOrderMap.Remove(SourceInfo.MixerSourceBuffer->GetInstanceID());
+			ActiveBusPatchWriters.Remove(SourceInfo.PlayOrder);
+			ActiveBusPatchReaders.Remove(SourceInfo.PlayOrder);
+#endif // UE_AUDIO_PROFILERTRACE_ENABLED
+
 			if (UseRenderScheduler())
 			{
 				RemoveBusPatchDependencies(SourceInfo.MixerSourceBuffer->GetInstanceID());
@@ -1310,6 +1381,13 @@ namespace Audio
 			SourceInfo.bIsSoundfield = InitParams.bIsSoundfield;
 			SourceInfo.bResampling = false;
 
+#if UE_AUDIO_PROFILERTRACE_ENABLED
+			if (SourceInfo.MixerSourceBuffer.IsValid())
+			{
+				InstanceIdToPlayOrderMap.Add(SourceInfo.MixerSourceBuffer->GetInstanceID(), SourceInfo.PlayOrder);
+			}
+#endif // UE_AUDIO_PROFILERTRACE_ENABLED
+
 			// Call initialization from the render thread so anything wanting to do any initialization here can do so (e.g. procedural sound waves)
 			SourceInfo.SourceListener = InitParams.SourceListener;
 			SourceInfo.SourceListener->OnBeginGenerate();
@@ -2745,6 +2823,18 @@ namespace Audio
 					}
 				}
 				AudioBusPtr->AddNewPatchInput(PendingAudioBusConnection->PatchVariant.Get<FPatchInput>());
+
+#if UE_AUDIO_PROFILERTRACE_ENABLED
+				{
+					const uint32 PlayOrder = ResolvePlayOrder(*PendingAudioBusConnection);
+					if (PlayOrder != INDEX_NONE)
+					{
+						TraceBusPatchWriterConnectedMessage(MixerDevice->DeviceID, PlayOrder, AudioBusKey.ObjectId);
+						ActiveBusPatchWriters.AddUnique(PlayOrder, AudioBusKey.ObjectId);
+					}
+				}
+#endif // UE_AUDIO_PROFILERTRACE_ENABLED
+
 				break;
 			case FPendingAudioBusConnection::FPatchVariant::IndexOfType<FPatchOutputStrongPtr>():
 				if (UseRenderScheduler())
@@ -2758,6 +2848,18 @@ namespace Audio
 					}
 				}
 				AudioBusPtr->AddNewPatchOutput(PendingAudioBusConnection->PatchVariant.Get<FPatchOutputStrongPtr>());
+
+#if UE_AUDIO_PROFILERTRACE_ENABLED
+				{
+					const uint32 PlayOrder = ResolvePlayOrder(*PendingAudioBusConnection);
+					if (PlayOrder != INDEX_NONE)
+					{
+						TraceBusPatchReaderConnectedMessage(MixerDevice->DeviceID, AudioBusKey.ObjectId, PlayOrder);
+						ActiveBusPatchReaders.AddUnique(PlayOrder, AudioBusKey.ObjectId);
+					}
+				}
+#endif // UE_AUDIO_PROFILERTRACE_ENABLED
+
 				break;
 			}
 		}
@@ -2842,6 +2944,16 @@ namespace Audio
 		{
 			if (BusIsAlivePingTimeSec >= BusIsAlivePingWaitTimeSec)
 			{
+				for (const auto& [PlayOrder, BusObjectId] : ActiveBusPatchWriters)
+				{
+					TraceBusPatchWriterConnectedMessage(MixerDevice->DeviceID, PlayOrder, BusObjectId);
+				}
+
+				for (const auto& [PlayOrder, BusObjectId] : ActiveBusPatchReaders)
+				{
+					TraceBusPatchReaderConnectedMessage(MixerDevice->DeviceID, BusObjectId, PlayOrder);
+				}
+
 				BusIsAlivePingTimeSec = 0.0;
 			}
 			else
diff --git a/Engine/Source/Runtime/AudioMixer/Private/AudioMixerSourceManager.h b/Engine/Source/Runtime/AudioMixer/Private/AudioMixerSourceManager.h
--- a/Engine/Source/Runtime/AudioMixer/Private/AudioMixerSourceManager.h
+++ b/Engine/Source/Runtime/AudioMixer/Private/AudioMixerSourceManager.h
@@ -200,6 +200,8 @@ namespace Audio
 
 	class FMixerSourceManager
 	{
+		struct FPendingAudioBusConnection;
+
 	public:
 		FMixerSourceManager(FMixerDevice* InMixerDevice);
 		~FMixerSourceManager();
@@ -329,14 +331,32 @@ namespace Audio
 		// Pushes a TFUnction command into an MPSC queue from an arbitrary thread to the audio render thread
 		void AudioMixerThreadMPSCCommand(TFunction<void()>&& InCommand, const char* InDebugString=nullptr);
 		
-		void AddPendingAudioBusConnection(FAudioBusKey AudioBusKey, int32 NumChannels, bool bIsAutomatic, FPatchInput PatchInput, uint32 WaveInstancePlayOrder = INDEX_NONE)
+		void AddPendingAudioBusConnection(FAudioBusKey AudioBusKey, int32 NumChannels, bool bIsAutomatic, FPatchInput PatchInput, uint32 WaveInstancePlayOrder = INDEX_NONE, const uint64 InstanceID = 0u)
 		{
-			PendingAudioBusConnections.Enqueue(FPendingAudioBusConnection{ FPendingAudioBusConnection::FPatchVariant(TInPlaceType<FPatchInput>(), MoveTemp(PatchInput)), MoveTemp(AudioBusKey), NumChannels, bIsAutomatic, WaveInstancePlayOrder });
+			PendingAudioBusConnections.Enqueue(FPendingAudioBusConnection{ FPendingAudioBusConnection::FPatchVariant(TInPlaceType<FPatchInput>()
+												, MoveTemp(PatchInput))
+												, MoveTemp(AudioBusKey)
+												, NumChannels
+												, bIsAutomatic
+												, WaveInstancePlayOrder
+#if	UE_AUDIO_PROFILERTRACE_ENABLED
+												, InstanceID
+#endif // UE_AUDIO_PROFILERTRACE_ENABLED
+			});
 		}
 
-		void AddPendingAudioBusConnection(FAudioBusKey AudioBusKey, int32 NumChannels, bool bIsAutomatic, FPatchOutputStrongPtr PatchOutputStrongPtr, uint32 WaveInstancePlayOrder = INDEX_NONE)
+		void AddPendingAudioBusConnection(FAudioBusKey AudioBusKey, int32 NumChannels, bool bIsAutomatic, FPatchOutputStrongPtr PatchOutputStrongPtr, uint32 WaveInstancePlayOrder = INDEX_NONE, const uint64 InstanceID = 0u)
 		{
-			PendingAudioBusConnections.Enqueue(FPendingAudioBusConnection{ FPendingAudioBusConnection::FPatchVariant(TInPlaceType<FPatchOutputStrongPtr>(), MoveTemp(PatchOutputStrongPtr)), MoveTemp(AudioBusKey), NumChannels, bIsAutomatic, WaveInstancePlayOrder });
+			PendingAudioBusConnections.Enqueue(FPendingAudioBusConnection{ FPendingAudioBusConnection::FPatchVariant(TInPlaceType<FPatchOutputStrongPtr>()
+												, MoveTemp(PatchOutputStrongPtr))
+												, MoveTemp(AudioBusKey)
+												, NumChannels
+												, bIsAutomatic
+												, WaveInstancePlayOrder
+#if	UE_AUDIO_PROFILERTRACE_ENABLED
+												, InstanceID 
+#endif // UE_AUDIO_PROFILERTRACE_ENABLED
+			});
 		}
 
 		bool UseRenderScheduler() const
@@ -661,11 +681,18 @@ namespace Audio
 		void TraceSourceToSubmixSendMessage(const FSourceInfo& SourceInfo) const;
 		void TraceSourceToBusSendMessage(const uint32 SourceId, const FSourceInfo& SourceInfo) const;
 
+		uint32 ResolvePlayOrder(const FPendingAudioBusConnection& Connection) const;
+
 		IConsoleVariable* StartAudioBusEnvelopeFollowerCVar = nullptr;
 		IConsoleVariable* StopAudioBusEnvelopeFollowerCVar = nullptr;
 
 		double BusIsAlivePingTimeSec = 0.0;
 		static constexpr double BusIsAlivePingWaitTimeSec = 1.0;
+
+		// Audio render thread only
+		TMultiMap<uint32, uint32> ActiveBusPatchWriters;  // WaveInstancePlayOrder -> BusObjectId (source->bus, writer)
+		TMultiMap<uint32, uint32> ActiveBusPatchReaders;  // WaveInstancePlayOrder -> BusObjectId (bus->source, reader)
+		TMap<uint64, uint32> InstanceIdToPlayOrderMap;    // MixerSourceBuffer InstanceID -> WaveInstancePlayOrder
 #endif // UE_AUDIO_PROFILERTRACE_ENABLED
 
 		// Array of listener transforms
@@ -742,6 +769,10 @@ namespace Audio
 			int32 NumChannels = 0;
 			bool bIsAutomatic = false;
 			uint32 WaveInstancePlayOrder = INDEX_NONE;
+
+#if UE_AUDIO_PROFILERTRACE_ENABLED
+			uint64 InstanceID = 0u; // MixerSourceBuffer instance ID, used for PlayOrder resolution in trace profiling
+#endif // UE_AUDIO_PROFILERTRACE_ENABLED
 		};
 
 		TMpscQueue<FPendingAudioBusConnection> PendingAudioBusConnections;
diff --git a/Engine/Source/Runtime/AudioMixer/Public/AudioBusSubsystem.h b/Engine/Source/Runtime/AudioMixer/Public/AudioBusSubsystem.h
--- a/Engine/Source/Runtime/AudioMixer/Public/AudioBusSubsystem.h
+++ b/Engine/Source/Runtime/AudioMixer/Public/AudioBusSubsystem.h
@@ -2,6 +2,7 @@
 
 #pragma once
 
+#include "AudioMixerTrace.h"
 #include "DSP/MultithreadedPatching.h"
 #include "Misc/TVariant.h"
 #include "Subsystems/AudioEngineSubsystem.h"
@@ -117,6 +118,10 @@ private:
 		int32 BlockSizeFrames = 0;
 		int32 NumChannels = 0;
 		bool bIsAutomatic = false;
+
+#if UE_AUDIO_PROFILERTRACE_ENABLED
+		uint64 InstanceID = 0u; // MixerSourceBuffer instance ID, used for PlayOrder resolution in trace profiling
+#endif // UE_AUDIO_PROFILERTRACE_ENABLED
 	};
 
 	void AddPendingConnection(uint64 SoundInstanceID, FPendingConnection&& PendingConnection);
